Paper2D: Pulled various updates from Paper2D git branch
- Sanitize generated asset names when importing JSON sprite sheet frames
- Fix sprites not showing up in sprite/flipbook editor viewports
- Force VMI_Lit for both orthographic and perspective modes in the Paper2D viewport clients
- Initialize CurrentMode
- Fix missing toolbar icons in sprite and flipbook editors
- Fixed typos in slow progress dialog (not from git)

diff --git a/Engine/Plugins/Experimental/Paper2D/Source/Paper2DEditor/Private/FlipbookEditor/FlipbookEditorViewportClient.cpp b/Engine/Plugins/Experimental/Paper2D/Source/Paper2DEditor/Private/FlipbookEditor/FlipbookEditorViewportClient.cpp
--- a/Engine/Plugins/Experimental/Paper2D/Source/Paper2DEditor/Private/FlipbookEditor/FlipbookEditorViewportClient.cpp
+++ b/Engine/Plugins/Experimental/Paper2D/Source/Paper2DEditor/Private/FlipbookEditor/FlipbookEditorViewportClient.cpp
@@ -28,8 +28,6 @@ FFlipbookEditorViewportClient::FFlipbookEditorViewportClient(const TAttribute<UP
 	AnimatedRenderComponent->UpdateBounds();
 	PreviewScene->AddComponent(AnimatedRenderComponent.Get(), FTransform::Identity);
 
-	ViewportType = LVT_OrthoXZ;
-
 	bShowPivot = false;
 	bDeferZoomToSprite = true;
 
diff --git a/Engine/Plugins/Experimental/Paper2D/Source/Paper2DEditor/Private/PaperEditorViewportClient.cpp b/Engine/Plugins/Experimental/Paper2D/Source/Paper2DEditor/Private/PaperEditorViewportClient.cpp
--- a/Engine/Plugins/Experimental/Paper2D/Source/Paper2DEditor/Private/PaperEditorViewportClient.cpp
+++ b/Engine/Plugins/Experimental/Paper2D/Source/Paper2DEditor/Private/PaperEditorViewportClient.cpp
@@ -14,6 +14,34 @@ FPaperEditorViewportClient::FPaperEditorViewportClient()
 	ZoomAmount = 1.0f;
 
 	ModifyCheckerboardTextureColors();
+
+
+	//@TODO: Pretty lame hardcoding
+	//@TODO: Doesn't handle negatives either (not really)
+	const bool XX = FMath::IsNearlyEqual(PaperAxisX.X, 1.0f);
+	const bool XY = FMath::IsNearlyEqual(PaperAxisX.Y, 1.0f);
+	const bool YY = FMath::IsNearlyEqual(PaperAxisY.Y, 1.0f);
+	const bool YZ = FMath::IsNearlyEqual(PaperAxisY.Z, 1.0f);
+
+	ELevelViewportType NewViewportType = LVT_OrthoXZ;
+	if (XX && YY)
+	{
+		NewViewportType = LVT_OrthoXY;
+	}
+	else if (XX && YZ)
+	{
+		NewViewportType = LVT_OrthoXZ;
+	}
+	else if (XY && YZ)
+	{
+		NewViewportType = LVT_OrthoYZ;
+	}
+	else
+	{
+		//@TODO: Unsupported axes
+	}
+	SetViewModes(VMI_Lit, VMI_Lit);
+	SetViewportType(NewViewportType);
 }
 
 FPaperEditorViewportClient::~FPaperEditorViewportClient()
diff --git a/Engine/Plugins/Experimental/Paper2D/Source/Paper2DEditor/Private/PaperStyle.cpp b/Engine/Plugins/Experimental/Paper2D/Source/Paper2DEditor/Private/PaperStyle.cpp
--- a/Engine/Plugins/Experimental/Paper2D/Source/Paper2DEditor/Private/PaperStyle.cpp
+++ b/Engine/Plugins/Experimental/Paper2D/Source/Paper2DEditor/Private/PaperStyle.cpp
@@ -5,6 +5,8 @@
 #include "SlateStyle.h"
 #include "EditorStyle.h"
 
+#define IMAGE_BRUSH( RelativePath, ... ) FSlateImageBrush( Style.RootToContentDir( RelativePath, TEXT(".png") ), __VA_ARGS__ )
+
 //////////////////////////////////////////////////////////////////////////
 // FPaperStyle
 
@@ -47,34 +49,50 @@ TSharedRef<FSlateStyleSet> FPaperStyle::Create()
 
 	// Sprite editor
 	{
-		Style.Set( "SpriteEditor.SetShowGrid", new FSlateImageBrush( TEXT("Icons/icon_MatEd_Grid_40x"), Icon40x40 ) );
-		Style.Set( "SpriteEditor.SetShowGrid.Small", new FSlateImageBrush( TEXT("Icons/icon_MatEd_Grid_40x"), Icon20x20 ) );
-		Style.Set( "SpriteEditor.SetShowSourceTexture", new FSlateImageBrush(TEXT("texture://Paper2D/SlateBrushes/ShowSpriteSheetButton.ShowSpriteSheetButton"), Icon40x40 ) );
-		Style.Set( "SpriteEditor.SetShowSourceTexture.Small", new FSlateImageBrush(TEXT("texture://Paper2D/SlateBrushes/ShowSpriteSheetButton.ShowSpriteSheetButton"), Icon20x20 ) );
-		Style.Set( "SpriteEditor.SetShowBounds", new FSlateImageBrush( TEXT("Icons/icon_StaticMeshEd_Bounds_40x"), Icon40x40 ) );
-		Style.Set( "SpriteEditor.SetShowBounds.Small", new FSlateImageBrush( TEXT("Icons/icon_StaticMeshEd_Bounds_40x"), Icon20x20 ) );
-		Style.Set( "SpriteEditor.SetShowCollision", new FSlateImageBrush( TEXT("Icons/icon_StaticMeshEd_Collision_40x"), Icon40x40 ) );
-		Style.Set( "SpriteEditor.SetShowCollision.Small", new FSlateImageBrush( TEXT("Icons/icon_StaticMeshEd_Collision_40x"), Icon20x20 ) );
-
-		Style.Set( "SpriteEditor.SetShowSockets", new FSlateImageBrush( TEXT("Icons/icon_StaticMeshEd_ShowSockets_40x"), Icon40x40 ) );
-		Style.Set( "SpriteEditor.SetShowSockets.Small", new FSlateImageBrush( TEXT("Icons/icon_StaticMeshEd_ShowSockets_40x"), Icon20x20 ) );
-		Style.Set( "SpriteEditor.SetShowNormals", new FSlateImageBrush( TEXT("Icons/icon_StaticMeshEd_Normals_40x"), Icon40x40 ) );
-		Style.Set( "SpriteEditor.SetShowNormals.Small", new FSlateImageBrush( TEXT("Icons/icon_StaticMeshEd_Normals_40x"), Icon20x20 ) );
-		Style.Set( "SpriteEditor.SetShowPivot", new FSlateImageBrush( TEXT("Icons/icon_StaticMeshEd_ShowPivot_40x"), Icon40x40 ) );
-		Style.Set( "SpriteEditor.SetShowPivot.Small", new FSlateImageBrush( TEXT("Icons/icon_StaticMeshEd_ShowPivot_40x"), Icon20x20 ) );
-
-		Style.Set( "SpriteEditor.EnterViewMode", new FSlateImageBrush( TEXT("Icons/icon_TextureEd_AlphaChannel_40x"), Icon40x40 ) );
-		Style.Set( "SpriteEditor.EnterViewMode.Small", new FSlateImageBrush( TEXT("Icons/icon_TextureEd_AlphaChannel_40x"), Icon20x20 ) );
-
-		Style.Set( "SpriteEditor.EnterCollisionEditMode", new FSlateImageBrush( TEXT("texture://Paper2D/SlateBrushes/EditCollisionGeom.EditCollisionGeom"), Icon40x40 ) );
-		Style.Set( "SpriteEditor.EnterCollisionEditMode.Small", new FSlateImageBrush( TEXT("texture://Paper2D/SlateBrushes/EditCollisionGeom.EditCollisionGeom"), Icon20x20 ) );
-
-		Style.Set( "SpriteEditor.EnterRenderingEditMode", new FSlateImageBrush( TEXT("texture://Paper2D/SlateBrushes/EditRenderGeomButton.EditRenderGeomButton"), Icon40x40 ) );
-		Style.Set( "SpriteEditor.EnterRenderingEditMode.Small", new FSlateImageBrush( TEXT("texture://Paper2D/SlateBrushes/EditRenderGeomButton.EditRenderGeomButton"), Icon20x20 ) );
-
-		Style.Set( "SpriteEditor.EnterAddSpriteMode", new FSlateImageBrush( TEXT("texture://Paper2D/SlateBrushes/AddSpriteButton.AddSpriteButton"), Icon40x40 ) );
-		Style.Set( "SpriteEditor.EnterAddSpriteMode.Small", new FSlateImageBrush( TEXT("texture://Paper2D/SlateBrushes/AddSpriteButton.AddSpriteButton"), Icon20x20 ) );
+		Style.Set("SpriteEditor.SetShowGrid", new IMAGE_BRUSH(TEXT("Icons/icon_MatEd_Grid_40x"), Icon40x40));
+		Style.Set("SpriteEditor.SetShowGrid.Small", new IMAGE_BRUSH(TEXT("Icons/icon_MatEd_Grid_40x"), Icon20x20));
+		Style.Set("SpriteEditor.SetShowSourceTexture", new FSlateImageBrush(TEXT("texture://Paper2D/SlateBrushes/ShowSpriteSheetButton.ShowSpriteSheetButton"), Icon40x40));
+		Style.Set("SpriteEditor.SetShowSourceTexture.Small", new FSlateImageBrush(TEXT("texture://Paper2D/SlateBrushes/ShowSpriteSheetButton.ShowSpriteSheetButton"), Icon20x20));
+		Style.Set("SpriteEditor.SetShowBounds", new IMAGE_BRUSH(TEXT("Icons/icon_StaticMeshEd_Bounds_40x"), Icon40x40));
+		Style.Set("SpriteEditor.SetShowBounds.Small", new IMAGE_BRUSH(TEXT("Icons/icon_StaticMeshEd_Bounds_40x"), Icon20x20));
+		Style.Set("SpriteEditor.SetShowCollision", new IMAGE_BRUSH(TEXT("Icons/icon_StaticMeshEd_Collision_40x"), Icon40x40));
+		Style.Set("SpriteEditor.SetShowCollision.Small", new IMAGE_BRUSH(TEXT("Icons/icon_StaticMeshEd_Collision_40x"), Icon20x20));
+
+		Style.Set("SpriteEditor.SetShowSockets", new IMAGE_BRUSH(TEXT("Icons/icon_StaticMeshEd_ShowSockets_40x"), Icon40x40));
+		Style.Set("SpriteEditor.SetShowSockets.Small", new IMAGE_BRUSH(TEXT("Icons/icon_StaticMeshEd_ShowSockets_40x"), Icon20x20));
+		Style.Set("SpriteEditor.SetShowNormals", new IMAGE_BRUSH(TEXT("Icons/icon_StaticMeshEd_Normals_40x"), Icon40x40));
+		Style.Set("SpriteEditor.SetShowNormals.Small", new IMAGE_BRUSH(TEXT("Icons/icon_StaticMeshEd_Normals_40x"), Icon20x20));
+		Style.Set("SpriteEditor.SetShowPivot", new IMAGE_BRUSH(TEXT("Icons/icon_StaticMeshEd_ShowPivot_40x"), Icon40x40));
+		Style.Set("SpriteEditor.SetShowPivot.Small", new IMAGE_BRUSH(TEXT("Icons/icon_StaticMeshEd_ShowPivot_40x"), Icon20x20));
+
+		Style.Set("SpriteEditor.EnterViewMode", new IMAGE_BRUSH(TEXT("Icons/icon_TextureEd_AlphaChannel_40x"), Icon40x40));
+		Style.Set("SpriteEditor.EnterViewMode.Small", new IMAGE_BRUSH(TEXT("Icons/icon_TextureEd_AlphaChannel_40x"), Icon20x20));
+
+		Style.Set("SpriteEditor.EnterCollisionEditMode", new FSlateImageBrush(TEXT("texture://Paper2D/SlateBrushes/EditCollisionGeom.EditCollisionGeom"), Icon40x40));
+		Style.Set("SpriteEditor.EnterCollisionEditMode.Small", new FSlateImageBrush(TEXT("texture://Paper2D/SlateBrushes/EditCollisionGeom.EditCollisionGeom"), Icon20x20));
+
+		Style.Set("SpriteEditor.EnterRenderingEditMode", new FSlateImageBrush(TEXT("texture://Paper2D/SlateBrushes/EditRenderGeomButton.EditRenderGeomButton"), Icon40x40));
+		Style.Set("SpriteEditor.EnterRenderingEditMode.Small", new FSlateImageBrush(TEXT("texture://Paper2D/SlateBrushes/EditRenderGeomButton.EditRenderGeomButton"), Icon20x20));
+
+		Style.Set("SpriteEditor.EnterAddSpriteMode", new FSlateImageBrush(TEXT("texture://Paper2D/SlateBrushes/AddSpriteButton.AddSpriteButton"), Icon40x40));
+		Style.Set("SpriteEditor.EnterAddSpriteMode.Small", new FSlateImageBrush(TEXT("texture://Paper2D/SlateBrushes/AddSpriteButton.AddSpriteButton"), Icon20x20));
+	}
+
+	// Flipbook editor
+	{
+		Style.Set("FlipbookEditor.SetShowGrid", new IMAGE_BRUSH(TEXT("Icons/icon_MatEd_Grid_40x"), Icon40x40));
+		Style.Set("FlipbookEditor.SetShowGrid.Small", new IMAGE_BRUSH(TEXT("Icons/icon_MatEd_Grid_40x"), Icon20x20));
+		Style.Set("FlipbookEditor.SetShowBounds", new IMAGE_BRUSH(TEXT("Icons/icon_StaticMeshEd_Bounds_40x"), Icon40x40));
+		Style.Set("FlipbookEditor.SetShowBounds.Small", new IMAGE_BRUSH(TEXT("Icons/icon_StaticMeshEd_Bounds_40x"), Icon20x20));
+		Style.Set("FlipbookEditor.SetShowCollision", new IMAGE_BRUSH(TEXT("Icons/icon_StaticMeshEd_Collision_40x"), Icon40x40));
+		Style.Set("FlipbookEditor.SetShowCollision.Small", new IMAGE_BRUSH(TEXT("Icons/icon_StaticMeshEd_Collision_40x"), Icon20x20));
+		Style.Set("FlipbookEditor.SetShowPivot", new IMAGE_BRUSH(TEXT("Icons/icon_StaticMeshEd_ShowPivot_40x"), Icon40x40));
+		Style.Set("FlipbookEditor.SetShowPivot.Small", new IMAGE_BRUSH(TEXT("Icons/icon_StaticMeshEd_ShowPivot_40x"), Icon20x20));
 	}
 
 	return StyleRef;
 }
+
+//////////////////////////////////////////////////////////////////////////
+
+#undef IMAGE_BRUSH
diff --git a/Engine/Plugins/Experimental/Paper2D/Source/Paper2DEditor/Private/SpriteEditor/SpriteEditorViewportClient.cpp b/Engine/Plugins/Experimental/Paper2D/Source/Paper2DEditor/Private/SpriteEditor/SpriteEditorViewportClient.cpp
--- a/Engine/Plugins/Experimental/Paper2D/Source/Paper2DEditor/Private/SpriteEditor/SpriteEditorViewportClient.cpp
+++ b/Engine/Plugins/Experimental/Paper2D/Source/Paper2DEditor/Private/SpriteEditor/SpriteEditorViewportClient.cpp
@@ -38,7 +38,8 @@ IMPLEMENT_HIT_PROXY(HSpriteSelectableObjectHitProxy, HHitProxy);
 // FSpriteEditorViewportClient
 
 FSpriteEditorViewportClient::FSpriteEditorViewportClient(TWeakPtr<FSpriteEditor> InSpriteEditor, TWeakPtr<class SSpriteEditorViewport> InSpriteEditorViewportPtr)
-	: SpriteEditorPtr(InSpriteEditor)
+	: CurrentMode(ESpriteEditorMode::ViewMode)
+	, SpriteEditorPtr(InSpriteEditor)
 	, SpriteEditorViewportPtr(InSpriteEditorViewportPtr)
 {
 	check(SpriteEditorPtr.IsValid() && SpriteEditorViewportPtr.IsValid());
@@ -47,31 +48,6 @@ FSpriteEditorViewportClient::FSpriteEditorViewportClient(TWeakPtr<FSpriteEditor>
 
 	SetRealtime(true);
 
-	//@TODO: Pretty lame hardcoding
-	//@TODO: Doesn't handle negatives either (not really)
-	const bool XX = FMath::IsNearlyEqual(PaperAxisX.X, 1.0f);
-	const bool XY = FMath::IsNearlyEqual(PaperAxisX.Y, 1.0f);
-	const bool YY = FMath::IsNearlyEqual(PaperAxisY.Y, 1.0f);
-	const bool YZ = FMath::IsNearlyEqual(PaperAxisY.Z, 1.0f);
-
-	ViewportType = LVT_OrthoXZ;
-	if (XX && YY)
-	{
-		ViewportType = LVT_OrthoXY;
-	}
-	else if (XX && YZ)
-	{
-		ViewportType = LVT_OrthoXZ;
-	}
-	else if (XY && YZ)
-	{
-		ViewportType = LVT_OrthoYZ;
-	}
-	else
-	{
-		//@TODO: Unsupported axes
-	}
-
 	WidgetMode = FWidget::WM_Translate;
 	bManipulating = false;
 	bManipulationDirtiedSomething = false;
diff --git a/Engine/Plugins/Experimental/Paper2D/Source/PaperJsonImporter/Private/PaperJsonImporterFactory.cpp b/Engine/Plugins/Experimental/Paper2D/Source/PaperJsonImporter/Private/PaperJsonImporterFactory.cpp
--- a/Engine/Plugins/Experimental/Paper2D/Source/PaperJsonImporter/Private/PaperJsonImporterFactory.cpp
+++ b/Engine/Plugins/Experimental/Paper2D/Source/PaperJsonImporter/Private/PaperJsonImporterFactory.cpp
@@ -50,6 +50,8 @@ UObject* UPaperJsonImporterFactory::FactoryCreateText(UClass* InClass, UObject*
 {
 	FEditorDelegates::OnAssetPreImport.Broadcast(this, InClass, InParent, InName, Type);
 
+	FAssetToolsModule& AssetToolsModule = FModuleManager::GetModuleChecked<FAssetToolsModule>("AssetTools");
+
 	bool bLoadedSuccessfully = true;
 
 	const FString CurrentFilename = UFactory::GetCurrentFilename();
@@ -91,7 +93,7 @@ UObject* UPaperJsonImporterFactory::FactoryCreateText(UClass* InClass, UObject*
 			const FString FlashPrefix(TEXT("Adobe Flash"));
 			if (bLoadedSuccessfully && !AppName.StartsWith(FlashPrefix))
 			{
-				UE_LOG( LogPaperJsonImporter, Warning, TEXT( "Failed to parse sprite descriptor file '%s'.  Expected 'app' to start with %s" ), *NameForErrors, *FlashPrefix );
+				UE_LOG(LogPaperJsonImporter, Warning, TEXT("Failed to parse sprite descriptor file '%s'.  Expected 'app' to start with %s"), *NameForErrors, *FlashPrefix);
 				bLoadedSuccessfully = false;
 			}
 
@@ -99,7 +101,7 @@ UObject* UPaperJsonImporterFactory::FactoryCreateText(UClass* InClass, UObject*
 			{
 				if (Image.IsEmpty())
 				{
-					UE_LOG( LogPaperJsonImporter, Warning, TEXT( "Failed to parse sprite descriptor file '%s'.  Expected valid 'image' tag" ), *NameForErrors );
+					UE_LOG(LogPaperJsonImporter, Warning, TEXT("Failed to parse sprite descriptor file '%s'.  Expected valid 'image' tag"), *NameForErrors);
 					bLoadedSuccessfully = false;
 				}
 				else
@@ -112,14 +114,13 @@ UObject* UPaperJsonImporterFactory::FactoryCreateText(UClass* InClass, UObject*
 					TArray<FString> SheetFileNames;
 					SheetFileNames.Add(SourceSheetImageFilename);
 
-					FAssetToolsModule& AssetToolsModule = FModuleManager::GetModuleChecked<FAssetToolsModule>("AssetTools");
 					TArray<UObject*> ImportedSheets = AssetToolsModule.Get().ImportAssets(SheetFileNames, TargetSubPath);
 
 					ImportedTexture = (ImportedSheets.Num() > 0) ? Cast<UTexture2D>(ImportedSheets[0]) : NULL;
 
 					if (ImportedTexture == NULL)
 					{
-						UE_LOG( LogPaperJsonImporter, Warning, TEXT( "Failed to import sprite sheet image '%s'." ), *SourceSheetImageFilename );
+						UE_LOG(LogPaperJsonImporter, Warning, TEXT("Failed to import sprite sheet image '%s'."), *SourceSheetImageFilename);
 						bLoadedSuccessfully = false;
 					}
 					else
@@ -137,7 +138,7 @@ UObject* UPaperJsonImporterFactory::FactoryCreateText(UClass* InClass, UObject*
 		else
 		{
 			// Missing meta tag
-			UE_LOG( LogPaperJsonImporter, Warning, TEXT( "Failed to parse sprite descriptor file '%s'.  Missing meta block" ), *NameForErrors );
+			UE_LOG(LogPaperJsonImporter, Warning, TEXT("Failed to parse sprite descriptor file '%s'.  Missing meta block"), *NameForErrors);
 			bLoadedSuccessfully = false;
 		}
 
@@ -188,7 +189,7 @@ UObject* UPaperJsonImporterFactory::FactoryCreateText(UClass* InClass, UObject*
 						if (Frame.bRotated || !Frame.bTrimmed)
 						{
 							//@TODO: Should learn how to support this combo
-							UE_LOG( LogPaperJsonImporter, Warning, TEXT( "Frame %s is either rotated or not trimmed, which may not be handled correctly" ), *Frame.FrameName.ToString() );
+							UE_LOG(LogPaperJsonImporter, Warning, TEXT("Frame %s is either rotated or not trimmed, which may not be handled correctly"), *Frame.FrameName.ToString());
 						}
 
 						bReadFrameSuccessfully = bReadFrameSuccessfully && ReadRectangle(FrameData, "frame", /*out*/ Frame.SpritePosInSheet, /*out*/ Frame.SpriteSizeInSheet);
@@ -214,7 +215,7 @@ UObject* UPaperJsonImporterFactory::FactoryCreateText(UClass* InClass, UObject*
 					}
 					else
 					{
-						UE_LOG( LogPaperJsonImporter, Warning, TEXT( "Frame %s is in an unexpected format" ), *Frame.FrameName.ToString() );
+						UE_LOG(LogPaperJsonImporter, Warning, TEXT("Frame %s is in an unexpected format"), *Frame.FrameName.ToString());
 						bLoadedSuccessfully = false;
 					}
 
@@ -231,16 +232,16 @@ UObject* UPaperJsonImporterFactory::FactoryCreateText(UClass* InClass, UObject*
 					Result = Flipbook;
 				}
 
-				GWarn->BeginSlowTask(NSLOCTEXT("Paper2D", "PaperJsonImporterFactory_ImportingSprites", "Imprting Sprite Frame"), true, true);
+				GWarn->BeginSlowTask(NSLOCTEXT("Paper2D", "PaperJsonImporterFactory_ImportingSprites", "Importing Sprite Frame"), true, true);
 
 				// Create objects for each successfully parsed frame
 				const int32 FrameRun = 1; //@TODO: Don't make a keyframe for every single item if we can help it
 				for (int32 FrameIndex = 0; FrameIndex < ParsedFrames.Num(); ++FrameIndex)
 				{
-					GWarn->StatusUpdate(FrameIndex, ParsedFrames.Num(), NSLOCTEXT("Paper2D", "PaperJsonImporterFactory_ImportingSprites", "Imprting Sprite Frames"));
+					GWarn->StatusUpdate(FrameIndex, ParsedFrames.Num(), NSLOCTEXT("Paper2D", "PaperJsonImporterFactory_ImportingSprites", "Importing Sprite Frames"));
 
 					// Check for user canceling the import
-					if ( GWarn->ReceivedUserCancel() )
+					if (GWarn->ReceivedUserCancel())
 					{
 						break;
 					}
@@ -252,18 +253,20 @@ UObject* UPaperJsonImporterFactory::FactoryCreateText(UClass* InClass, UObject*
 
 					UObject* OuterForFrame = NULL; // @TODO: Use this if we don't want them to be individual assets - Flipbook;
 
-					// Create a package for each frame
-					FString NewPackageName = TargetSubPath + TEXT("/") + Frame.FrameName.ToString();
-					NewPackageName = PackageTools::SanitizePackageName(NewPackageName);
-					OuterForFrame = CreatePackage(NULL, *NewPackageName);
+					// Create a unique package name and asset name for the frame
+					const FString TentativePackagePath = PackageTools::SanitizePackageName(TargetSubPath + TEXT("/") + Frame.FrameName.ToString());
+					FString DefaultSuffix;
+					FString AssetName;
+					FString PackageName;
+					AssetToolsModule.Get().CreateUniqueAssetName(TentativePackagePath, /*out*/ DefaultSuffix, /*out*/ PackageName, /*out*/ AssetName);
+
+					// Create a package for the frame
+					OuterForFrame = CreatePackage(NULL, *PackageName);
 
 					// Create a frame in the package
-					UPaperSprite* TargetSprite = FindObject<UPaperSprite>(OuterForFrame, *Frame.FrameName.ToString());
-					if (TargetSprite == NULL)
-					{
-						TargetSprite = NewNamedObject<UPaperSprite>(OuterForFrame, Frame.FrameName, Flags);
-						FAssetRegistryModule::AssetCreated(TargetSprite);
-					}
+					UPaperSprite* TargetSprite = NewNamedObject<UPaperSprite>(OuterForFrame, *AssetName, Flags);
+					FAssetRegistryModule::AssetCreated(TargetSprite);
+
 					TargetSprite->Modify();
 
 					TargetSprite->InitializeSprite(ImportedTexture, Frame.SpritePosInSheet, Frame.SpriteSizeInSheet);
@@ -284,7 +287,7 @@ UObject* UPaperJsonImporterFactory::FactoryCreateText(UClass* InClass, UObject*
 			}
 			else
 			{
-				UE_LOG( LogPaperJsonImporter, Warning, TEXT( "Failed to parse sprite descriptor file '%s'.  Missing frames block" ), *NameForErrors );
+				UE_LOG(LogPaperJsonImporter, Warning, TEXT( "Failed to parse sprite descriptor file '%s'.  Missing frames block" ), *NameForErrors);
 			}
 		}
 	}
@@ -315,13 +318,13 @@ TSharedPtr<FJsonObject> UPaperJsonImporterFactory::ParseJSON(const FString& File
 		}
 		else
 		{
-			UE_LOG( LogPaperJsonImporter, Warning, TEXT( "Failed to parse sprite descriptor file '%s'.  Error: '%s'" ), *NameForErrors, *Reader->GetErrorMessage() );
+			UE_LOG(LogPaperJsonImporter, Warning, TEXT("Failed to parse sprite descriptor file '%s'.  Error: '%s'"), *NameForErrors, *Reader->GetErrorMessage());
 			return NULL;
 		}
 	}
 	else
 	{
-		UE_LOG( LogPaperJsonImporter, Warning, TEXT( "Sprite descriptor file '%s' was empty.  This sprite cannot be imported." ), *NameForErrors );
+		UE_LOG(LogPaperJsonImporter, Warning, TEXT("Sprite descriptor file '%s' was empty.  This sprite cannot be imported."), *NameForErrors);
 		return NULL;
 	}
 }
