Paper2D: Fixed Ctrl+B on sprite actors (implemented GetReferencedContentObjects)

diff --git a/Engine/Plugins/Experimental/Paper2D/Source/Paper2D/Classes/PaperFlipbookActor.h b/Engine/Plugins/Experimental/Paper2D/Source/Paper2D/Classes/PaperFlipbookActor.h
--- a/Engine/Plugins/Experimental/Paper2D/Source/Paper2D/Classes/PaperFlipbookActor.h
+++ b/Engine/Plugins/Experimental/Paper2D/Source/Paper2D/Classes/PaperFlipbookActor.h
@@ -12,5 +12,11 @@ class APaperFlipbookActor : public AActor
 
 	UPROPERTY(Category=Sprite, VisibleAnywhere)
 	TSubobjectPtr<class UPaperAnimatedRenderComponent> RenderComponent;
+
+	// AActor interface
+#if WITH_EDITOR
+	virtual bool GetReferencedContentObjects(TArray<UObject*>& Objects) const OVERRIDE;
+#endif
+	// End of AActor interface
 };
 
diff --git a/Engine/Plugins/Experimental/Paper2D/Source/Paper2D/Classes/PaperRenderActor.h b/Engine/Plugins/Experimental/Paper2D/Source/Paper2D/Classes/PaperRenderActor.h
--- a/Engine/Plugins/Experimental/Paper2D/Source/Paper2D/Classes/PaperRenderActor.h
+++ b/Engine/Plugins/Experimental/Paper2D/Source/Paper2D/Classes/PaperRenderActor.h
@@ -12,5 +12,11 @@ class APaperRenderActor : public AActor
 
 	UPROPERTY(Category=Sprite, VisibleAnywhere)
 	TSubobjectPtr<class UPaperRenderComponent> RenderComponent;
+
+	// AActor interface
+#if WITH_EDITOR
+	virtual bool GetReferencedContentObjects(TArray<UObject*>& Objects) const OVERRIDE;
+#endif
+	// End of AActor interface
 };
 
diff --git a/Engine/Plugins/Experimental/Paper2D/Source/Paper2D/Classes/PaperTileMapActor.h b/Engine/Plugins/Experimental/Paper2D/Source/Paper2D/Classes/PaperTileMapActor.h
--- a/Engine/Plugins/Experimental/Paper2D/Source/Paper2D/Classes/PaperTileMapActor.h
+++ b/Engine/Plugins/Experimental/Paper2D/Source/Paper2D/Classes/PaperTileMapActor.h
@@ -11,4 +11,10 @@ class APaperTileMapActor : public AActor
 
 	UPROPERTY(Category=TileMapActor, VisibleAnywhere, EditInline)
 	TSubobjectPtr<class UPaperTileMapRenderComponent> RenderComponent;
+
+	// AActor interface
+#if WITH_EDITOR
+	virtual bool GetReferencedContentObjects(TArray<UObject*>& Objects) const OVERRIDE;
+#endif
+	// End of AActor interface
 };
diff --git a/Engine/Plugins/Experimental/Paper2D/Source/Paper2D/Private/PaperFlipbookActor.cpp b/Engine/Plugins/Experimental/Paper2D/Source/Paper2D/Private/PaperFlipbookActor.cpp
--- a/Engine/Plugins/Experimental/Paper2D/Source/Paper2D/Private/PaperFlipbookActor.cpp
+++ b/Engine/Plugins/Experimental/Paper2D/Source/Paper2D/Private/PaperFlipbookActor.cpp
@@ -13,3 +13,14 @@ APaperFlipbookActor::APaperFlipbookActor(const FPostConstructInitializePropertie
 
 	RootComponent = RenderComponent;
 }
+
+#if WITH_EDITOR
+bool APaperFlipbookActor::GetReferencedContentObjects(TArray<UObject*>& Objects) const
+{
+	if (UPaperFlipbook* FlipbookAsset = RenderComponent->GetFlipbook())
+	{
+		Objects.Add(FlipbookAsset);
+	}
+	return true;
+}
+#endif
diff --git a/Engine/Plugins/Experimental/Paper2D/Source/Paper2D/Private/PaperRenderActor.cpp b/Engine/Plugins/Experimental/Paper2D/Source/Paper2D/Private/PaperRenderActor.cpp
--- a/Engine/Plugins/Experimental/Paper2D/Source/Paper2D/Private/PaperRenderActor.cpp
+++ b/Engine/Plugins/Experimental/Paper2D/Source/Paper2D/Private/PaperRenderActor.cpp
@@ -13,3 +13,14 @@ APaperRenderActor::APaperRenderActor(const FPostConstructInitializeProperties& P
 
 	RootComponent = RenderComponent;
 }
+
+#if WITH_EDITOR
+bool APaperRenderActor::GetReferencedContentObjects(TArray<UObject*>& Objects) const
+{
+	if (UPaperSprite* SourceSprite = RenderComponent->GetSprite())
+	{
+		Objects.Add(SourceSprite);
+	}
+	return true;
+}
+#endif
diff --git a/Engine/Plugins/Experimental/Paper2D/Source/Paper2D/Private/PaperTileMapActor.cpp b/Engine/Plugins/Experimental/Paper2D/Source/Paper2D/Private/PaperTileMapActor.cpp
--- a/Engine/Plugins/Experimental/Paper2D/Source/Paper2D/Private/PaperTileMapActor.cpp
+++ b/Engine/Plugins/Experimental/Paper2D/Source/Paper2D/Private/PaperTileMapActor.cpp
@@ -12,3 +12,14 @@ APaperTileMapActor::APaperTileMapActor(const FPostConstructInitializeProperties&
 
 	RootComponent = RenderComponent;
 }
+
+#if WITH_EDITOR
+bool APaperTileMapActor::GetReferencedContentObjects(TArray<UObject*>& Objects) const
+{
+	if (RenderComponent->DefaultLayerTileSet != nullptr)
+	{
+		Objects.Add(RenderComponent->DefaultLayerTileSet);
+	}
+	return true;
+}
+#endif
