Paper2D: Rename UPaperRenderComponent::TestMaterial to MaterialOverride and treat it as an actual override

diff --git a/Engine/Config/BaseEngine.ini b/Engine/Config/BaseEngine.ini
--- a/Engine/Config/BaseEngine.ini
+++ b/Engine/Config/BaseEngine.ini
@@ -535,6 +535,9 @@ bUseStreamingPause=false
 +EnumRedirects=(EnumName="EMeshComponentUpdateFlag",OldEnumEntry="SMU_AlwaysTickPose",NewEnumEntry="EMeshComponentUpdateFlag::AlwaysTickPose")
 +EnumRedirects=(EnumName="EMeshComponentUpdateFlag",OldEnumEntry="SMU_OnlyTickPoseWhenRendered",NewEnumEntry="EMeshComponentUpdateFlag::OnlyTickPoseWhenRendered")
 
+; Paper2D modifications (4.2 era)
++TaggedPropertyRedirects=(ClassName="PaperRenderComponent",OldPropertyName="TestMaterial",NewPropertyName="MaterialOverride")
+
 ; collision response container refactor
 +TaggedPropertyRedirects=(ClassName="CollisionResponseContainer",OldPropertyName="RigidBody",NewPropertyName="PhysicsBody")
 +TaggedPropertyRedirects=(ClassName="CollisionResponseContainer",OldPropertyName="Static",NewPropertyName="WorldStatic")
diff --git a/Engine/Plugins/Experimental/Paper2D/Source/Paper2D/Classes/PaperRenderComponent.h b/Engine/Plugins/Experimental/Paper2D/Source/Paper2D/Classes/PaperRenderComponent.h
--- a/Engine/Plugins/Experimental/Paper2D/Source/Paper2D/Classes/PaperRenderComponent.h
+++ b/Engine/Plugins/Experimental/Paper2D/Source/Paper2D/Classes/PaperRenderComponent.h
@@ -14,8 +14,9 @@ protected:
 	UPROPERTY(Category=Sprite, EditAnywhere, BlueprintReadOnly, meta=(DisplayThumbnail = "true"))
 	UPaperSprite* SourceSprite;
 
+	// The material override for this sprite component (if any)
 	UPROPERTY(Category=Sprite, EditAnywhere, BlueprintReadOnly)
-	UMaterialInterface* TestMaterial;
+	UMaterialInterface* MaterialOverride;
 
 	/** Physics scene information for this component, holds a single rigid body with multiple shapes. */
 	UPROPERTY(EditAnywhere, BlueprintReadOnly, Category=Collision)//, meta=(ShowOnlyInnerProperties))
diff --git a/Engine/Plugins/Experimental/Paper2D/Source/Paper2D/Classes/PaperSprite.h b/Engine/Plugins/Experimental/Paper2D/Source/Paper2D/Classes/PaperSprite.h
--- a/Engine/Plugins/Experimental/Paper2D/Source/Paper2D/Classes/PaperSprite.h
+++ b/Engine/Plugins/Experimental/Paper2D/Source/Paper2D/Classes/PaperSprite.h
@@ -153,6 +153,8 @@ public:
 
 	UTexture2D* GetPrimaryAtlasTexture() const { return SourceTexture; }
 
+	UMaterialInterface* GetDefaultMaterial() const { return DefaultMaterial; }
+
 	// Returns the render bounds of this sprite
 	FBoxSphereBounds GetRenderBounds() const;
 
diff --git a/Engine/Plugins/Experimental/Paper2D/Source/Paper2D/Private/PaperRenderComponent.cpp b/Engine/Plugins/Experimental/Paper2D/Source/Paper2D/Private/PaperRenderComponent.cpp
--- a/Engine/Plugins/Experimental/Paper2D/Source/Paper2D/Private/PaperRenderComponent.cpp
+++ b/Engine/Plugins/Experimental/Paper2D/Source/Paper2D/Private/PaperRenderComponent.cpp
@@ -14,8 +14,7 @@ UPaperRenderComponent::UPaperRenderComponent(const FPostConstructInitializePrope
 
 	SetCollisionProfileName(UCollisionProfile::BlockAllDynamic_ProfileName);
 
-	static ConstructorHelpers::FObjectFinder<UMaterial> DefaultMaterial(TEXT("/Paper2D/DefaultSpriteMaterial.DefaultSpriteMaterial"));
-	TestMaterial = DefaultMaterial.Object;
+	MaterialOverride = nullptr;
 
 	SpriteColor = FLinearColor::White;
 
diff --git a/Engine/Plugins/Experimental/Paper2D/Source/Paper2D/Private/PaperRenderSceneProxy.cpp b/Engine/Plugins/Experimental/Paper2D/Source/Paper2D/Private/PaperRenderSceneProxy.cpp
--- a/Engine/Plugins/Experimental/Paper2D/Source/Paper2D/Private/PaperRenderSceneProxy.cpp
+++ b/Engine/Plugins/Experimental/Paper2D/Source/Paper2D/Private/PaperRenderSceneProxy.cpp
@@ -151,9 +151,17 @@ FPaperRenderSceneProxy::FPaperRenderSceneProxy(const UPrimitiveComponent* InComp
 {
 	if (const UPaperRenderComponent* RenderComp = Cast<const UPaperRenderComponent>(InComponent))
 	{
-		Material = RenderComp->TestMaterial;
-		
 		SourceSprite = RenderComp->SourceSprite; //@TODO: This is totally not threadsafe, and won't keep up to date if the actor's sprite changes, etc....
+		if (SourceSprite)
+		{
+			Material = SourceSprite->GetDefaultMaterial();
+		}
+		
+		if (RenderComp->MaterialOverride)
+		{
+			Material = RenderComp->MaterialOverride;
+		}
+
 		if (Material)
 		{
 			MaterialRelevance = Material->GetRelevance();
