Paper2D: Set UPaperRenderComponent default mobility to static

diff --git a/Engine/Plugins/Experimental/Paper2D/Source/Paper2D/Private/PaperRenderComponent.cpp b/Engine/Plugins/Experimental/Paper2D/Source/Paper2D/Private/PaperRenderComponent.cpp
--- a/Engine/Plugins/Experimental/Paper2D/Source/Paper2D/Private/PaperRenderComponent.cpp
+++ b/Engine/Plugins/Experimental/Paper2D/Source/Paper2D/Private/PaperRenderComponent.cpp
@@ -18,7 +18,7 @@ UPaperRenderComponent::UPaperRenderComponent(const FPostConstructInitializePrope
 
 	SpriteColor = FLinearColor::White;
 
-	Mobility = EComponentMobility::Movable; //@TODO: Change the default!
+	Mobility = EComponentMobility::Static;
 	PrimaryComponentTick.bCanEverTick = true;
 	PrimaryComponentTick.TickGroup = TG_DuringPhysics;
 	bTickInEditor = true;
