Paper2D: Removed obsolete bounds expansion hack originally copied from UStaticMeshComponent

diff --git a/Engine/Plugins/Experimental/Paper2D/Source/Paper2D/Private/PaperRenderComponent.cpp b/Engine/Plugins/Experimental/Paper2D/Source/Paper2D/Private/PaperRenderComponent.cpp
--- a/Engine/Plugins/Experimental/Paper2D/Source/Paper2D/Private/PaperRenderComponent.cpp
+++ b/Engine/Plugins/Experimental/Paper2D/Source/Paper2D/Private/PaperRenderComponent.cpp
@@ -50,10 +50,7 @@ FBoxSphereBounds UPaperRenderComponent::CalcBounds(const FTransform & LocalToWor
 			}
 		}
 
-		// Takes into account that the static mesh collision code nudges collisions out by up to 1 unit.
-		//@TODO: Needed?  (copied from StaticMeshComponent)
-		NewBounds.BoxExtent += FVector(1,1,1);
-		NewBounds.SphereRadius += 1.0f;
+		// Apply bounds scale
 		NewBounds.BoxExtent *= BoundsScale;
 		NewBounds.SphereRadius *= BoundsScale;
 
