Paper2D: Fix support for RTS_Actor and RTS_Component spaces in UPaperRenderComponent::GetSocketTransform

diff --git a/Engine/Plugins/Experimental/Paper2D/Source/Paper2D/Private/PaperRenderComponent.cpp b/Engine/Plugins/Experimental/Paper2D/Source/Paper2D/Private/PaperRenderComponent.cpp
--- a/Engine/Plugins/Experimental/Paper2D/Source/Paper2D/Private/PaperRenderComponent.cpp
+++ b/Engine/Plugins/Experimental/Paper2D/Source/Paper2D/Private/PaperRenderComponent.cpp
@@ -154,26 +154,24 @@ FTransform UPaperRenderComponent::GetSocketTransform(FName InSocketName, ERelati
 	{
 		if (FPaperSpriteSocket* Socket = SourceSprite->FindSocket(InSocketName))
 		{
-			return Socket->LocalTransform * ComponentToWorld;
-			switch(TransformSpace)
+			switch (TransformSpace)
 			{
 				case RTS_World:
-				{
 					return Socket->LocalTransform * ComponentToWorld;
-				}
+
 				case RTS_Actor:
-				{
-					if( const AActor* Actor = GetOwner() )
+					if (const AActor* Actor = GetOwner())
 					{
-						FTransform SocketTransform = Socket->LocalTransform * ComponentToWorld;
+						const FTransform SocketTransform = Socket->LocalTransform * ComponentToWorld;
 						return SocketTransform.GetRelativeTransform(Actor->GetTransform());
 					}
 					break;
-				}
+
 				case RTS_Component:
-				{
 					return Socket->LocalTransform;
-				}
+
+				default:
+					check(false);
 			}
 		}
 	}
