Paper2D: Sprites now render in wireframe mode
- Add a show flag for Paper2D sprites
- Add support for Wireframe show flag in FPaperRenderSceneProxy
- Improve other debug rendering handling code in FPaperRenderSceneProxy

diff --git a/Engine/Plugins/Experimental/Paper2D/Source/Paper2D/Classes/PaperRenderComponent.h b/Engine/Plugins/Experimental/Paper2D/Source/Paper2D/Classes/PaperRenderComponent.h
--- a/Engine/Plugins/Experimental/Paper2D/Source/Paper2D/Classes/PaperRenderComponent.h
+++ b/Engine/Plugins/Experimental/Paper2D/Source/Paper2D/Classes/PaperRenderComponent.h
@@ -64,6 +64,8 @@ public:
 	/** Return the BodySetup to use for this PrimitiveComponent (single body case) */
 	virtual class UBodySetup2D* GetRBBodySetup2D();
 
+	// Returns the wireframe color to use for this component.
+	FLinearColor GetWireframeColor() const;
 
 protected:
 	//@TODO: Document
diff --git a/Engine/Plugins/Experimental/Paper2D/Source/Paper2D/Private/PaperRenderComponent.cpp b/Engine/Plugins/Experimental/Paper2D/Source/Paper2D/Private/PaperRenderComponent.cpp
--- a/Engine/Plugins/Experimental/Paper2D/Source/Paper2D/Private/PaperRenderComponent.cpp
+++ b/Engine/Plugins/Experimental/Paper2D/Source/Paper2D/Private/PaperRenderComponent.cpp
@@ -240,3 +240,8 @@ void UPaperRenderComponent::SetSpriteColor(FLinearColor NewColor)
 		MarkRenderDynamicDataDirty();
 	}
 }
+
+FLinearColor UPaperRenderComponent::GetWireframeColor() const
+{
+	return FLinearColor::Yellow;
+}
\ No newline at end of file
diff --git a/Engine/Plugins/Experimental/Paper2D/Source/Paper2D/Private/PaperRenderSceneProxy.cpp b/Engine/Plugins/Experimental/Paper2D/Source/Paper2D/Private/PaperRenderSceneProxy.cpp
--- a/Engine/Plugins/Experimental/Paper2D/Source/Paper2D/Private/PaperRenderSceneProxy.cpp
+++ b/Engine/Plugins/Experimental/Paper2D/Source/Paper2D/Private/PaperRenderSceneProxy.cpp
@@ -148,9 +148,13 @@ FPaperRenderSceneProxy::FPaperRenderSceneProxy(const UPrimitiveComponent* InComp
 	, Material(NULL)
 	, Owner(InComponent->GetOwner())
 	, SourceSprite(NULL)
+	, WireframeColor(FLinearColor::White)
+	, CollisionResponse(InComponent->GetCollisionResponseToChannels())
 {
 	if (const UPaperRenderComponent* RenderComp = Cast<const UPaperRenderComponent>(InComponent))
 	{
+		WireframeColor = RenderComp->GetWireframeColor();
+
 		SourceSprite = RenderComp->SourceSprite; //@TODO: This is totally not threadsafe, and won't keep up to date if the actor's sprite changes, etc....
 		if (SourceSprite)
 		{
@@ -171,6 +175,10 @@ FPaperRenderSceneProxy::FPaperRenderSceneProxy(const UPrimitiveComponent* InComp
 
 void FPaperRenderSceneProxy::DrawDynamicElements(FPrimitiveDrawInterface* PDI, const FSceneView* View)
 {
+	QUICK_SCOPE_CYCLE_COUNTER(STAT_PaperRenderSceneProxy_DrawDynamicElements);
+
+	checkSlow(IsInRenderingThread());
+
 	if (SourceSprite != NULL)//const UPaperRenderComponent* TypedFoo = Cast<const UPaperRenderComponent>(OwnerComponent))
 	{
 		// Show 2D physics
@@ -199,40 +207,9 @@ void FPaperRenderSceneProxy::DrawDynamicElements(FPrimitiveDrawInterface* PDI, c
 
 					const FColoredMaterialRenderProxy CollisionMaterialInstance(
 						LevelColorationMaterial->GetRenderProxy(IsSelected(), IsHovered()),
-						FColor(0, 255, 128, 255)//TypedFoo->GetWireframeColor()
+						WireframeColor
 						);
 
-					//@TODO:
-// 					FColor UStaticMeshComponent::GetWireframeColor() const
-// 					{
-// 						if(bOverrideWireframeColor)
-// 						{
-// 							return WireframeColorOverride;
-// 						}
-// 						else
-// 						{
-// 							if(Mobility == EComponentMobility::Static)
-// 							{
-// 								return FColor(0, 255, 255, 255);
-// 							}
-// 							else if(Mobility == EComponentMobility::Stationary)
-// 							{
-// 								return FColor(128, 128, 255, 255);
-// 							}
-// 							else // Movable
-// 							{
-// 								if(BodyInstance.bSimulatePhysics)
-// 								{
-// 									return FColor(0, 255, 128, 255);
-// 								}
-// 								else
-// 								{
-// 									return FColor(255, 0, 255, 255);
-// 								}
-// 							}
-// 						}
-// 					}
-
 					// Draw the static mesh's body setup.
 
 					// Get transform without scaling.
@@ -257,7 +234,7 @@ void FPaperRenderSceneProxy::DrawDynamicElements(FPrimitiveDrawInterface* PDI, c
 	}
 
 #if !(UE_BUILD_SHIPPING || UE_BUILD_TEST)
-//	if (View->Family->EngineShowFlags.StaticMeshes)
+	if (View->Family->EngineShowFlags.Paper2DSprites)
 	{
 		RenderBounds(PDI, View->Family->EngineShowFlags, GetBounds(), (Owner == NULL) || IsSelected());
 	}
@@ -357,15 +334,17 @@ void FPaperRenderSceneProxy::DrawDynamicElements_RichMesh(FPrimitiveDrawInterfac
 			BatchElement.PrimitiveUniformBuffer = GetUniformBuffer();
 			BatchElement.NumPrimitives = Vertices.Num() / 3;
 
+			const bool bIsWireframeView = View->Family->EngineShowFlags.Wireframe;
+
 			DrawRichMesh(
 				PDI, 
 				Mesh,
-				FLinearColor(1.0f, 0.0f, 0.0f),//WireframeColor
+				WireframeColor,
 				FLinearColor(1.0f, 1.0f, 1.0f),//LevelColor,
 				FLinearColor(1.0f, 1.0f, 1.0f),//PropertyColor,
 				this,
 				IsSelected(),
-				false//bIsWireframe
+				bIsWireframeView
 				);
 		}
 	}
@@ -373,34 +352,59 @@ void FPaperRenderSceneProxy::DrawDynamicElements_RichMesh(FPrimitiveDrawInterfac
 
 FPrimitiveViewRelevance FPaperRenderSceneProxy::GetViewRelevance(const FSceneView* View)
 {
-	const bool bVisible = View->Family->EngineShowFlags.BillboardSprites || View->Family->EngineShowFlags.Game;
+	checkSlow(IsInRenderingThread());
 
 	FPrimitiveViewRelevance Result;
-	Result.bDrawRelevance = bVisible;//IsShown(View);
-	//Result.bOpaqueRelevance = Material != NULL) ? Material-> true;
-	//Result.bTranslucentRelevance = true;
+	Result.bDrawRelevance = IsShown(View) && View->Family->EngineShowFlags.Paper2DSprites;
+	Result.bRenderCustomDepth = ShouldRenderCustomDepth();
+	Result.bRenderInMainPass = ShouldRenderInMainPass();
+
 	Result.bShadowRelevance = IsShadowCast(View);
 
+	MaterialRelevance.SetPrimitiveViewRelevance(Result);
 
-	//@TODO: Avoid this in some runtime cases where it's not necessary
-	Result.bDynamicRelevance = true;
-	// bStaticRelevance
+#undef SUPPORT_EXTRA_RENDERING
+#define SUPPORT_EXTRA_RENDERING !(UE_BUILD_SHIPPING || UE_BUILD_TEST) || WITH_EDITOR
+	
 
-	if (Material != NULL)
-	{
-		FMaterialRelevance ConcurrentMaterialRelevance = Material->GetRelevance_Concurrent();
+#if SUPPORT_EXTRA_RENDERING
+	bool bDrawSimpleCollision = false;
+	bool bDrawComplexCollision = false;
+	const bool bInCollisionView = IsCollisionView(View, bDrawSimpleCollision, bDrawComplexCollision);
+#endif
 
-		ConcurrentMaterialRelevance.SetPrimitiveViewRelevance(Result);
+	Result.bDynamicRelevance = true;
 
-// 		if (!View->Family->EngineShowFlags.Materials)
-// 		{
-// 			Result.bOpaqueRelevance = true;
-// 		}
-	}
-	else
+// 	if (
+// #if SUPPORT_EXTRA_RENDERING
+// 		IsRichView(View) ||
+// 		View->Family->EngineShowFlags.Collision ||
+// 		bInCollisionView ||
+// 		View->Family->EngineShowFlags.Bounds ||
+// #endif
+// 		// Force down dynamic rendering path if invalid lightmap settings, so we can apply an error material in DrawRichMesh
+// 		(HasStaticLighting() && !HasValidSettingsForStaticLighting()) ||
+// 		HasViewDependentDPG()
+// #if WITH_EDITOR
+// 		//only check these in the editor
+// 		|| IsSelected() || IsHovered()
+// #endif
+// 		)
+// 	{
+//		Result.bDynamicRelevance = true;
+// 	}
+// 	else
+// 	{
+// 		Result.bStaticRelevance = true;
+// 	}
+
+	if (!View->Family->EngineShowFlags.Materials
+#if SUPPORT_EXTRA_RENDERING
+		|| bInCollisionView
+#endif
+		)
 	{
-		Result.bOpaqueRelevance = false;
-		Result.bNormalTranslucencyRelevance = true;
+		Result.bOpaqueRelevance = true;
 	}
 
 	return Result;
@@ -428,3 +432,26 @@ void FPaperRenderSceneProxy::SetDrawCall_RenderThread(const FSpriteDrawCallRecor
 	FSpriteDrawCallRecord& Record = *new (BatchedSprites) FSpriteDrawCallRecord;
 	Record = NewDynamicData;
 }
+
+bool FPaperRenderSceneProxy::IsCollisionView(const FSceneView* View, bool& bDrawSimpleCollision, bool& bDrawComplexCollision) const
+{
+	bDrawSimpleCollision = false;
+	bDrawComplexCollision = false;
+
+	// If in a 'collision view' and collision is enabled
+	const bool bInCollisionView = View->Family->EngineShowFlags.CollisionVisibility || View->Family->EngineShowFlags.CollisionPawn;
+	if (bInCollisionView && IsCollisionEnabled())
+	{
+		// See if we have a response to the interested channel
+		bool bHasResponse = View->Family->EngineShowFlags.CollisionPawn && (CollisionResponse.GetResponse(ECC_Pawn) != ECR_Ignore);
+		bHasResponse |= View->Family->EngineShowFlags.CollisionVisibility && (CollisionResponse.GetResponse(ECC_Visibility) != ECR_Ignore);
+
+		if (bHasResponse)
+		{
+			bDrawComplexCollision = View->Family->EngineShowFlags.CollisionVisibility;
+			bDrawSimpleCollision = View->Family->EngineShowFlags.CollisionPawn;
+		}
+	}
+
+	return bInCollisionView;
+}
\ No newline at end of file
diff --git a/Engine/Plugins/Experimental/Paper2D/Source/Paper2D/Private/PaperRenderSceneProxy.h b/Engine/Plugins/Experimental/Paper2D/Source/Paper2D/Private/PaperRenderSceneProxy.h
--- a/Engine/Plugins/Experimental/Paper2D/Source/Paper2D/Private/PaperRenderSceneProxy.h
+++ b/Engine/Plugins/Experimental/Paper2D/Source/Paper2D/Private/PaperRenderSceneProxy.h
@@ -24,10 +24,21 @@ protected:
 	void DrawDynamicElements_Batched(FPrimitiveDrawInterface* PDI, const FSceneView* View, bool bUseOverrideColor, const FLinearColor& OverrideColor);
 	void DrawDynamicElements_RichMesh(FPrimitiveDrawInterface* PDI, const FSceneView* View, bool bUseOverrideColor, const FLinearColor& OverrideColor);
 
+	bool IsCollisionView(const FSceneView* View, bool& bDrawSimpleCollision, bool& bDrawComplexCollision) const;
+
+protected:
 	TArray<FSpriteDrawCallRecord> BatchedSprites;
 	class UMaterialInterface* Material;
 	FVector Origin;
 	AActor* Owner;
 	const UPaperSprite* SourceSprite;
+
+	// The color to draw as in wireframe mode
+	FLinearColor WireframeColor;
+
+	// The view relevance for the associated material
 	FMaterialRelevance MaterialRelevance;
+
+	// The Collision Response of the component being proxied
+	FCollisionResponseContainer CollisionResponse;
 };
diff --git a/Engine/Source/Runtime/Engine/Private/ShowFlags.cpp b/Engine/Source/Runtime/Engine/Private/ShowFlags.cpp
--- a/Engine/Source/Runtime/Engine/Private/ShowFlags.cpp
+++ b/Engine/Source/Runtime/Engine/Private/ShowFlags.cpp
@@ -447,6 +447,7 @@ void EngineShowFlagOverride(EShowFlagInitMode ShowFlagInitMode, EViewModeIndex V
 			 DISABLE_ENGINE_SHOWFLAG(SkeletalMeshes)
 			 DISABLE_ENGINE_SHOWFLAG(StaticMeshes)
 			 DISABLE_ENGINE_SHOWFLAG(BSP)
+			 DISABLE_ENGINE_SHOWFLAG(Paper2DSprites)
 #undef DISABLE_ENGINE_SHOWFLAG
 		}
 	}
diff --git a/Engine/Source/Runtime/Engine/Public/ShowFlagsValues.inl b/Engine/Source/Runtime/Engine/Public/ShowFlagsValues.inl
--- a/Engine/Source/Runtime/Engine/Public/ShowFlagsValues.inl
+++ b/Engine/Source/Runtime/Engine/Public/ShowFlagsValues.inl
@@ -260,6 +260,8 @@ SHOWFLAG_ALWAYS_ACCESSIBLE(PrecomputedVisibilityCells, SFG_Visualize, LOCTEXT("P
 SHOWFLAG_ALWAYS_ACCESSIBLE(VolumeLightingSamples, SFG_Visualize, LOCTEXT("VolumeLightingSamplesSF", "Volume Lighting Samples"))
 /** needed for VMI_LpvLightingViewMode, Whether to show only LPV lighting*/
 SHOWFLAG_ALWAYS_ACCESSIBLE(LpvLightingOnly, SFG_Hidden, LOCTEXT("VisualizeLPVSF_ViewMode", "Visualize LPV"))
+/** Render Paper2D sprites */
+SHOWFLAG_ALWAYS_ACCESSIBLE(Paper2DSprites, SFG_Advanced, LOCTEXT("Paper2DSpritesSF", "Paper 2D Sprites"))
 
 #undef SHOWFLAG_ALWAYS_ACCESSIBLE
 #undef SHOWFLAG_FIXED_IN_SHIPPING
