Paper2D: Fixed negative volume collision bug which caused NAN masses and missing collisions
- Prevented bounding box collision from generating negative volumes
- Added custom archive version for UPaperSprite assets

diff --git a/Engine/Plugins/Experimental/Paper2D/Source/Paper2D/Classes/PaperSprite.h b/Engine/Plugins/Experimental/Paper2D/Source/Paper2D/Classes/PaperSprite.h
--- a/Engine/Plugins/Experimental/Paper2D/Source/Paper2D/Classes/PaperSprite.h
+++ b/Engine/Plugins/Experimental/Paper2D/Source/Paper2D/Classes/PaperSprite.h
@@ -113,6 +113,8 @@ public:
 #if WITH_EDITORONLY_DATA
 	// UObject interface
 	virtual void PostEditChangeProperty(FPropertyChangedEvent& PropertyChangedEvent) OVERRIDE;
+	virtual void Serialize(FArchive& Ar) OVERRIDE;
+	virtual void PostLoad() OVERRIDE;
 	// End of UObject interface
 
 	FVector2D ConvertTextureSpaceToPivotSpace(FVector2D Input) const;
diff --git a/Engine/Plugins/Experimental/Paper2D/Source/Paper2D/Private/PaperCustomVersion.cpp b/Engine/Plugins/Experimental/Paper2D/Source/Paper2D/Private/PaperCustomVersion.cpp
new file mode 100644
--- /dev/null
+++ b/Engine/Plugins/Experimental/Paper2D/Source/Paper2D/Private/PaperCustomVersion.cpp
@@ -0,0 +1,9 @@
+// Copyright 1998-2014 Epic Games, Inc. All Rights Reserved.
+
+#include "Paper2DPrivatePCH.h"
+#include "PaperCustomVersion.h"
+
+const FGuid FPaperCustomVersion::GUID(0x11310AED, 0x2E554D61, 0xAF679AA3, 0xC5A1082C);
+
+// Register the custom version with core
+FCustomVersionRegistration GRegisterPaperCustomVersion(FPaperCustomVersion::GUID, FPaperCustomVersion::LatestVersion, TEXT("Paper2DVer"));
\ No newline at end of file
diff --git a/Engine/Plugins/Experimental/Paper2D/Source/Paper2D/Private/PaperCustomVersion.h b/Engine/Plugins/Experimental/Paper2D/Source/Paper2D/Private/PaperCustomVersion.h
new file mode 100644
--- /dev/null
+++ b/Engine/Plugins/Experimental/Paper2D/Source/Paper2D/Private/PaperCustomVersion.h
@@ -0,0 +1,27 @@
+// Copyright 1998-2014 Epic Games, Inc. All Rights Reserved.
+
+#pragma once
+
+// Custom serialization version for all packages containing Paper2D asset types
+struct FPaperCustomVersion
+{
+	enum Type
+	{
+		// Before any version changes were made in the plugin
+		BeforeCustomVersionWasAdded = 0,
+
+		// Fixed negative volume in UPaperSprite 3D aggregate geometry (caused collision issues)
+		FixedNegativeVolume = 1,
+
+
+		// -----<new versions can be added above this line>-------------------------------------------------
+		VersionPlusOne,
+		LatestVersion = VersionPlusOne - 1
+	};
+
+	// The GUID for this custom version number
+	const static FGuid GUID;
+
+private:
+	FPaperCustomVersion() {}
+};
diff --git a/Engine/Plugins/Experimental/Paper2D/Source/Paper2D/Private/PaperSprite.cpp b/Engine/Plugins/Experimental/Paper2D/Source/Paper2D/Private/PaperSprite.cpp
--- a/Engine/Plugins/Experimental/Paper2D/Source/Paper2D/Private/PaperSprite.cpp
+++ b/Engine/Plugins/Experimental/Paper2D/Source/Paper2D/Private/PaperSprite.cpp
@@ -3,6 +3,7 @@
 #include "Paper2DPrivatePCH.h"
 #include "PaperSprite.h"
 #include "PhysicsEngine/BodySetup.h"
+#include "PaperCustomVersion.h"
 
 #if WITH_EDITORONLY_DATA
 #include "GeomTools.h"
@@ -566,7 +567,7 @@ void UPaperSprite::BuildBoundingBoxCollisionData(bool bUseTightBounds)
 			// Create a new box primitive
 			const FVector BoxSize3D = (PaperAxisX * BoxSize2D.X) + (PaperAxisY * BoxSize2D.Y) + (PaperAxisZ * CollisionThickness);
 
-			FKBoxElem& Box = *new (BodySetup3D->AggGeom.BoxElems) FKBoxElem(BoxSize3D.X, BoxSize3D.Y, BoxSize3D.Z);
+			FKBoxElem& Box = *new (BodySetup3D->AggGeom.BoxElems) FKBoxElem(FMath::Abs(BoxSize3D.X), FMath::Abs(BoxSize3D.Y), FMath::Abs(BoxSize3D.Z));
 			Box.Center = (PaperAxisX * CenterInPivotSpace.X) + (PaperAxisY * CenterInPivotSpace.Y);
 
 			BodySetup3D->InvalidatePhysicsData();
@@ -1095,3 +1096,22 @@ void UPaperSprite::QuerySupportedSockets(TArray<FComponentSocketDescription>& Ou
 		new (OutSockets) FComponentSocketDescription(Socket.SocketName, EComponentSocketType::Socket);
 	}
 }
+
+void UPaperSprite::Serialize(FArchive& Ar)
+{
+	Super::Serialize(Ar);
+
+	Ar.UsingCustomVersion(FPaperCustomVersion::GUID);
+}
+
+void UPaperSprite::PostLoad()
+{
+	Super::PostLoad();
+
+#if WITH_EDITORONLY_DATA
+	if (GetLinkerCustomVersion(FPaperCustomVersion::GUID) < FPaperCustomVersion::FixedNegativeVolume)
+	{
+		RebuildCollisionData();
+	}
+#endif
+}
