Paper2D: Implement AdditionalStatObject to return the associated asset in the Paper2D render components (used for stats and error messages)

diff --git a/Engine/Plugins/Experimental/Paper2D/Source/Paper2D/Classes/PaperAnimatedRenderComponent.h b/Engine/Plugins/Experimental/Paper2D/Source/Paper2D/Classes/PaperAnimatedRenderComponent.h
--- a/Engine/Plugins/Experimental/Paper2D/Source/Paper2D/Classes/PaperAnimatedRenderComponent.h
+++ b/Engine/Plugins/Experimental/Paper2D/Source/Paper2D/Classes/PaperAnimatedRenderComponent.h
@@ -53,6 +53,7 @@ public:
 	// UActorComponent interface
 	virtual void TickComponent(float DeltaTime, enum ELevelTick TickType, FActorComponentTickFunction* ThisTickFunction) OVERRIDE;
 	virtual void SendRenderDynamicData_Concurrent() OVERRIDE;
+	virtual const UObject* AdditionalStatObject() const OVERRIDE;
 	// End of UActorComponent interface
 
 	// UPrimitiveComponent interface
diff --git a/Engine/Plugins/Experimental/Paper2D/Source/Paper2D/Classes/PaperRenderComponent.h b/Engine/Plugins/Experimental/Paper2D/Source/Paper2D/Classes/PaperRenderComponent.h
--- a/Engine/Plugins/Experimental/Paper2D/Source/Paper2D/Classes/PaperRenderComponent.h
+++ b/Engine/Plugins/Experimental/Paper2D/Source/Paper2D/Classes/PaperRenderComponent.h
@@ -45,6 +45,7 @@ public:
 	virtual void SendRenderDynamicData_Concurrent() OVERRIDE;
 	virtual void CreatePhysicsState() OVERRIDE;
 	virtual void DestroyPhysicsState() OVERRIDE;
+	virtual const UObject* AdditionalStatObject() const OVERRIDE;
 	// End of UActorComponent interface
 
 	// USceneComponent interface
diff --git a/Engine/Plugins/Experimental/Paper2D/Source/Paper2D/Classes/PaperTileMapRenderComponent.h b/Engine/Plugins/Experimental/Paper2D/Source/Paper2D/Classes/PaperTileMapRenderComponent.h
--- a/Engine/Plugins/Experimental/Paper2D/Source/Paper2D/Classes/PaperTileMapRenderComponent.h
+++ b/Engine/Plugins/Experimental/Paper2D/Source/Paper2D/Classes/PaperTileMapRenderComponent.h
@@ -53,6 +53,7 @@ public:
 
 	// UActorComponent interface
 	virtual void TickComponent(float DeltaTime, enum ELevelTick TickType, FActorComponentTickFunction* ThisTickFunction) OVERRIDE;
+	virtual const UObject* AdditionalStatObject() const OVERRIDE;
 	// End of UActorComponent interface
 
 	// UPrimitiveComponent interface
diff --git a/Engine/Plugins/Experimental/Paper2D/Source/Paper2D/Private/PaperAnimatedRenderComponent.cpp b/Engine/Plugins/Experimental/Paper2D/Source/Paper2D/Private/PaperAnimatedRenderComponent.cpp
--- a/Engine/Plugins/Experimental/Paper2D/Source/Paper2D/Private/PaperAnimatedRenderComponent.cpp
+++ b/Engine/Plugins/Experimental/Paper2D/Source/Paper2D/Private/PaperAnimatedRenderComponent.cpp
@@ -168,3 +168,8 @@ void UPaperAnimatedRenderComponent::SetCurrentTime(float NewTime)
 		CalculateCurrentFrame();
 	}
 }
+
+const UObject* UPaperAnimatedRenderComponent::AdditionalStatObject() const
+{
+	return SourceFlipbook;
+}
diff --git a/Engine/Plugins/Experimental/Paper2D/Source/Paper2D/Private/PaperRenderComponent.cpp b/Engine/Plugins/Experimental/Paper2D/Source/Paper2D/Private/PaperRenderComponent.cpp
--- a/Engine/Plugins/Experimental/Paper2D/Source/Paper2D/Private/PaperRenderComponent.cpp
+++ b/Engine/Plugins/Experimental/Paper2D/Source/Paper2D/Private/PaperRenderComponent.cpp
@@ -244,4 +244,9 @@ void UPaperRenderComponent::SetSpriteColor(FLinearColor NewColor)
 FLinearColor UPaperRenderComponent::GetWireframeColor() const
 {
 	return FLinearColor::Yellow;
-}
\ No newline at end of file
+}
+
+const UObject* UPaperRenderComponent::AdditionalStatObject() const
+{
+	return SourceSprite;
+}
diff --git a/Engine/Plugins/Experimental/Paper2D/Source/Paper2D/Private/PaperTileMapRenderComponent.cpp b/Engine/Plugins/Experimental/Paper2D/Source/Paper2D/Private/PaperTileMapRenderComponent.cpp
--- a/Engine/Plugins/Experimental/Paper2D/Source/Paper2D/Private/PaperTileMapRenderComponent.cpp
+++ b/Engine/Plugins/Experimental/Paper2D/Source/Paper2D/Private/PaperTileMapRenderComponent.cpp
@@ -104,4 +104,9 @@ void UPaperTileMapRenderComponent::UpdateBodySetup()
 	}
 }
 
+const UObject* UPaperTileMapRenderComponent::AdditionalStatObject() const
+{
+	return DefaultLayerTileSet;
+}
+
 #undef LOCTEXT_NAMESPACE
\ No newline at end of file
