Paper2D: Add placeholder atlas asset types

diff --git a/Engine/Plugins/Experimental/Paper2D/Source/Paper2D/Classes/PaperSpriteAtlas.h b/Engine/Plugins/Experimental/Paper2D/Source/Paper2D/Classes/PaperSpriteAtlas.h
new file mode 100644
--- /dev/null
+++ b/Engine/Plugins/Experimental/Paper2D/Source/Paper2D/Classes/PaperSpriteAtlas.h
@@ -0,0 +1,22 @@
+// Copyright 1998-2014 Epic Games, Inc. All Rights Reserved.
+
+#pragma once
+
+#include "PaperSpriteAtlas.generated.h"
+
+// Groups together a set of sprites that will try to share the same texture atlas (allowing them to be combined into a single draw call)
+UCLASS(MinimalAPI, BlueprintType, meta=(DisplayThumbnail = "true"))
+class UPaperSpriteAtlas : public UObject
+{
+	GENERATED_UCLASS_BODY()
+
+#if WITH_EDITORONLY_DATA
+	// Description of this atlas, which shows up in the content browser tooltip
+	UPROPERTY(EditAnywhere, Category=General)
+	FString AtlasDescription;
+#endif
+
+	// UObject interface
+	virtual void GetAssetRegistryTags(TArray<FAssetRegistryTag>& OutTags) const OVERRIDE;
+	// End of UObject interface
+};
diff --git a/Engine/Plugins/Experimental/Paper2D/Source/Paper2D/Private/PaperSpriteAtlas.cpp b/Engine/Plugins/Experimental/Paper2D/Source/Paper2D/Private/PaperSpriteAtlas.cpp
new file mode 100644
--- /dev/null
+++ b/Engine/Plugins/Experimental/Paper2D/Source/Paper2D/Private/PaperSpriteAtlas.cpp
@@ -0,0 +1,20 @@
+// Copyright 1998-2014 Epic Games, Inc. All Rights Reserved.
+
+#include "Paper2DPrivatePCH.h"
+
+//////////////////////////////////////////////////////////////////////////
+// UPaperSpriteAtlas
+
+UPaperSpriteAtlas::UPaperSpriteAtlas(const FPostConstructInitializeProperties& PCIP)
+	: Super(PCIP)
+{
+}
+
+void UPaperSpriteAtlas::GetAssetRegistryTags(TArray<FAssetRegistryTag>& OutTags) const
+{
+	Super::GetAssetRegistryTags(OutTags);
+
+#if WITH_EDITORONLY_DATA
+	OutTags.Add(FAssetRegistryTag(GET_MEMBER_NAME_CHECKED(UPaperSpriteAtlas, AtlasDescription), AtlasDescription, FAssetRegistryTag::TT_Hidden));
+#endif
+}
diff --git a/Engine/Plugins/Experimental/Paper2D/Source/Paper2DEditor/Private/Atlasing/AtlasAssetTypeActions.cpp b/Engine/Plugins/Experimental/Paper2D/Source/Paper2DEditor/Private/Atlasing/AtlasAssetTypeActions.cpp
new file mode 100644
--- /dev/null
+++ b/Engine/Plugins/Experimental/Paper2D/Source/Paper2DEditor/Private/Atlasing/AtlasAssetTypeActions.cpp
@@ -0,0 +1,55 @@
+// Copyright 1998-2014 Epic Games, Inc. All Rights Reserved.
+
+#include "Paper2DEditorPrivatePCH.h"
+
+#include "AtlasAssetTypeActions.h"
+//#include "FlipbookEditor/FlipbookEditor.h"
+
+#define LOCTEXT_NAMESPACE "AssetTypeActions"
+
+//////////////////////////////////////////////////////////////////////////
+// FAtlasAssetTypeActions
+
+FText FAtlasAssetTypeActions::GetName() const
+{
+	return LOCTEXT("FAtlasAssetTypeActionsName", "Sprite Atlas Group");
+}
+
+FColor FAtlasAssetTypeActions::GetTypeColor() const
+{
+	return FColor(0, 255, 255);
+}
+
+UClass* FAtlasAssetTypeActions::GetSupportedClass() const
+{
+	return UPaperSpriteAtlas::StaticClass();
+}
+
+void FAtlasAssetTypeActions::OpenAssetEditor(const TArray<UObject*>& InObjects, TSharedPtr<class IToolkitHost> EditWithinLevelEditor)
+{
+	//@TODO: Atlas will need a custom editor at some point
+	FSimpleAssetEditor::CreateEditor(EToolkitMode::Standalone, EditWithinLevelEditor, InObjects);
+}
+
+uint32 FAtlasAssetTypeActions::GetCategories()
+{
+	return EAssetTypeCategories::Misc;
+}
+
+FText FAtlasAssetTypeActions::GetAssetDescription(const FAssetData& AssetData) const
+{
+	if (const FString* pDescription = AssetData.TagsAndValues.Find(GET_MEMBER_NAME_CHECKED(UPaperSpriteAtlas, AtlasDescription)))
+	{
+		if (!pDescription->IsEmpty())
+		{
+			const FString DescriptionStr(*pDescription);
+			return FText::FromString(DescriptionStr.Replace(TEXT("\\n"), TEXT("\n")));
+		}
+	}
+
+	return FText::GetEmpty();
+}
+
+//////////////////////////////////////////////////////////////////////////
+
+#undef LOCTEXT_NAMESPACE
\ No newline at end of file
diff --git a/Engine/Plugins/Experimental/Paper2D/Source/Paper2DEditor/Private/Atlasing/AtlasAssetTypeActions.h b/Engine/Plugins/Experimental/Paper2D/Source/Paper2DEditor/Private/Atlasing/AtlasAssetTypeActions.h
new file mode 100644
--- /dev/null
+++ b/Engine/Plugins/Experimental/Paper2D/Source/Paper2DEditor/Private/Atlasing/AtlasAssetTypeActions.h
@@ -0,0 +1,18 @@
+// Copyright 1998-2014 Epic Games, Inc. All Rights Reserved.
+
+#pragma once
+
+#include "AssetTypeActions_Base.h"
+
+class FAtlasAssetTypeActions : public FAssetTypeActions_Base
+{
+public:
+	// IAssetTypeActions interface
+	virtual FText GetName() const OVERRIDE;
+	virtual FColor GetTypeColor() const OVERRIDE;
+	virtual UClass* GetSupportedClass() const OVERRIDE;
+	virtual void OpenAssetEditor(const TArray<UObject*>& InObjects, TSharedPtr<class IToolkitHost> EditWithinLevelEditor = TSharedPtr<IToolkitHost>()) OVERRIDE;
+	virtual uint32 GetCategories() OVERRIDE;
+	virtual FText GetAssetDescription(const FAssetData& AssetData) const OVERRIDE;
+	// End of IAssetTypeActions interface
+};
\ No newline at end of file
diff --git a/Engine/Plugins/Experimental/Paper2D/Source/Paper2DEditor/Private/Atlasing/PaperSpriteAtlasFactory.cpp b/Engine/Plugins/Experimental/Paper2D/Source/Paper2DEditor/Private/Atlasing/PaperSpriteAtlasFactory.cpp
new file mode 100644
--- /dev/null
+++ b/Engine/Plugins/Experimental/Paper2D/Source/Paper2DEditor/Private/Atlasing/PaperSpriteAtlasFactory.cpp
@@ -0,0 +1,22 @@
+// Copyright 1998-2014 Epic Games, Inc. All Rights Reserved.
+
+#include "Paper2DEditorPrivatePCH.h"
+#include "PaperSpriteAtlasFactory.h"
+
+//////////////////////////////////////////////////////////////////////////
+// UPaperSpriteAtlasFactory
+
+UPaperSpriteAtlasFactory::UPaperSpriteAtlasFactory(const class FPostConstructInitializeProperties& PCIP)
+	: Super(PCIP)
+{
+	SupportedClass = UPaperSpriteAtlas::StaticClass();
+
+	bCreateNew = true;
+	bEditorImport = false;
+	bEditAfterNew = true;
+}
+
+UObject* UPaperSpriteAtlasFactory::FactoryCreateNew(UClass* InClass, UObject* InParent, FName InName, EObjectFlags Flags, UObject* Context, FFeedbackContext* Warn)
+{
+	return ConstructObject<UObject>(SupportedClass, InParent, InName, Flags);
+}
diff --git a/Engine/Plugins/Experimental/Paper2D/Source/Paper2DEditor/Private/Atlasing/PaperSpriteAtlasFactory.h b/Engine/Plugins/Experimental/Paper2D/Source/Paper2DEditor/Private/Atlasing/PaperSpriteAtlasFactory.h
new file mode 100644
--- /dev/null
+++ b/Engine/Plugins/Experimental/Paper2D/Source/Paper2DEditor/Private/Atlasing/PaperSpriteAtlasFactory.h
@@ -0,0 +1,16 @@
+// Copyright 1998-2014 Epic Games, Inc. All Rights Reserved.
+
+#pragma once
+
+#include "PaperSpriteAtlasFactory.generated.h"
+
+UCLASS()
+class UPaperSpriteAtlasFactory : public UFactory
+{
+	GENERATED_UCLASS_BODY()
+
+	// UFactory interface
+	virtual UObject* FactoryCreateNew(UClass* Class, UObject* InParent, FName Name, EObjectFlags Flags, UObject* Context, FFeedbackContext* Warn) OVERRIDE;
+	// End of UFactory interface
+};
+
diff --git a/Engine/Plugins/Experimental/Paper2D/Source/Paper2DEditor/Private/Paper2DEditorModule.cpp b/Engine/Plugins/Experimental/Paper2D/Source/Paper2DEditor/Private/Paper2DEditorModule.cpp
--- a/Engine/Plugins/Experimental/Paper2D/Source/Paper2DEditor/Private/Paper2DEditorModule.cpp
+++ b/Engine/Plugins/Experimental/Paper2D/Source/Paper2DEditor/Private/Paper2DEditorModule.cpp
@@ -2,9 +2,12 @@
 
 #include "Paper2DEditorPrivatePCH.h"
 #include "Paper2DEditorModule.h"
+
 #include "SpriteAssetTypeActions.h"
 #include "FlipbookAssetTypeActions.h"
 #include "TileSetAssetTypeActions.h"
+#include "Atlasing/AtlasAssetTypeActions.h"
+
 #include "AssetToolsModule.h"
 #include "PropertyEditorModule.h"
 #include "PaperStyle.h"
@@ -19,8 +22,6 @@
 
 #include "ContentBrowserExtensions/ContentBrowserExtensions.h"
 
-
-
 //////////////////////////////////////////////////////////////////////////
 // FPaper2DEditor
 
@@ -42,7 +43,7 @@ private:
 	TSharedPtr<FExtensibilityManager> FlipbookEditor_MenuExtensibilityManager;
 	TSharedPtr<FExtensibilityManager> FlipbookEditor_ToolBarExtensibilityManager;
 
-	/** All created asset type actions.  Cached here so that we can unregister it during shutdown. */
+	/** All created asset type actions.  Cached here so that we can unregister them during shutdown. */
 	TArray< TSharedPtr<IAssetTypeActions> > CreatedAssetTypeActions;
 
 	TSharedPtr<FEdModeTileMap> TileMapEditorModePtr;
@@ -69,6 +70,7 @@ public:
 		RegisterAssetTypeAction(AssetTools, MakeShareable(new FSpriteAssetTypeActions));
 		RegisterAssetTypeAction(AssetTools, MakeShareable(new FFlipbookAssetTypeActions));
 		RegisterAssetTypeAction(AssetTools, MakeShareable(new FTileSetAssetTypeActions));
+		RegisterAssetTypeAction(AssetTools, MakeShareable(new FAtlasAssetTypeActions));
 
 		PaperSpriteBroker = MakeShareable(new FPaperSpriteAssetBroker);
 		FComponentAssetBrokerage::RegisterBroker(PaperSpriteBroker, UPaperRenderComponent::StaticClass(), true, true);
