Paper2D: Move UPaperFlipbook to Animation category in Content Browser

diff --git a/Engine/Plugins/Experimental/Paper2D/Source/Paper2DEditor/Private/FlipbookAssetTypeActions.cpp b/Engine/Plugins/Experimental/Paper2D/Source/Paper2DEditor/Private/FlipbookAssetTypeActions.cpp
--- a/Engine/Plugins/Experimental/Paper2D/Source/Paper2DEditor/Private/FlipbookAssetTypeActions.cpp
+++ b/Engine/Plugins/Experimental/Paper2D/Source/Paper2DEditor/Private/FlipbookAssetTypeActions.cpp
@@ -8,11 +8,11 @@
 #define LOCTEXT_NAMESPACE "AssetTypeActions"
 
 //////////////////////////////////////////////////////////////////////////
-// FSpriteAssetTypeActions
+// FFlipbookAssetTypeActions
 
 FText FFlipbookAssetTypeActions::GetName() const
 {
-	return LOCTEXT("FFlipbookAssetTypeActionsName", "Flipbook");
+	return LOCTEXT("FFlipbookAssetTypeActionsName", "Sprite Flipbook");
 }
 
 FColor FFlipbookAssetTypeActions::GetTypeColor() const
@@ -41,7 +41,7 @@ void FFlipbookAssetTypeActions::OpenAssetEditor(const TArray<UObject*>& InObject
 
 uint32 FFlipbookAssetTypeActions::GetCategories()
 {
-	return EAssetTypeCategories::Misc;
+	return EAssetTypeCategories::Animation;
 }
 
 //////////////////////////////////////////////////////////////////////////
