Paper2D: Add missing AssetData.h include (non-unity)

diff --git a/Engine/Plugins/Experimental/Paper2D/Source/Paper2DEditor/Private/Atlasing/AtlasAssetTypeActions.cpp b/Engine/Plugins/Experimental/Paper2D/Source/Paper2DEditor/Private/Atlasing/AtlasAssetTypeActions.cpp
--- a/Engine/Plugins/Experimental/Paper2D/Source/Paper2DEditor/Private/Atlasing/AtlasAssetTypeActions.cpp
+++ b/Engine/Plugins/Experimental/Paper2D/Source/Paper2DEditor/Private/Atlasing/AtlasAssetTypeActions.cpp
@@ -3,7 +3,7 @@
 #include "Paper2DEditorPrivatePCH.h"
 
 #include "AtlasAssetTypeActions.h"
-//#include "FlipbookEditor/FlipbookEditor.h"
+#include "AssetData.h"
 
 #define LOCTEXT_NAMESPACE "AssetTypeActions"
 
