Paper2D: Add prototype sprite atlas generation
Note: Still several unresolved data issues if someone else has a sprite checked out or you fail to save the sprites, etc... when rebuilding the atlas

diff --git a/Engine/Plugins/Experimental/Paper2D/Source/Paper2D/Classes/PaperSprite.h b/Engine/Plugins/Experimental/Paper2D/Source/Paper2D/Classes/PaperSprite.h
--- a/Engine/Plugins/Experimental/Paper2D/Source/Paper2D/Classes/PaperSprite.h
+++ b/Engine/Plugins/Experimental/Paper2D/Source/Paper2D/Classes/PaperSprite.h
@@ -45,16 +45,25 @@ class PAPER2D_API UPaperSprite : public UObject, public IInterface_CollisionData
 protected:
 
 #if WITH_EDITORONLY_DATA
+	// Position with SourceTexture
 	UPROPERTY(Category=Sprite, EditAnywhere, AssetRegistrySearchable)
 	FVector2D SourceUV;
 
+	// Dimensions within SourceTexture
 	UPROPERTY(Category=Sprite, EditAnywhere, AssetRegistrySearchable)
 	FVector2D SourceDimension;
+
+	// Position within BakedSourceTexture
+	UPROPERTY()
+	FVector2D BakedSourceUV;
 #endif
 
 	UPROPERTY(Category=Sprite, EditAnywhere, AssetRegistrySearchable)
 	UTexture2D* SourceTexture;
 
+	UPROPERTY()
+	UTexture2D* BakedSourceTexture;
+
 	// The material to use on a sprite instance if not overridden
 	UPROPERTY(Category=Sprite, EditAnywhere, BlueprintReadOnly)
 	UMaterialInterface* DefaultMaterial;
@@ -98,6 +107,10 @@ protected:
 	// Custom render geometry polygons (in texture space)
 	UPROPERTY(Category=Rendering, EditAnywhere)
 	FSpritePolygonCollection RenderGeometry;
+
+	// Spritesheet group that this sprite belongs to
+	UPROPERTY(Category=Sprite, EditAnywhere, AssetRegistrySearchable)
+	class UPaperSpriteAtlas* AtlasGroup;
 #endif
 
 public:
@@ -115,6 +128,7 @@ public:
 	virtual void PostEditChangeProperty(FPropertyChangedEvent& PropertyChangedEvent) OVERRIDE;
 	virtual void Serialize(FArchive& Ar) OVERRIDE;
 	virtual void PostLoad() OVERRIDE;
+	virtual void GetAssetRegistryTags(TArray<FAssetRegistryTag>& OutTags) const OVERRIDE;
 	// End of UObject interface
 
 	FVector2D ConvertTextureSpaceToPivotSpace(FVector2D Input) const;
@@ -153,7 +167,8 @@ public:
 	UTexture2D* GetSourceTexture() const { return SourceTexture; }
 #endif
 
-	UTexture2D* GetPrimaryAtlasTexture() const { return SourceTexture; }
+	// Returns the texture this should be rendered with
+	UTexture2D* GetBakedTexture() const;
 
 	UMaterialInterface* GetDefaultMaterial() const { return DefaultMaterial; }
 
@@ -176,5 +191,6 @@ public:
 	friend class FSpriteDetailsCustomization;
 	friend class FSpriteSelectedVertex;
 	friend class FSpriteSelectedEdge;
+	friend struct FPaperAtlasGenerator;
 };
 
diff --git a/Engine/Plugins/Experimental/Paper2D/Source/Paper2D/Classes/PaperSpriteAtlas.h b/Engine/Plugins/Experimental/Paper2D/Source/Paper2D/Classes/PaperSpriteAtlas.h
--- a/Engine/Plugins/Experimental/Paper2D/Source/Paper2D/Classes/PaperSpriteAtlas.h
+++ b/Engine/Plugins/Experimental/Paper2D/Source/Paper2D/Classes/PaperSpriteAtlas.h
@@ -14,9 +14,30 @@ class UPaperSpriteAtlas : public UObject
 	// Description of this atlas, which shows up in the content browser tooltip
 	UPROPERTY(EditAnywhere, Category=General)
 	FString AtlasDescription;
+
+	// Maximum atlas page width (single pages might be smaller)
+	UPROPERTY(EditAnywhere, Category = General)
+	int32 MaxWidth;
+
+	// Maximum atlas page height (single pages might be smaller)
+	UPROPERTY(EditAnywhere, Category = General)
+	int32 MaxHeight;
+
+	// List of generated atlas textures
+	UPROPERTY(VisibleAnywhere, Category=General)
+	TArray<UTexture*> GeneratedTextures;
+
+	// The GUID of the atlas group, used to match up sprites that belong to this group even thru atlas renames
+	UPROPERTY(VisibleAnywhere, Category=General)
+	FGuid AtlasGUID;
 #endif
 
+
+
+public:
 	// UObject interface
+	virtual void PostInitProperties() OVERRIDE;
+	virtual void PostDuplicate(bool bDuplicateForPIE) OVERRIDE;
 	virtual void GetAssetRegistryTags(TArray<FAssetRegistryTag>& OutTags) const OVERRIDE;
 	// End of UObject interface
 };
diff --git a/Engine/Plugins/Experimental/Paper2D/Source/Paper2D/Private/PaperSprite.cpp b/Engine/Plugins/Experimental/Paper2D/Source/Paper2D/Private/PaperSprite.cpp
--- a/Engine/Plugins/Experimental/Paper2D/Source/Paper2D/Private/PaperSprite.cpp
+++ b/Engine/Plugins/Experimental/Paper2D/Source/Paper2D/Private/PaperSprite.cpp
@@ -6,6 +6,8 @@
 #include "PaperCustomVersion.h"
 
 #if WITH_EDITORONLY_DATA
+
+#include "PaperSpriteAtlas.h"
 #include "GeomTools.h"
 
 //////////////////////////////////////////////////////////////////////////
@@ -352,7 +354,7 @@ void FSpriteDrawCallRecord::BuildFromSprite(class UPaperSprite* Sprite)
 		//SourceDimension = Sprite->GetSourceSize();
 		Destination = FVector::ZeroVector;
 		//DestinationDimension = Sprite->GetSourceSize();
-		Texture = Sprite->GetPrimaryAtlasTexture();
+		Texture = Sprite->GetBakedTexture();
 		Color = FLinearColor::White;
 		BlendMode = Sprite->BlendMode;
 
@@ -611,14 +613,20 @@ void UPaperSprite::RebuildRenderData()
 	Triangulate(RenderGeometry, /*out*/ TriangluatedPoints);
 
 	// Adjust for the pivot and store in the baked geometry buffer
-	const float InverseWidth = (SourceTexture != NULL) ? (1.0f / SourceTexture->GetSizeX()) : 1.0f;
-	const float InverseHeight = (SourceTexture != NULL) ? (1.0f / SourceTexture->GetSizeY()) : 1.0f;
+	UTexture2D* EffectiveTexture = GetBakedTexture();
+
+	const float InverseWidth = (EffectiveTexture != nullptr) ? (1.0f / EffectiveTexture->GetSizeX()) : 1.0f;
+	const float InverseHeight = (EffectiveTexture != nullptr) ? (1.0f / EffectiveTexture->GetSizeY()) : 1.0f;
+
+	const FVector2D DeltaUV((BakedSourceTexture != nullptr) ? (BakedSourceUV - SourceUV) : FVector2D::ZeroVector);
+
 
 	BakedRenderData.Empty(TriangluatedPoints.Num());
 	for (int32 PointIndex = 0; PointIndex < TriangluatedPoints.Num(); ++PointIndex)
 	{
-		const FVector2D& UV = TriangluatedPoints[PointIndex];
-		const FVector2D PivotSpacePos = ConvertTextureSpaceToPivotSpace(UV);
+		const FVector2D& SourcePos = TriangluatedPoints[PointIndex];
+		const FVector2D PivotSpacePos = ConvertTextureSpaceToPivotSpace(SourcePos);
+		const FVector2D UV(SourcePos + DeltaUV);
 
 		new (BakedRenderData) FVector4(PivotSpacePos.X, PivotSpacePos.Y, UV.X * InverseWidth, UV.Y * InverseHeight);
 	}
@@ -1041,6 +1049,16 @@ FVector2D UPaperSprite::GetPivotPosition() const
 	};
 }
 
+void UPaperSprite::GetAssetRegistryTags(TArray<FAssetRegistryTag>& OutTags) const
+{
+	Super::GetAssetRegistryTags(OutTags);
+
+	if (AtlasGroup != nullptr)
+	{
+		OutTags.Add(FAssetRegistryTag(TEXT("AtlasGroupGUID"), AtlasGroup->AtlasGUID.ToString(EGuidFormats::Digits), FAssetRegistryTag::TT_Hidden));
+	}
+}
+
 #endif
 
 bool UPaperSprite::GetPhysicsTriMeshData(FTriMeshCollisionData* OutCollisionData, bool InUseAllTriData)
@@ -1115,3 +1133,8 @@ void UPaperSprite::PostLoad()
 	}
 #endif
 }
+
+UTexture2D* UPaperSprite::GetBakedTexture() const
+{
+	return (BakedSourceTexture != nullptr) ? BakedSourceTexture : SourceTexture;
+}
\ No newline at end of file
diff --git a/Engine/Plugins/Experimental/Paper2D/Source/Paper2D/Private/PaperSpriteAtlas.cpp b/Engine/Plugins/Experimental/Paper2D/Source/Paper2D/Private/PaperSpriteAtlas.cpp
--- a/Engine/Plugins/Experimental/Paper2D/Source/Paper2D/Private/PaperSpriteAtlas.cpp
+++ b/Engine/Plugins/Experimental/Paper2D/Source/Paper2D/Private/PaperSpriteAtlas.cpp
@@ -7,6 +7,8 @@
 
 UPaperSpriteAtlas::UPaperSpriteAtlas(const FPostConstructInitializeProperties& PCIP)
 	: Super(PCIP)
+	, MaxWidth(2048)
+	, MaxHeight(2048)
 {
 }
 
@@ -18,3 +20,15 @@ void UPaperSpriteAtlas::GetAssetRegistryTags(TArray<FAssetRegistryTag>& OutTags)
 	OutTags.Add(FAssetRegistryTag(GET_MEMBER_NAME_CHECKED(UPaperSpriteAtlas, AtlasDescription), AtlasDescription, FAssetRegistryTag::TT_Hidden));
 #endif
 }
+
+void UPaperSpriteAtlas::PostDuplicate(bool bDuplicateForPIE)
+{
+	Super::PostDuplicate(bDuplicateForPIE);
+	FPlatformMisc::CreateGuid(AtlasGUID);
+}
+
+void UPaperSpriteAtlas::PostInitProperties()
+{
+	Super::PostInitProperties();
+	FPlatformMisc::CreateGuid(AtlasGUID);
+}
diff --git a/Engine/Plugins/Experimental/Paper2D/Source/Paper2DEditor/Private/Atlasing/PaperAtlasGenerator.cpp b/Engine/Plugins/Experimental/Paper2D/Source/Paper2DEditor/Private/Atlasing/PaperAtlasGenerator.cpp
new file mode 100644
--- /dev/null
+++ b/Engine/Plugins/Experimental/Paper2D/Source/Paper2DEditor/Private/Atlasing/PaperAtlasGenerator.cpp
@@ -0,0 +1,192 @@
+// Copyright 1998-2014 Epic Games, Inc. All Rights Reserved.
+
+#include "Paper2DEditorPrivatePCH.h"
+#include "PaperAtlasGenerator.h"
+
+#include "AssetRegistryModule.h"
+#include "Textures/TextureAtlas.h"
+
+//////////////////////////////////////////////////////////////////////////
+// 
+
+void CopyTextureData(const uint8* Source, uint8* Dest, uint32 SizeX, uint32 SizeY, uint32 BytesPerPixel, uint32 SourceStride, uint32 DestStride)
+{
+	const uint32 NumBytesPerRow = SizeX * BytesPerPixel;
+
+	for (uint32 Y = 0; Y < SizeY; ++Y)
+	{
+		FMemory::Memcpy(
+			Dest + (DestStride * Y),
+			Source + (SourceStride * Y),
+			NumBytesPerRow
+			);
+	}
+}
+
+struct FSingleTexturePaperAtlas : public FSlateTextureAtlas
+{
+public:
+	FSingleTexturePaperAtlas(uint32 InWidth, uint32 InHeight)
+	: FSlateTextureAtlas(InWidth, InHeight, sizeof(FColor), 0)
+	{
+		//@TODO: check(StrideBytes >= 4*Width);
+	}
+
+	virtual void ConditionalUpdateTexture() OVERRIDE {}
+
+	void CopyToTexture(UTexture2D* Texture)
+	{
+		const int32 NumMips = 1;
+
+		Texture->Source.Init(AtlasWidth, AtlasHeight, /*NumSlices=*/ 1, NumMips, ETextureSourceFormat::TSF_BGRA8, AtlasData.GetTypedData());
+		Texture->UpdateResource();
+	}
+
+	const FAtlasedTextureSlot* AddSprite(UPaperSprite* Sprite)
+	{
+		const FVector2D SpriteSizeFloat = Sprite->GetSourceSize();
+		const FIntPoint SpriteSize(FMath::Trunc(SpriteSizeFloat.X), FMath::Trunc(SpriteSizeFloat.Y));
+
+		TArray<uint8> DummyBuffer;
+		DummyBuffer.AddZeroed(SpriteSize.X * SpriteSize.Y * Stride);
+		
+		check(Sprite->GetSourceTexture());
+		FTextureSource& SourceData = Sprite->GetSourceTexture()->Source;
+
+		//@TODO: Handle different texture formats!
+		if (SourceData.GetFormat() == TSF_BGRA8)
+		{
+			uint32 BytesPerPixel = SourceData.GetBytesPerPixel();
+			uint8* OffsetSource = SourceData.LockMip(0) + (FMath::Trunc(Sprite->GetSourceUV().X) + FMath::Trunc(Sprite->GetSourceUV().Y) * SourceData.GetSizeX()) * BytesPerPixel;
+			uint8* OffsetDest = DummyBuffer.GetTypedData();
+
+			CopyTextureData(OffsetSource, OffsetDest, SpriteSize.X, SpriteSize.Y, BytesPerPixel, SourceData.GetSizeX() * BytesPerPixel, SpriteSize.X * BytesPerPixel);
+
+			SourceData.UnlockMip(0);
+		}
+		else
+		{
+			UE_LOG(LogPaper2DEditor, Error, TEXT("Sprite %s is not BGRA8, which isn't supported in atlases yet"), *(Sprite->GetPathName()));
+		}
+
+		const FAtlasedTextureSlot* Slot = AddTexture(SpriteSize.X, SpriteSize.Y, DummyBuffer);
+		if (Slot != nullptr)
+		{
+			SpriteToSlotMap.Add(Sprite, Slot);
+		}
+
+		return Slot;
+	}
+
+public:
+	TMap<UPaperSprite*, const FAtlasedTextureSlot*> SpriteToSlotMap;
+};
+
+//////////////////////////////////////////////////////////////////////////
+// FPaperAtlasGenerator
+
+void FPaperAtlasGenerator::HandleAssetChangedEvent(UPaperSpriteAtlas* Atlas)
+{
+	Atlas->MaxWidth = FMath::Clamp(Atlas->MaxWidth, 16, 4096);
+	Atlas->MaxHeight = FMath::Clamp(Atlas->MaxHeight, 16, 4096);
+
+	// Find all sprites that reference the atlas and force them loaded
+	FAssetRegistryModule& AssetRegistryModule = FModuleManager::LoadModuleChecked<FAssetRegistryModule>(TEXT("AssetRegistry"));
+
+	TArray<FAssetData> AssetList;
+	TMultiMap<FName, FString> TagsAndValues;
+	TagsAndValues.Add(TEXT("AtlasGroupGUID"), Atlas->AtlasGUID.ToString(EGuidFormats::Digits));
+	AssetRegistryModule.Get().GetAssetsByTagValues(TagsAndValues, AssetList);
+
+	TIndirectArray<FSingleTexturePaperAtlas> AtlasGenerators;
+
+	// Start off with one page
+	new (AtlasGenerators) FSingleTexturePaperAtlas(Atlas->MaxWidth, Atlas->MaxHeight);
+
+	for (const FAssetData& SpriteRef : AssetList)
+	{
+		if (UPaperSprite* Sprite = Cast<UPaperSprite>(SpriteRef.GetAsset()))
+		{
+			//@TODO: Use the tight bounds instead of the source bounds
+			const FVector2D SpriteSizeFloat = Sprite->GetSourceSize();
+			const FIntPoint SpriteSize(FMath::Trunc(SpriteSizeFloat.X), FMath::Trunc(SpriteSizeFloat.Y));
+
+			if (Sprite->GetSourceTexture() == nullptr)
+			{
+				UE_LOG(LogPaper2DEditor, Error, TEXT("Sprite %s has no source texture and cannot be packed"), *(Sprite->GetPathName()));
+				continue;
+			}
+
+			//@TODO: Take padding into account (push this into the generator)
+			if ((SpriteSize.X > Atlas->MaxWidth) || (SpriteSize.Y >= Atlas->MaxHeight))
+			{
+				// This sprite cannot ever fit into an atlas page
+				UE_LOG(LogPaper2DEditor, Error, TEXT("Sprite %s (%d x %d) can never fit into atlas %s (%d x %d) due to maximum page size restrictions"),
+					*(Sprite->GetPathName()),
+					SpriteSize.X,
+					SpriteSize.Y,
+					*(Atlas->GetPathName()),
+					Atlas->MaxWidth,
+					Atlas->MaxHeight);
+			}
+			else
+			{
+				const FAtlasedTextureSlot* Slot = nullptr;
+
+				// Does it fit in any existing pages?
+				for (auto& Generator : AtlasGenerators)
+				{
+					Slot = Generator.AddSprite(Sprite);
+
+					if (Slot != nullptr)
+					{
+						break;
+					}
+				}
+
+				if (Slot == nullptr)
+				{
+					// Doesn't fit in any current pages, make a new one
+					FSingleTexturePaperAtlas* NewGenerator = new (AtlasGenerators) FSingleTexturePaperAtlas(Atlas->MaxWidth, Atlas->MaxHeight);
+					Slot = NewGenerator->AddSprite(Sprite);
+				}
+
+				if (Slot != nullptr)
+				{
+					UE_LOG(LogPaper2DEditor, Warning, TEXT("Allocated %s to (%d,%d)"), *(Sprite->GetPathName()), Slot->X, Slot->Y);
+				}
+				else
+				{
+					// ERROR: Didn't fit into a brand new page, maybe padding was wrong
+					UE_LOG(LogPaper2DEditor, Error, TEXT("Failed to allocate %s into a brand new page"), *(Sprite->GetPathName()));
+				}
+			}
+		}
+	}
+
+	// Turn the generators back into textures
+	Atlas->GeneratedTextures.Empty(AtlasGenerators.Num());
+	for (int32 GeneratorIndex = 0; GeneratorIndex < AtlasGenerators.Num(); ++GeneratorIndex)
+	{
+		FSingleTexturePaperAtlas& AtlasPage = AtlasGenerators[GeneratorIndex];
+
+		UTexture2D* Texture = NewNamedObject<UTexture2D>(Atlas, NAME_None, RF_Public);
+		Atlas->GeneratedTextures.Add(Texture);
+
+		AtlasPage.CopyToTexture(Texture);
+
+		// Now update the baked data for all the sprites to point there
+		for (auto& SpriteToSlot : AtlasPage.SpriteToSlotMap)
+		{
+			UPaperSprite* Sprite = SpriteToSlot.Key;
+			Sprite->Modify();
+			const FAtlasedTextureSlot* Slot = SpriteToSlot.Value;
+			Sprite->BakedSourceTexture = Texture;
+			Sprite->BakedSourceUV = FVector2D(Slot->X, Slot->Y);
+			Sprite->RebuildRenderData();
+		}
+	}
+
+	//@TODO: Adjust the atlas rebuild code so that it tries to preserve existing sprites (optimize for the 'just added one to the end' case, preventing dirtying lots of assets unnecessarily)
+	//@TODO: invoke this code when a sprite has the atlas group property modified
+}
diff --git a/Engine/Plugins/Experimental/Paper2D/Source/Paper2DEditor/Private/Atlasing/PaperAtlasGenerator.h b/Engine/Plugins/Experimental/Paper2D/Source/Paper2DEditor/Private/Atlasing/PaperAtlasGenerator.h
new file mode 100644
--- /dev/null
+++ b/Engine/Plugins/Experimental/Paper2D/Source/Paper2DEditor/Private/Atlasing/PaperAtlasGenerator.h
@@ -0,0 +1,13 @@
+// Copyright 1998-2014 Epic Games, Inc. All Rights Reserved.
+
+#pragma once
+
+//////////////////////////////////////////////////////////////////////////
+// FPaperAtlasGenerator
+
+struct FPaperAtlasGenerator
+{
+
+public:
+	static void HandleAssetChangedEvent(UPaperSpriteAtlas* Atlas);
+};
diff --git a/Engine/Plugins/Experimental/Paper2D/Source/Paper2DEditor/Private/Paper2DEditorModule.cpp b/Engine/Plugins/Experimental/Paper2D/Source/Paper2DEditor/Private/Paper2DEditorModule.cpp
--- a/Engine/Plugins/Experimental/Paper2D/Source/Paper2DEditor/Private/Paper2DEditorModule.cpp
+++ b/Engine/Plugins/Experimental/Paper2D/Source/Paper2DEditor/Private/Paper2DEditorModule.cpp
@@ -6,7 +6,6 @@
 #include "SpriteAssetTypeActions.h"
 #include "FlipbookAssetTypeActions.h"
 #include "TileSetAssetTypeActions.h"
-#include "Atlasing/AtlasAssetTypeActions.h"
 
 #include "AssetToolsModule.h"
 #include "PropertyEditorModule.h"
@@ -22,6 +21,12 @@
 
 #include "ContentBrowserExtensions/ContentBrowserExtensions.h"
 
+// Atlas support
+#include "Atlasing/AtlasAssetTypeActions.h"
+#include "Atlasing/PaperAtlasGenerator.h"
+
+DEFINE_LOG_CATEGORY(LogPaper2DEditor);
+
 //////////////////////////////////////////////////////////////////////////
 // FPaper2DEditor
 
@@ -50,6 +55,9 @@ private:
 
 	TSharedPtr<IComponentAssetBroker> PaperSpriteBroker;
 	TSharedPtr<IComponentAssetBroker> PaperFlipbookBroker;
+
+	FCoreDelegates::FOnObjectPropertyChanged::FDelegate OnPropertyChangedHandle;
+
 public:
 	virtual void StartupModule() OVERRIDE
 	{
@@ -90,6 +98,10 @@ public:
 			PropertyModule.NotifyCustomizationModuleChanged();
 		}
 
+		// Register to be notified when properties are edited
+		OnPropertyChangedHandle = FCoreDelegates::FOnObjectPropertyChanged::FDelegate::CreateRaw(this, &FPaper2DEditor::OnPropertyChanged);
+		FCoreDelegates::OnObjectPropertyChanged.Add(OnPropertyChangedHandle);
+
 		// Register the thumbnail renderers
 		UThumbnailManager::Get().RegisterCustomRenderer(UPaperSprite::StaticClass(), UPaperSpriteThumbnailRenderer::StaticClass());
 		UThumbnailManager::Get().RegisterCustomRenderer(UPaperTileSet::StaticClass(), UPaperTileSetThumbnailRenderer::StaticClass());
@@ -125,6 +137,9 @@ public:
 			UThumbnailManager::Get().UnregisterCustomRenderer(UPaperSprite::StaticClass());
 			UThumbnailManager::Get().UnregisterCustomRenderer(UPaperTileSet::StaticClass());
 			UThumbnailManager::Get().UnregisterCustomRenderer(UPaperFlipbook::StaticClass());
+
+			// Unregister the property modification handler
+			FCoreDelegates::OnObjectPropertyChanged.Remove(OnPropertyChangedHandle);
 		}
 
 		// Unregister the details customization
@@ -150,6 +165,15 @@ private:
 		AssetTools.RegisterAssetTypeActions(Action);
 		CreatedAssetTypeActions.Add(Action);
 	}
+
+	// Called when a property on the specified object is modified
+	void OnPropertyChanged(UObject* ObjectBeingModified)
+	{
+		if (UPaperSpriteAtlas* Atlas = Cast<UPaperSpriteAtlas>(ObjectBeingModified))
+		{
+			FPaperAtlasGenerator::HandleAssetChangedEvent(Atlas);
+		}
+	}
 };
 
 //////////////////////////////////////////////////////////////////////////
diff --git a/Engine/Plugins/Experimental/Paper2D/Source/Paper2DEditor/Private/Paper2DEditorPrivatePCH.h b/Engine/Plugins/Experimental/Paper2D/Source/Paper2DEditor/Private/Paper2DEditorPrivatePCH.h
--- a/Engine/Plugins/Experimental/Paper2D/Source/Paper2DEditor/Private/Paper2DEditorPrivatePCH.h
+++ b/Engine/Plugins/Experimental/Paper2D/Source/Paper2DEditor/Private/Paper2DEditorPrivatePCH.h
@@ -5,4 +5,7 @@
 #include "UnrealEd.h"
 #include "Paper2DClasses.h"
 #include "Paper2DModule.h"
-#include "Paper2DEditorClasses.h"
\ No newline at end of file
+#include "Paper2DEditorClasses.h"
+
+DECLARE_LOG_CATEGORY_EXTERN(LogPaper2DEditor, Verbose, All);
+
