Paper2D: Fix FMath::Trunc references

diff --git a/Engine/Plugins/Experimental/Paper2D/Source/Paper2DEditor/Private/Atlasing/PaperAtlasGenerator.cpp b/Engine/Plugins/Experimental/Paper2D/Source/Paper2DEditor/Private/Atlasing/PaperAtlasGenerator.cpp
--- a/Engine/Plugins/Experimental/Paper2D/Source/Paper2DEditor/Private/Atlasing/PaperAtlasGenerator.cpp
+++ b/Engine/Plugins/Experimental/Paper2D/Source/Paper2DEditor/Private/Atlasing/PaperAtlasGenerator.cpp
@@ -45,7 +45,7 @@ public:
 	const FAtlasedTextureSlot* AddSprite(UPaperSprite* Sprite)
 	{
 		const FVector2D SpriteSizeFloat = Sprite->GetSourceSize();
-		const FIntPoint SpriteSize(FMath::Trunc(SpriteSizeFloat.X), FMath::Trunc(SpriteSizeFloat.Y));
+		const FIntPoint SpriteSize(FMath::TruncToInt(SpriteSizeFloat.X), FMath::TruncToInt(SpriteSizeFloat.Y));
 
 		TArray<uint8> DummyBuffer;
 		DummyBuffer.AddZeroed(SpriteSize.X * SpriteSize.Y * Stride);
@@ -57,7 +57,7 @@ public:
 		if (SourceData.GetFormat() == TSF_BGRA8)
 		{
 			uint32 BytesPerPixel = SourceData.GetBytesPerPixel();
-			uint8* OffsetSource = SourceData.LockMip(0) + (FMath::Trunc(Sprite->GetSourceUV().X) + FMath::Trunc(Sprite->GetSourceUV().Y) * SourceData.GetSizeX()) * BytesPerPixel;
+			uint8* OffsetSource = SourceData.LockMip(0) + (FMath::TruncToInt(Sprite->GetSourceUV().X) + FMath::TruncToInt(Sprite->GetSourceUV().Y) * SourceData.GetSizeX()) * BytesPerPixel;
 			uint8* OffsetDest = DummyBuffer.GetTypedData();
 
 			CopyTextureData(OffsetSource, OffsetDest, SpriteSize.X, SpriteSize.Y, BytesPerPixel, SourceData.GetSizeX() * BytesPerPixel, SpriteSize.X * BytesPerPixel);
@@ -109,7 +109,7 @@ void FPaperAtlasGenerator::HandleAssetChangedEvent(UPaperSpriteAtlas* Atlas)
 		{
 			//@TODO: Use the tight bounds instead of the source bounds
 			const FVector2D SpriteSizeFloat = Sprite->GetSourceSize();
-			const FIntPoint SpriteSize(FMath::Trunc(SpriteSizeFloat.X), FMath::Trunc(SpriteSizeFloat.Y));
+			const FIntPoint SpriteSize(FMath::TruncToInt(SpriteSizeFloat.X), FMath::TruncToInt(SpriteSizeFloat.Y));
 
 			if (Sprite->GetSourceTexture() == nullptr)
 			{
