Paper2D: Remove unused FSpriteDrawCallRecord::BlendMode

diff --git a/Engine/Plugins/Experimental/Paper2D/Source/Paper2D/Classes/SpriteDrawCall.h b/Engine/Plugins/Experimental/Paper2D/Source/Paper2D/Classes/SpriteDrawCall.h
--- a/Engine/Plugins/Experimental/Paper2D/Source/Paper2D/Classes/SpriteDrawCall.h
+++ b/Engine/Plugins/Experimental/Paper2D/Source/Paper2D/Classes/SpriteDrawCall.h
@@ -20,9 +20,6 @@ public:
 	UPROPERTY(Category=Sprite, EditAnywhere)
 	FLinearColor Color;
 
-	UPROPERTY(Category=Sprite, EditAnywhere)
-	TEnumAsByte<EBlendMode> BlendMode;
-
 	//@TODO: The rest of this struct doesn't need to be properties either, but has to be due to serialization for now
 	// Render triangle list (stored as loose vertices)
 	TArray< FVector4, TInlineAllocator<6> > RenderVerts;
@@ -32,7 +29,6 @@ public:
 	FSpriteDrawCallRecord()
 		: Texture(NULL)
 		, Color(FLinearColor::White)
-		, BlendMode(BLEND_Translucent)
 	{
 	}
 };
diff --git a/Engine/Plugins/Experimental/Paper2D/Source/Paper2D/Private/PaperSprite.cpp b/Engine/Plugins/Experimental/Paper2D/Source/Paper2D/Private/PaperSprite.cpp
--- a/Engine/Plugins/Experimental/Paper2D/Source/Paper2D/Private/PaperSprite.cpp
+++ b/Engine/Plugins/Experimental/Paper2D/Source/Paper2D/Private/PaperSprite.cpp
@@ -351,12 +351,9 @@ void FSpriteDrawCallRecord::BuildFromSprite(class UPaperSprite* Sprite)
 {
 	if (Sprite != NULL)
 	{
-		//SourceDimension = Sprite->GetSourceSize();
 		Destination = FVector::ZeroVector;
-		//DestinationDimension = Sprite->GetSourceSize();
 		Texture = Sprite->GetBakedTexture();
 		Color = FLinearColor::White;
-		BlendMode = Sprite->BlendMode;
 
 		RenderVerts = Sprite->BakedRenderData;
 	}
diff --git a/Engine/Plugins/Experimental/Paper2D/Source/Paper2D/Private/PaperTileMapRenderSceneProxy.cpp b/Engine/Plugins/Experimental/Paper2D/Source/Paper2D/Private/PaperTileMapRenderSceneProxy.cpp
--- a/Engine/Plugins/Experimental/Paper2D/Source/Paper2D/Private/PaperTileMapRenderSceneProxy.cpp
+++ b/Engine/Plugins/Experimental/Paper2D/Source/Paper2D/Private/PaperTileMapRenderSceneProxy.cpp
@@ -214,13 +214,9 @@ void FPaperTileMapRenderSceneProxy::DrawDynamicElements(FPrimitiveDrawInterface*
 					const FVector DrawPos(TileSeparationX * X, LayerIndex * LayerSeparation, -TileSeparationY * Y);
 
 					FSpriteDrawCallRecord& NewTile = *(new (BatchedSprites) FSpriteDrawCallRecord());
-					//NewTile.SourceUV = SourceUV;
-					//NewTile.SourceDimension = SourceDimensionsUV;
 					NewTile.Destination = DrawPos;
-					//NewTile.DestinationDimension = TileSizeXY;
 					NewTile.Texture = SourceTexture;
 					NewTile.Color = DrawColor;
-					NewTile.BlendMode = BlendMode;
 					
 					//@TODO: PAPER: TILE EDITOR: Build data
 				}
