Paper2D: Batch manager experiment WIP

diff --git a/Engine/Plugins/Experimental/Paper2D/Source/Paper2D/Private/Paper2DModule.cpp b/Engine/Plugins/Experimental/Paper2D/Source/Paper2D/Private/Paper2DModule.cpp
--- a/Engine/Plugins/Experimental/Paper2D/Source/Paper2D/Private/Paper2DModule.cpp
+++ b/Engine/Plugins/Experimental/Paper2D/Source/Paper2D/Private/Paper2DModule.cpp
@@ -3,9 +3,8 @@
 #include "Paper2DPrivatePCH.h"
 #include "Paper2DModule.h"
 
-
-
 #include "Physics/Box2DIntegration.h"
+#include "Rendering/PaperBatchManager.h"
 
 //////////////////////////////////////////////////////////////////////////
 // FPaper2DModule
@@ -29,10 +28,13 @@ public:
 		}
 
 		PaperAxisZ = FVector::CrossProduct(PaperAxisX, PaperAxisY);
+
+		FPaperBatchManager::Initialize();
 	}
 
 	virtual void ShutdownModule() OVERRIDE
 	{
+		FPaperBatchManager::Shutdown();
 		FPhysicsIntegration2D::ShutdownPhysics();
 	}
 };
diff --git a/Engine/Plugins/Experimental/Paper2D/Source/Paper2D/Private/PaperRenderSceneProxy.cpp b/Engine/Plugins/Experimental/Paper2D/Source/Paper2D/Private/PaperRenderSceneProxy.cpp
--- a/Engine/Plugins/Experimental/Paper2D/Source/Paper2D/Private/PaperRenderSceneProxy.cpp
+++ b/Engine/Plugins/Experimental/Paper2D/Source/Paper2D/Private/PaperRenderSceneProxy.cpp
@@ -3,6 +3,8 @@
 #include "Paper2DPrivatePCH.h"
 #include "PaperRenderSceneProxy.h"
 #include "PhysicsEngine/BodySetup.h"
+#include "Rendering/PaperBatchManager.h"
+#include "Rendering/PaperBatchSceneProxy.h"
 
 static FPackedNormal PackedNormalX(FVector(1.0f, 0.0f, 0.0f));
 static FPackedNormal PackedNormalY(FVector(0.0f, 1.0f, 0.0f));
@@ -173,6 +175,18 @@ FPaperRenderSceneProxy::FPaperRenderSceneProxy(const UPrimitiveComponent* InComp
 	}
 }
 
+FPaperRenderSceneProxy::~FPaperRenderSceneProxy()
+{
+	FPaperBatchSceneProxy* Batcher = FPaperBatchManager::GetBatcher(GetScene());
+	Batcher->UnregisterManagedProxy(this);
+}
+
+void FPaperRenderSceneProxy::CreateRenderThreadResources()
+{
+	FPaperBatchSceneProxy* Batcher = FPaperBatchManager::GetBatcher(GetScene());
+	Batcher->RegisterManagedProxy(this);
+}
+
 void FPaperRenderSceneProxy::DrawDynamicElements(FPrimitiveDrawInterface* PDI, const FSceneView* View)
 {
 	QUICK_SCOPE_CYCLE_COUNTER(STAT_PaperRenderSceneProxy_DrawDynamicElements);
@@ -241,9 +255,6 @@ void FPaperRenderSceneProxy::DrawDynamicElements(FPrimitiveDrawInterface* PDI, c
 #endif
 
 	// Set the selection/hover color from the current engine setting.
-	// The color is multiplied by 10 because this value is normally expected to be blended
-	// additively, this is not how the sprites work and therefore need an extra boost
-	// to appear the same color as previously.
 	FLinearColor OverrideColor;
 	bool bUseOverrideColor = false;
 
@@ -254,21 +265,25 @@ void FPaperRenderSceneProxy::DrawDynamicElements(FPrimitiveDrawInterface* PDI, c
 		OverrideColor = GetSelectionColor(FLinearColor::White, bShowAsSelected, IsHovered());
 	}
 
-	//@TODO: Figure out how to show selection highlight better (maybe outline?)
 	bUseOverrideColor = false;
 
 	// Sprites of locked actors draw in red.
 	//FLinearColor LevelColorToUse = IsSelected() ? ColorToUse : (FLinearColor)LevelColor;
 	//FLinearColor PropertyColorToUse = PropertyColor;
 
-	if (Material != NULL)
-	{
-		DrawDynamicElements_RichMesh(PDI, View, bUseOverrideColor, OverrideColor);
-	}
+#if TEST_BATCHING
+#else
+	DrawDynamicElements_RichMesh(PDI, View, bUseOverrideColor, OverrideColor);
+#endif
 }
 
 void FPaperRenderSceneProxy::DrawDynamicElements_RichMesh(FPrimitiveDrawInterface* PDI, const FSceneView* View, bool bUseOverrideColor, const FLinearColor& OverrideColor)
 {
+	if (Material == nullptr)
+	{
+		return;
+	}
+
 	const uint8 DPG = GetDepthPriorityGroup(View);
 	
 	static TGlobalResource<FPaperSpriteVertexFactory> GPaperSpriteVertexFactory;
diff --git a/Engine/Plugins/Experimental/Paper2D/Source/Paper2D/Private/PaperRenderSceneProxy.h b/Engine/Plugins/Experimental/Paper2D/Source/Paper2D/Private/PaperRenderSceneProxy.h
--- a/Engine/Plugins/Experimental/Paper2D/Source/Paper2D/Private/PaperRenderSceneProxy.h
+++ b/Engine/Plugins/Experimental/Paper2D/Source/Paper2D/Private/PaperRenderSceneProxy.h
@@ -2,6 +2,8 @@
 
 #pragma once
 
+#define TEST_BATCHING 0
+
 //////////////////////////////////////////////////////////////////////////
 // FPaperRenderSceneProxy
 
@@ -9,6 +11,7 @@ class FPaperRenderSceneProxy : public FPrimitiveSceneProxy
 {
 public:
 	FPaperRenderSceneProxy(const UPrimitiveComponent* InComponent);
+	~FPaperRenderSceneProxy();
 
 	// FPrimitiveSceneProxy interface.
 	virtual void DrawDynamicElements(FPrimitiveDrawInterface* PDI, const FSceneView* View) OVERRIDE;
@@ -16,16 +19,24 @@ public:
 	virtual void OnTransformChanged() OVERRIDE;
 	virtual uint32 GetMemoryFootprint() const OVERRIDE;
 	virtual bool CanBeOccluded() const OVERRIDE;
+	virtual void CreateRenderThreadResources() OVERRIDE;
 	// End of FPrimitiveSceneProxy interface.
 
 	void SetDrawCall_RenderThread(const FSpriteDrawCallRecord& NewDynamicData);
 
+	class UMaterialInterface* GetMaterial() const
+	{
+		return Material;
+	}
+
 protected:
-	void DrawDynamicElements_Batched(FPrimitiveDrawInterface* PDI, const FSceneView* View, bool bUseOverrideColor, const FLinearColor& OverrideColor);
 	void DrawDynamicElements_RichMesh(FPrimitiveDrawInterface* PDI, const FSceneView* View, bool bUseOverrideColor, const FLinearColor& OverrideColor);
 
 	bool IsCollisionView(const FSceneView* View, bool& bDrawSimpleCollision, bool& bDrawComplexCollision) const;
 
+	friend class FPaperBatchSceneProxy;
+
+
 protected:
 	TArray<FSpriteDrawCallRecord> BatchedSprites;
 	class UMaterialInterface* Material;
diff --git a/Engine/Plugins/Experimental/Paper2D/Source/Paper2D/Private/Rendering/PaperBatchComponent.cpp b/Engine/Plugins/Experimental/Paper2D/Source/Paper2D/Private/Rendering/PaperBatchComponent.cpp
new file mode 100644
--- /dev/null
+++ b/Engine/Plugins/Experimental/Paper2D/Source/Paper2D/Private/Rendering/PaperBatchComponent.cpp
@@ -0,0 +1,27 @@
+// Copyright 1998-2014 Epic Games, Inc. All Rights Reserved.
+
+#include "Paper2DPrivatePCH.h"
+#include "PaperBatchSceneProxy.h"
+#include "PaperBatchComponent.h"
+
+//////////////////////////////////////////////////////////////////////////
+// UPaperBatchComponent
+
+UPaperBatchComponent::UPaperBatchComponent(const FPostConstructInitializeProperties& PCIP)
+	: Super(PCIP)
+{
+	BodyInstance.SetCollisionEnabled(ECollisionEnabled::NoCollision);
+}
+
+FPrimitiveSceneProxy* UPaperBatchComponent::CreateSceneProxy()
+{
+	FPaperBatchSceneProxy* NewProxy = new FPaperBatchSceneProxy(this);
+	return NewProxy;
+}
+
+FBoxSphereBounds UPaperBatchComponent::CalcBounds(const FTransform & LocalToWorld) const
+{
+	// Always visible
+	FBoxSphereBounds Bounds(FVector::ZeroVector, FVector(HALF_WORLD_MAX, HALF_WORLD_MAX, HALF_WORLD_MAX), HALF_WORLD_MAX);
+	return Bounds;
+}
diff --git a/Engine/Plugins/Experimental/Paper2D/Source/Paper2D/Private/Rendering/PaperBatchComponent.h b/Engine/Plugins/Experimental/Paper2D/Source/Paper2D/Private/Rendering/PaperBatchComponent.h
new file mode 100644
--- /dev/null
+++ b/Engine/Plugins/Experimental/Paper2D/Source/Paper2D/Private/Rendering/PaperBatchComponent.h
@@ -0,0 +1,17 @@
+// Copyright 1998-2014 Epic Games, Inc. All Rights Reserved.
+
+#pragma once
+
+#include "PaperBatchComponent.generated.h"
+
+// Dummy component designed to create a batch manager render scene proxy that does work on the render thread
+UCLASS()
+class UPaperBatchComponent : public UPrimitiveComponent
+{
+	GENERATED_UCLASS_BODY()
+
+	// UPrimitiveComponent interface
+	virtual FPrimitiveSceneProxy* CreateSceneProxy() OVERRIDE;
+	virtual FBoxSphereBounds CalcBounds(const FTransform & LocalToWorld) const OVERRIDE;
+	// End of UPrimitiveComponent interface
+};
diff --git a/Engine/Plugins/Experimental/Paper2D/Source/Paper2D/Private/Rendering/PaperBatchManager.cpp b/Engine/Plugins/Experimental/Paper2D/Source/Paper2D/Private/Rendering/PaperBatchManager.cpp
new file mode 100644
--- /dev/null
+++ b/Engine/Plugins/Experimental/Paper2D/Source/Paper2D/Private/Rendering/PaperBatchManager.cpp
@@ -0,0 +1,62 @@
+// Copyright 1998-2014 Epic Games, Inc. All Rights Reserved.
+
+#include "Paper2DPrivatePCH.h"
+#include "PaperBatchManager.h"
+#include "PaperBatchComponent.h"
+
+namespace FPaperBatchManagerNS
+{
+	static TMap< UWorld*, UPaperBatchComponent* > WorldToBatcherMap;
+	static FWorldDelegates::FWorldInitializationEvent::FDelegate OnWorldCreatedDelegate;
+	static FWorldDelegates::FWorldEvent::FDelegate OnWorldDestroyedDelegate;
+}
+
+//////////////////////////////////////////////////////////////////////////
+// FPaperBatchManager
+
+void FPaperBatchManager::Initialize()
+{
+	FPaperBatchManagerNS::OnWorldCreatedDelegate = FWorldDelegates::FWorldInitializationEvent::FDelegate::CreateStatic(&FPaperBatchManager::OnWorldCreated);
+	FPaperBatchManagerNS::OnWorldDestroyedDelegate = FWorldDelegates::FWorldEvent::FDelegate::CreateStatic(&FPaperBatchManager::OnWorldDestroyed);
+	FWorldDelegates::OnPostWorldInitialization.Add(FPaperBatchManagerNS::OnWorldCreatedDelegate);
+	FWorldDelegates::OnPreWorldFinishDestroy.Add(FPaperBatchManagerNS::OnWorldDestroyedDelegate);
+}
+
+void FPaperBatchManager::Shutdown()
+{
+	FWorldDelegates::OnPostWorldInitialization.Remove(FPaperBatchManagerNS::OnWorldCreatedDelegate);
+	FWorldDelegates::OnPreWorldFinishDestroy.Remove(FPaperBatchManagerNS::OnWorldDestroyedDelegate);
+}
+
+void FPaperBatchManager::OnWorldCreated(UWorld* World, const UWorld::InitializationValues IVS)
+{
+	if (IVS.bInitializeScenes)
+	{
+		UPaperBatchComponent* Batcher = NewObject<UPaperBatchComponent>();
+		Batcher->AddToRoot();
+		FPaperBatchManagerNS::WorldToBatcherMap.Add(World, Batcher);
+
+		World->Scene->AddPrimitive(Batcher);
+	}
+}
+
+void FPaperBatchManager::OnWorldDestroyed(UWorld* World)
+{
+	UPaperBatchComponent* Batcher = FPaperBatchManagerNS::WorldToBatcherMap.FindChecked(World);
+	FPaperBatchManagerNS::WorldToBatcherMap.Remove(World);
+
+	World->Scene->RemovePrimitive(Batcher);
+	Batcher->RemoveFromRoot();
+}
+
+UPaperBatchComponent* FPaperBatchManager::GetBatchComponent(UWorld* World)
+{
+	UPaperBatchComponent* Batcher = FPaperBatchManagerNS::WorldToBatcherMap.FindChecked(World);
+	return Batcher;
+}
+
+FPaperBatchSceneProxy* FPaperBatchManager::GetBatcher(FSceneInterface* Scene)
+{
+	UPaperBatchComponent* BatchComponent = GetBatchComponent(Scene->GetWorld());
+	return static_cast<FPaperBatchSceneProxy*>(BatchComponent->SceneProxy);
+}
diff --git a/Engine/Plugins/Experimental/Paper2D/Source/Paper2D/Private/Rendering/PaperBatchManager.h b/Engine/Plugins/Experimental/Paper2D/Source/Paper2D/Private/Rendering/PaperBatchManager.h
new file mode 100644
--- /dev/null
+++ b/Engine/Plugins/Experimental/Paper2D/Source/Paper2D/Private/Rendering/PaperBatchManager.h
@@ -0,0 +1,19 @@
+// Copyright 1998-2014 Epic Games, Inc. All Rights Reserved.
+
+#pragma once
+
+struct FPaperBatchManager
+{
+public:
+	static void Initialize();
+	static void Shutdown();
+
+	static class UPaperBatchComponent* GetBatchComponent(UWorld* World);
+
+	// Gets the batcher for the specified scene
+	// Note: Can only be called on the render thread
+	static class FPaperBatchSceneProxy* GetBatcher(FSceneInterface* Scene);
+private:
+	static void OnWorldCreated(UWorld* World, const UWorld::InitializationValues IVS);
+	static void OnWorldDestroyed(UWorld* World);
+};
diff --git a/Engine/Plugins/Experimental/Paper2D/Source/Paper2D/Private/Rendering/PaperBatchSceneProxy.cpp b/Engine/Plugins/Experimental/Paper2D/Source/Paper2D/Private/Rendering/PaperBatchSceneProxy.cpp
new file mode 100644
--- /dev/null
+++ b/Engine/Plugins/Experimental/Paper2D/Source/Paper2D/Private/Rendering/PaperBatchSceneProxy.cpp
@@ -0,0 +1,150 @@
+// Copyright 1998-2014 Epic Games, Inc. All Rights Reserved.
+
+#include "Paper2DPrivatePCH.h"
+#include "PaperBatchSceneProxy.h"
+#include "PaperRenderSceneProxy.h"
+
+//////////////////////////////////////////////////////////////////////////
+// FPaperBatchSceneProxy
+
+FPaperBatchSceneProxy::FPaperBatchSceneProxy(const UPrimitiveComponent* InComponent)
+	: FPrimitiveSceneProxy(InComponent)
+{
+}
+
+uint32 FPaperBatchSceneProxy::GetMemoryFootprint() const
+{
+	return sizeof(*this) + GetAllocatedSize();
+}
+
+FPrimitiveViewRelevance FPaperBatchSceneProxy::GetViewRelevance(const FSceneView* View)
+{
+	checkSlow(IsInRenderingThread());
+
+	FPrimitiveViewRelevance Result;
+	Result.bDrawRelevance = IsShown(View) && View->Family->EngineShowFlags.Paper2DSprites;
+	Result.bRenderCustomDepth = ShouldRenderCustomDepth();
+	Result.bRenderInMainPass = ShouldRenderInMainPass();
+
+	Result.bMaskedRelevance = true;
+	//Result.bNormalTranslucencyRelevance = true;
+	Result.bDynamicRelevance = true;
+	Result.bOpaqueRelevance = true;
+
+	return Result;
+}
+
+
+struct FProxyBatchPair
+{
+	void* Texture;
+	int32 ProxyIndex;
+	int32 BatchIndex;
+};
+
+struct FSortSpriteProxiesByMaterial
+{
+	FORCEINLINE bool operator()(const FPaperRenderSceneProxy& A, const FPaperRenderSceneProxy& B) const
+	{
+		return A.GetMaterial() < B.GetMaterial();
+	}
+};
+
+
+struct FSortSpriteBatchesByTexture
+{
+	FORCEINLINE bool operator()(const FProxyBatchPair& A, const FProxyBatchPair& B) const
+	{
+		return (A.Texture == B.Texture) ? ((A.ProxyIndex == B.ProxyIndex) ? (A.BatchIndex < B.BatchIndex) : (A.ProxyIndex < B.ProxyIndex)) : (A.Texture < B.Texture);
+	}
+};
+
+void FPaperBatchSceneProxy::DrawDynamicElements(FPrimitiveDrawInterface* PDI, const FSceneView* View)
+{
+#if TEST_BATCHING
+	// First sort by material
+	ManagedProxies.Sort(FSortSpriteProxiesByMaterial());
+
+	// Figure out how many total sprites and verts we need to deal with
+	int32 NumSprites = 0;
+	int32 NumVerts = 0;
+	for (FPaperRenderSceneProxy* Proxy : ManagedProxies)
+	{
+		NumSprites += Proxy->BatchedSprites.Num();
+		for (FSpriteDrawCallRecord& Sprite : Proxy->BatchedSprites)
+		{
+			NumVerts += Sprite.RenderVerts.Num();
+		}
+	}
+
+
+	// Aggregate the individual draw calls
+	TArray<FProxyBatchPair> SpriteDrawCalls;
+	SpriteDrawCalls.Empty(NumSprites);
+
+	for (int32 ProxyIndex = 0; ProxyIndex < ManagedProxies.Num(); ++ProxyIndex)
+	{
+		FPaperRenderSceneProxy* Proxy = ManagedProxies[ProxyIndex];
+		for (int32 BatchIndex = 0; BatchIndex < Proxy->BatchedSprites.Num(); ++BatchIndex)
+		{
+			FProxyBatchPair& Pair = *new (SpriteDrawCalls) FProxyBatchPair();
+			Pair.Texture = Proxy->BatchedSprites[BatchIndex].Texture;//@TODO: Don't need to stick this here, but all of this will change when the batcher handles per-item batches too
+			Pair.ProxyIndex = ProxyIndex;
+			Pair.BatchIndex = BatchIndex;
+		}
+	}
+
+	// Sort by texture
+	SpriteDrawCalls.Sort(FSortSpriteBatchesByTexture());
+
+	// Ready to render in a nice batched fashion
+	for (FProxyBatchPair& Pair : SpriteDrawCalls)
+	{
+		ManagedProxies[Pair.ProxyIndex]->DrawDynamicElements_RichMesh(PDI, View, false, FLinearColor::White);
+	}
+
+
+// 	for (int32 PrimitiveIndex = 0, Num = View.VisibleDynamicPrimitives.Num(); PrimitiveIndex < Num; PrimitiveIndex++)
+// 	{
+// 		const FPrimitiveSceneInfo* PrimitiveSceneInfo = View.VisibleDynamicPrimitives[PrimitiveIndex];
+// 		int32 PrimitiveId = PrimitiveSceneInfo->GetIndex();
+// 		const FPrimitiveViewRelevance& PrimitiveViewRelevance = View.PrimitiveViewRelevanceMap[PrimitiveId];
+// 
+// 		const bool bVisible = View.PrimitiveVisibilityMap[PrimitiveId];
+// 
+// 		// Only draw the primitive if it's visible
+// 		if (bVisible &&
+// 			// only draw opaque and masked primitives if wireframe is disabled
+// 			(PrimitiveViewRelevance.bOpaqueRelevance || ViewFamily.EngineShowFlags.Wireframe) &&
+// 			PrimitiveViewRelevance.bRenderInMainPass
+// 			)
+// 
+// 		{
+// 			FScopeCycleCounter Context(PrimitiveSceneInfo->Proxy->GetStatId());
+// 			Drawer.SetPrimitive(PrimitiveSceneInfo->Proxy);
+// 			PrimitiveSceneInfo->Proxy->DrawDynamicElements(
+// 				&Drawer,
+// 				&View
+// 				);
+// 		}
+// 	}
+#endif
+	//for 
+}
+
+void FPaperBatchSceneProxy::RegisterManagedProxy(FPaperRenderSceneProxy* Proxy)
+{
+	check(IsInRenderingThread());
+	ManagedProxies.Add(Proxy);
+}
+
+void FPaperBatchSceneProxy::UnregisterManagedProxy(FPaperRenderSceneProxy* Proxy)
+{
+	check(IsInRenderingThread());
+	ManagedProxies.RemoveSwap(Proxy);
+}
+
+bool FPaperBatchSceneProxy::CanBeOccluded() const
+{
+	return false;
+}
diff --git a/Engine/Plugins/Experimental/Paper2D/Source/Paper2D/Private/Rendering/PaperBatchSceneProxy.h b/Engine/Plugins/Experimental/Paper2D/Source/Paper2D/Private/Rendering/PaperBatchSceneProxy.h
new file mode 100644
--- /dev/null
+++ b/Engine/Plugins/Experimental/Paper2D/Source/Paper2D/Private/Rendering/PaperBatchSceneProxy.h
@@ -0,0 +1,24 @@
+// Copyright 1998-2014 Epic Games, Inc. All Rights Reserved.
+
+#pragma once
+
+//////////////////////////////////////////////////////////////////////////
+// FPaperBatchSceneProxy
+
+class FPaperBatchSceneProxy : public FPrimitiveSceneProxy
+{
+public:
+	FPaperBatchSceneProxy(const UPrimitiveComponent* InComponent);
+
+ 	// FPrimitiveSceneProxy interface.
+ 	virtual void DrawDynamicElements(FPrimitiveDrawInterface* PDI, const FSceneView* View) OVERRIDE;
+ 	virtual FPrimitiveViewRelevance GetViewRelevance(const FSceneView* View) OVERRIDE;
+ 	virtual uint32 GetMemoryFootprint() const OVERRIDE;
+ 	virtual bool CanBeOccluded() const OVERRIDE;
+ 	// End of FPrimitiveSceneProxy interface.
+
+	void RegisterManagedProxy(class FPaperRenderSceneProxy* Proxy);
+	void UnregisterManagedProxy(class FPaperRenderSceneProxy* Proxy);
+protected:
+	TArray<class FPaperRenderSceneProxy*> ManagedProxies;
+};
