Fixed rendering thread crash with Paper2D sprites
- FMeshBatch pointed to stack-allocated memory address for the material render proxy

diff --git a/Engine/Plugins/Experimental/Paper2D/Source/Paper2D/Private/PaperRenderSceneProxy.cpp b/Engine/Plugins/Experimental/Paper2D/Source/Paper2D/Private/PaperRenderSceneProxy.cpp
--- a/Engine/Plugins/Experimental/Paper2D/Source/Paper2D/Private/PaperRenderSceneProxy.cpp
+++ b/Engine/Plugins/Experimental/Paper2D/Source/Paper2D/Private/PaperRenderSceneProxy.cpp
@@ -98,7 +98,7 @@ public:
 /**
  * A material render proxy which overrides a named texture parameter.
  */
-class FTextureOverrideRenderProxy : public FMaterialRenderProxy
+class FTextureOverrideRenderProxy : public FDynamicPrimitiveResource, public FMaterialRenderProxy
 {
 public:
 	const FMaterialRenderProxy* const Parent;
@@ -112,10 +112,23 @@ public:
 		, TextureParameterName(InParameterName)
 	{}
 
+	virtual ~FTextureOverrideRenderProxy()
+	{
+	}
+
+	// FDynamicPrimitiveResource interface.
+	virtual void InitPrimitiveResource()
+	{
+	}
+	virtual void ReleasePrimitiveResource()
+	{
+		delete this;
+	}
+
 	// FMaterialRenderProxy interface.
-	virtual const FMaterial* GetMaterial(ERHIFeatureLevel::Type FeatureLevel) const OVERRIDE
+	virtual const FMaterial* GetMaterial(ERHIFeatureLevel::Type InFeatureLevel) const OVERRIDE
 	{
-		return Parent->GetMaterial(FeatureLevel);
+		return Parent->GetMaterial(InFeatureLevel);
 	}
 
 	virtual bool GetVectorValue(const FName ParameterName, FLinearColor* OutValue, const FMaterialRenderContext& Context) const OVERRIDE
@@ -328,9 +341,12 @@ void FPaperRenderSceneProxy::DrawDynamicElements_RichMesh(FPrimitiveDrawInterfac
 
 
 			FMaterialRenderProxy* ParentMaterialProxy = Material->GetRenderProxy((View->Family->EngineShowFlags.Selection) && IsSelected(), IsHovered());
-			FTextureOverrideRenderProxy TextureOverrideMaterialProxy(ParentMaterialProxy, Record.Texture, TEXT("SpriteTexture"));
 
-			Mesh.MaterialRenderProxy = &TextureOverrideMaterialProxy;
+			// Create a texture override material proxy and register it as a dynamic resource so that it won't be deleted until the rendering thread has finished with it
+			FTextureOverrideRenderProxy* TextureOverrideMaterialProxy = new FTextureOverrideRenderProxy(ParentMaterialProxy, Record.Texture, TEXT("SpriteTexture"));
+			PDI->RegisterDynamicResource( TextureOverrideMaterialProxy );
+
+			Mesh.MaterialRenderProxy = TextureOverrideMaterialProxy;
 			Mesh.LCI = NULL;
 			Mesh.ReverseCulling = IsLocalToWorldDeterminantNegative() ? true : false;
 			Mesh.CastShadow = false;
