Paper2D: Batch manager fixes
- Fix crash on editor shutdown
- Fix culling issue when TEST_BATCHING=1

diff --git a/Engine/Plugins/Experimental/Paper2D/Source/Paper2D/Private/Rendering/PaperBatchManager.cpp b/Engine/Plugins/Experimental/Paper2D/Source/Paper2D/Private/Rendering/PaperBatchManager.cpp
--- a/Engine/Plugins/Experimental/Paper2D/Source/Paper2D/Private/Rendering/PaperBatchManager.cpp
+++ b/Engine/Plugins/Experimental/Paper2D/Source/Paper2D/Private/Rendering/PaperBatchManager.cpp
@@ -7,6 +7,7 @@
 namespace FPaperBatchManagerNS
 {
 	static TMap< UWorld*, UPaperBatchComponent* > WorldToBatcherMap;
+
 	static FWorldDelegates::FWorldInitializationEvent::FDelegate OnWorldCreatedDelegate;
 	static FWorldDelegates::FWorldEvent::FDelegate OnWorldDestroyedDelegate;
 }
@@ -17,8 +18,9 @@ namespace FPaperBatchManagerNS
 void FPaperBatchManager::Initialize()
 {
 	FPaperBatchManagerNS::OnWorldCreatedDelegate = FWorldDelegates::FWorldInitializationEvent::FDelegate::CreateStatic(&FPaperBatchManager::OnWorldCreated);
-	FPaperBatchManagerNS::OnWorldDestroyedDelegate = FWorldDelegates::FWorldEvent::FDelegate::CreateStatic(&FPaperBatchManager::OnWorldDestroyed);
 	FWorldDelegates::OnPostWorldInitialization.Add(FPaperBatchManagerNS::OnWorldCreatedDelegate);
+
+	FPaperBatchManagerNS::OnWorldDestroyedDelegate = FWorldDelegates::FWorldEvent::FDelegate::CreateStatic(&FPaperBatchManager::OnWorldDestroyed);
 	FWorldDelegates::OnPreWorldFinishDestroy.Add(FPaperBatchManagerNS::OnWorldDestroyedDelegate);
 }
 
@@ -30,33 +32,41 @@ void FPaperBatchManager::Shutdown()
 
 void FPaperBatchManager::OnWorldCreated(UWorld* World, const UWorld::InitializationValues IVS)
 {
+	UPaperBatchComponent* Batcher = NewObject<UPaperBatchComponent>();
+	Batcher->UpdateBounds();
+	Batcher->AddToRoot();
+	FPaperBatchManagerNS::WorldToBatcherMap.Add(World, Batcher);
+
 	if (IVS.bInitializeScenes)
 	{
-		UPaperBatchComponent* Batcher = NewObject<UPaperBatchComponent>();
-		Batcher->AddToRoot();
-		FPaperBatchManagerNS::WorldToBatcherMap.Add(World, Batcher);
-
-		World->Scene->AddPrimitive(Batcher);
+		Batcher->RegisterComponentWithWorld(World);
 	}
 }
 
 void FPaperBatchManager::OnWorldDestroyed(UWorld* World)
 {
 	UPaperBatchComponent* Batcher = FPaperBatchManagerNS::WorldToBatcherMap.FindChecked(World);
-	FPaperBatchManagerNS::WorldToBatcherMap.Remove(World);
 
-	World->Scene->RemovePrimitive(Batcher);
+	if (Batcher->IsRegistered())
+	{
+		Batcher->UnregisterComponent();
+	}
+
+	FPaperBatchManagerNS::WorldToBatcherMap.Remove(World);
 	Batcher->RemoveFromRoot();
 }
 
 UPaperBatchComponent* FPaperBatchManager::GetBatchComponent(UWorld* World)
 {
 	UPaperBatchComponent* Batcher = FPaperBatchManagerNS::WorldToBatcherMap.FindChecked(World);
+	check(Batcher);
 	return Batcher;
 }
 
 FPaperBatchSceneProxy* FPaperBatchManager::GetBatcher(FSceneInterface* Scene)
 {
 	UPaperBatchComponent* BatchComponent = GetBatchComponent(Scene->GetWorld());
-	return static_cast<FPaperBatchSceneProxy*>(BatchComponent->SceneProxy);
+	FPaperBatchSceneProxy* BatchProxy = static_cast<FPaperBatchSceneProxy*>(BatchComponent->SceneProxy);
+	check(BatchProxy);
+	return BatchProxy;
 }
