Paper2D: Fix crash when cooking with Paper2D plugin enabled

diff --git a/Engine/Plugins/Experimental/Paper2D/Source/SmartSnapping/Private/SmartSnappingModule.cpp b/Engine/Plugins/Experimental/Paper2D/Source/SmartSnapping/Private/SmartSnappingModule.cpp
--- a/Engine/Plugins/Experimental/Paper2D/Source/SmartSnapping/Private/SmartSnappingModule.cpp
+++ b/Engine/Plugins/Experimental/Paper2D/Source/SmartSnapping/Private/SmartSnappingModule.cpp
@@ -17,25 +17,28 @@ public:
 
 	virtual void StartupModule() OVERRIDE
 	{
-		// Create and register the snapping policy
-		PlanarPolicy = MakeShareable(new FPlanarConstraintSnapPolicy);
-		
-		IViewportSnappingModule& SnappingModule = FModuleManager::LoadModuleChecked<IViewportSnappingModule>("ViewportSnapping");
-		SnappingModule.RegisterSnappingPolicy(PlanarPolicy);
-
-		// Register the extension with the level editor
+		if (!IsRunningCommandlet())
 		{
-			ViewMenuExtender = FLevelEditorModule::FLevelEditorMenuExtender::CreateRaw(this, &FSmartSnappingModule::OnExtendLevelEditorViewMenu);
+			// Create and register the snapping policy
+			PlanarPolicy = MakeShareable(new FPlanarConstraintSnapPolicy);
+
+			IViewportSnappingModule& SnappingModule = FModuleManager::LoadModuleChecked<IViewportSnappingModule>("ViewportSnapping");
+			SnappingModule.RegisterSnappingPolicy(PlanarPolicy);
+
+			// Register the extension with the level editor
+			{
+				ViewMenuExtender = FLevelEditorModule::FLevelEditorMenuExtender::CreateRaw(this, &FSmartSnappingModule::OnExtendLevelEditorViewMenu);
 
 
-			FLevelEditorModule& LevelEditor = FModuleManager::GetModuleChecked<FLevelEditorModule>( TEXT("LevelEditor") );
-			LevelEditor.GetAllLevelEditorToolbarViewMenuExtenders().Add(ViewMenuExtender);
+				FLevelEditorModule& LevelEditor = FModuleManager::GetModuleChecked<FLevelEditorModule>(TEXT("LevelEditor"));
+				LevelEditor.GetAllLevelEditorToolbarViewMenuExtenders().Add(ViewMenuExtender);
+			}
 		}
 	}
 
 	virtual void ShutdownModule() OVERRIDE
 	{
-		if (UObjectInitialized())
+		if (UObjectInitialized() && !IsRunningCommandlet())
 		{
 			// Unregister the level editor extensions
 			{
