Paper2D: Fix compile errors in non-editor builds

diff --git a/Engine/Plugins/Experimental/Paper2D/Source/Paper2D/Classes/PaperSpriteAtlas.h b/Engine/Plugins/Experimental/Paper2D/Source/Paper2D/Classes/PaperSpriteAtlas.h
--- a/Engine/Plugins/Experimental/Paper2D/Source/Paper2D/Classes/PaperSpriteAtlas.h
+++ b/Engine/Plugins/Experimental/Paper2D/Source/Paper2D/Classes/PaperSpriteAtlas.h
@@ -36,8 +36,10 @@ class UPaperSpriteAtlas : public UObject
 
 public:
 	// UObject interface
+#if WITH_EDITORONLY_DATA
 	virtual void PostInitProperties() OVERRIDE;
 	virtual void PostDuplicate(bool bDuplicateForPIE) OVERRIDE;
+#endif
 	virtual void GetAssetRegistryTags(TArray<FAssetRegistryTag>& OutTags) const OVERRIDE;
 	// End of UObject interface
 };
diff --git a/Engine/Plugins/Experimental/Paper2D/Source/Paper2D/Private/PaperSprite.cpp b/Engine/Plugins/Experimental/Paper2D/Source/Paper2D/Private/PaperSprite.cpp
--- a/Engine/Plugins/Experimental/Paper2D/Source/Paper2D/Private/PaperSprite.cpp
+++ b/Engine/Plugins/Experimental/Paper2D/Source/Paper2D/Private/PaperSprite.cpp
@@ -1112,6 +1112,7 @@ void UPaperSprite::QuerySupportedSockets(TArray<FComponentSocketDescription>& Ou
 	}
 }
 
+#if WITH_EDITORONLY_DATA
 void UPaperSprite::Serialize(FArchive& Ar)
 {
 	Super::Serialize(Ar);
@@ -1130,6 +1131,7 @@ void UPaperSprite::PostLoad()
 	}
 #endif
 }
+#endif
 
 UTexture2D* UPaperSprite::GetBakedTexture() const
 {
diff --git a/Engine/Plugins/Experimental/Paper2D/Source/Paper2D/Private/PaperSpriteAtlas.cpp b/Engine/Plugins/Experimental/Paper2D/Source/Paper2D/Private/PaperSpriteAtlas.cpp
--- a/Engine/Plugins/Experimental/Paper2D/Source/Paper2D/Private/PaperSpriteAtlas.cpp
+++ b/Engine/Plugins/Experimental/Paper2D/Source/Paper2D/Private/PaperSpriteAtlas.cpp
@@ -7,8 +7,10 @@
 
 UPaperSpriteAtlas::UPaperSpriteAtlas(const FPostConstructInitializeProperties& PCIP)
 	: Super(PCIP)
+#if WITH_EDITORONLY_DATA
 	, MaxWidth(2048)
 	, MaxHeight(2048)
+#endif
 {
 }
 
@@ -21,6 +23,7 @@ void UPaperSpriteAtlas::GetAssetRegistryTags(TArray<FAssetRegistryTag>& OutTags)
 #endif
 }
 
+#if WITH_EDITORONLY_DATA
 void UPaperSpriteAtlas::PostDuplicate(bool bDuplicateForPIE)
 {
 	Super::PostDuplicate(bDuplicateForPIE);
@@ -32,3 +35,4 @@ void UPaperSpriteAtlas::PostInitProperties()
 	Super::PostInitProperties();
 	FPlatformMisc::CreateGuid(AtlasGUID);
 }
+#endif
\ No newline at end of file
