Paper2D: Correcting module type to be Runtime again

diff --git a/Engine/Plugins/Experimental/Paper2D/Paper2D.uplugin b/Engine/Plugins/Experimental/Paper2D/Paper2D.uplugin
--- a/Engine/Plugins/Experimental/Paper2D/Paper2D.uplugin
+++ b/Engine/Plugins/Experimental/Paper2D/Paper2D.uplugin
@@ -14,7 +14,7 @@
 	[
 		{
 			"Name" : "Paper2D",
-			"Type" : "Editor",
+			"Type" : "Runtime",
 			"LoadingPhase" : "PreDefault"
 		},
 		{
