Paper2D: Remove unnecessary ticking from UPaperRenderComponent and add some property comments

diff --git a/Engine/Plugins/Experimental/Paper2D/Source/Paper2D/Classes/PaperRenderComponent.h b/Engine/Plugins/Experimental/Paper2D/Source/Paper2D/Classes/PaperRenderComponent.h
--- a/Engine/Plugins/Experimental/Paper2D/Source/Paper2D/Classes/PaperRenderComponent.h
+++ b/Engine/Plugins/Experimental/Paper2D/Source/Paper2D/Classes/PaperRenderComponent.h
@@ -11,6 +11,7 @@ class UPaperRenderComponent : public UPrimitiveComponent
 	GENERATED_UCLASS_BODY()
 
 protected:
+	// The sprite asset used by this component
 	UPROPERTY(Category=Sprite, EditAnywhere, BlueprintReadOnly, meta=(DisplayThumbnail = "true"))
 	UPaperSprite* SourceSprite;
 
@@ -18,10 +19,11 @@ protected:
 	UPROPERTY(Category=Sprite, EditAnywhere, BlueprintReadOnly)
 	UMaterialInterface* MaterialOverride;
 
-	/** Physics scene information for this component, holds a single rigid body with multiple shapes. */
+	// Physics scene information for this component, holds a single rigid body with multiple shapes.
 	UPROPERTY(EditAnywhere, BlueprintReadOnly, Category=Collision)//, meta=(ShowOnlyInnerProperties))
 	FBodyInstance2D BodyInstance2D;
 
+	// The color of the sprite (passed to the sprite material as a vertex color)
 	UPROPERTY(BlueprintReadOnly, Interp, Category=Sprite)
 	FLinearColor SpriteColor;
 protected:
@@ -41,7 +43,6 @@ public:
 	void SetSpriteColor(FLinearColor NewColor);
 
 	// UActorComponent interface
-	virtual void TickComponent(float DeltaTime, enum ELevelTick TickType, FActorComponentTickFunction* ThisTickFunction) OVERRIDE;
 	virtual void SendRenderDynamicData_Concurrent() OVERRIDE;
 	virtual void CreatePhysicsState() OVERRIDE;
 	virtual void DestroyPhysicsState() OVERRIDE;
diff --git a/Engine/Plugins/Experimental/Paper2D/Source/Paper2D/Private/PaperRenderComponent.cpp b/Engine/Plugins/Experimental/Paper2D/Source/Paper2D/Private/PaperRenderComponent.cpp
--- a/Engine/Plugins/Experimental/Paper2D/Source/Paper2D/Private/PaperRenderComponent.cpp
+++ b/Engine/Plugins/Experimental/Paper2D/Source/Paper2D/Private/PaperRenderComponent.cpp
@@ -19,9 +19,6 @@ UPaperRenderComponent::UPaperRenderComponent(const FPostConstructInitializePrope
 	SpriteColor = FLinearColor::White;
 
 	Mobility = EComponentMobility::Static;
-	PrimaryComponentTick.bCanEverTick = true;
-	PrimaryComponentTick.TickGroup = TG_DuringPhysics;
-	bTickInEditor = true;
 }
 
 FPrimitiveSceneProxy* UPaperRenderComponent::CreateSceneProxy()
@@ -62,12 +59,6 @@ FBoxSphereBounds UPaperRenderComponent::CalcBounds(const FTransform & LocalToWor
 	}
 }
 
-void UPaperRenderComponent::TickComponent(float DeltaTime, enum ELevelTick TickType, FActorComponentTickFunction* ThisTickFunction)
-{
-	// Indicate we need to send new dynamic data.
-	MarkRenderDynamicDataDirty();
-}
-
 void UPaperRenderComponent::SendRenderDynamicData_Concurrent()
 {
 	if (SceneProxy != NULL)
