Paper2D: Add settings panel

diff --git a/Engine/Plugins/Experimental/Paper2D/Source/Paper2D/Classes/PaperRuntimeSettings.h b/Engine/Plugins/Experimental/Paper2D/Source/Paper2D/Classes/PaperRuntimeSettings.h
new file mode 100644
--- /dev/null
+++ b/Engine/Plugins/Experimental/Paper2D/Source/Paper2D/Classes/PaperRuntimeSettings.h
@@ -0,0 +1,18 @@
+// Copyright 1998-2014 Epic Games, Inc. All Rights Reserved.
+
+#pragma once
+
+#include "PaperRuntimeSettings.generated.h"
+
+/**
+ * Implements the settings for the Paper2D plugin.
+ */
+UCLASS(config = Engine, defaultconfig)
+class PAPER2D_API UPaperRuntimeSettings : public UObject
+{
+	GENERATED_UCLASS_BODY()
+
+	// Test property
+	UPROPERTY(GlobalConfig, EditAnywhere, Category=Settings)
+	int32 TestProperty;
+};
diff --git a/Engine/Plugins/Experimental/Paper2D/Source/Paper2D/Private/PaperRuntimeSettings.cpp b/Engine/Plugins/Experimental/Paper2D/Source/Paper2D/Private/PaperRuntimeSettings.cpp
new file mode 100644
--- /dev/null
+++ b/Engine/Plugins/Experimental/Paper2D/Source/Paper2D/Private/PaperRuntimeSettings.cpp
@@ -0,0 +1,12 @@
+// Copyright 1998-2014 Epic Games, Inc. All Rights Reserved.
+
+#include "Paper2DPrivatePCH.h"
+#include "PaperRuntimeSettings.h"
+
+//////////////////////////////////////////////////////////////////////////
+// UPaperRuntimeSettings
+
+UPaperRuntimeSettings::UPaperRuntimeSettings(const FPostConstructInitializeProperties& PCIP)
+	: Super(PCIP)
+{
+}
diff --git a/Engine/Plugins/Experimental/Paper2D/Source/Paper2DEditor/Paper2DEditor.Build.cs b/Engine/Plugins/Experimental/Paper2D/Source/Paper2DEditor/Paper2DEditor.Build.cs
--- a/Engine/Plugins/Experimental/Paper2D/Source/Paper2DEditor/Paper2DEditor.Build.cs
+++ b/Engine/Plugins/Experimental/Paper2D/Source/Paper2DEditor/Paper2DEditor.Build.cs
@@ -34,5 +34,11 @@ public class Paper2DEditor : ModuleRules
 				"EditorStyle"
 			}
 			);
+
+		PrivateIncludePathModuleNames.AddRange(
+			new string[] {
+				"Settings",
+			}
+		);
 	}
 }
diff --git a/Engine/Plugins/Experimental/Paper2D/Source/Paper2DEditor/Private/Paper2DEditorModule.cpp b/Engine/Plugins/Experimental/Paper2D/Source/Paper2DEditor/Private/Paper2DEditorModule.cpp
--- a/Engine/Plugins/Experimental/Paper2D/Source/Paper2DEditor/Private/Paper2DEditorModule.cpp
+++ b/Engine/Plugins/Experimental/Paper2D/Source/Paper2DEditor/Private/Paper2DEditorModule.cpp
@@ -25,8 +25,14 @@
 #include "Atlasing/AtlasAssetTypeActions.h"
 #include "Atlasing/PaperAtlasGenerator.h"
 
+// Settings
+#include "PaperRuntimeSettings.h"
+#include "Settings.h"
+
 DEFINE_LOG_CATEGORY(LogPaper2DEditor);
 
+#define LOCTEXT_NAMESPACE "Paper2DEditor"
+
 //////////////////////////////////////////////////////////////////////////
 // FPaper2DEditor
 
@@ -113,6 +119,8 @@ public:
 
 		// Integrate Paper2D actions associated with existing engine types (e.g., Texture2D) into the content browser
 		FPaperContentBrowserExtensions::InstallHooks();
+
+		RegisterSettings();
 	}
 
 	virtual void ShutdownModule() OVERRIDE
@@ -125,6 +133,8 @@ public:
 
 		if (UObjectInitialized())
 		{
+			UnregisterSettings();
+
 			FPaperContentBrowserExtensions::RemoveHooks();
 
 			FComponentAssetBrokerage::UnregisterBroker(PaperFlipbookBroker);
@@ -174,8 +184,32 @@ private:
 			FPaperAtlasGenerator::HandleAssetChangedEvent(Atlas);
 		}
 	}
+
+	void RegisterSettings()
+	{
+		if (ISettingsModule* SettingsModule = ISettingsModule::Get())
+		{
+			SettingsModule->RegisterSettings("Project", "Plugins", "Paper2D",
+				LOCTEXT("RuntimeSettingsName", "Paper 2D"),
+				LOCTEXT("RuntimeSettingsDescription", "Configure the Paper 2D plugin"),
+				GetMutableDefault<UPaperRuntimeSettings>()
+				);
+		}
+	}
+
+	void UnregisterSettings()
+	{
+		if (ISettingsModule* SettingsModule = ISettingsModule::Get())
+		{
+			SettingsModule->UnregisterSettings("Project", "Plugins", "Paper2D");
+		}
+	}
 };
 
 //////////////////////////////////////////////////////////////////////////
 
 IMPLEMENT_MODULE(FPaper2DEditor, Paper2DEditor);
+
+//////////////////////////////////////////////////////////////////////////
+
+#undef LOCTEXT_NAMESPACE
\ No newline at end of file
