Paper2D: Removed hack flag to draw 2D scene (can be handled by the batch proxy now)

diff --git a/Engine/Plugins/Experimental/Paper2D/Source/Paper2D/Classes/BodyInstance2D.h b/Engine/Plugins/Experimental/Paper2D/Source/Paper2D/Classes/BodyInstance2D.h
--- a/Engine/Plugins/Experimental/Paper2D/Source/Paper2D/Classes/BodyInstance2D.h
+++ b/Engine/Plugins/Experimental/Paper2D/Source/Paper2D/Classes/BodyInstance2D.h
@@ -20,10 +20,6 @@ struct PAPER2D_API FBodyInstance2D
 	UPROPERTY(EditAnywhere, Category=Physics)
 	float HackHeight;
 
-	//@TODO: SO MUCH EVIL: this controls whether or not this sprite has been nominated to draw the physics debug scene...
-	UPROPERTY(EditAnywhere, BlueprintReadWrite, Category=Physics)
-	uint32 bDrawPhysicsSceneMaster:1;
-
 
 	FBodyInstance2D();
 
