Paper2D: Comment out unused checkerboard code in paper editor viewport client

diff --git a/Engine/Plugins/Experimental/Paper2D/Source/Paper2DEditor/Private/PaperEditorViewportClient.cpp b/Engine/Plugins/Experimental/Paper2D/Source/Paper2DEditor/Private/PaperEditorViewportClient.cpp
--- a/Engine/Plugins/Experimental/Paper2D/Source/Paper2DEditor/Private/PaperEditorViewportClient.cpp
+++ b/Engine/Plugins/Experimental/Paper2D/Source/Paper2DEditor/Private/PaperEditorViewportClient.cpp
@@ -13,7 +13,7 @@ FPaperEditorViewportClient::FPaperEditorViewportClient()
 	ZoomPos = FVector2D::ZeroVector;
 	ZoomAmount = 1.0f;
 
-	ModifyCheckerboardTextureColors();
+	//ModifyCheckerboardTextureColors();
 
 
 	//@TODO: Pretty lame hardcoding
@@ -46,7 +46,7 @@ FPaperEditorViewportClient::FPaperEditorViewportClient()
 
 FPaperEditorViewportClient::~FPaperEditorViewportClient()
 {
-	DestroyCheckerboardTexture();
+	//DestroyCheckerboardTexture();
 }
 
 void FPaperEditorViewportClient::Draw(FViewport* Viewport, FCanvas* Canvas)
