Paper2D: Fix bogus selection rectangle in tile set viewport

diff --git a/Engine/Plugins/Experimental/Paper2D/Source/Paper2DEditor/Private/TileSetEditor.cpp b/Engine/Plugins/Experimental/Paper2D/Source/Paper2DEditor/Private/TileSetEditor.cpp
--- a/Engine/Plugins/Experimental/Paper2D/Source/Paper2DEditor/Private/TileSetEditor.cpp
+++ b/Engine/Plugins/Experimental/Paper2D/Source/Paper2DEditor/Private/TileSetEditor.cpp
@@ -114,7 +114,7 @@ void FTileSetEditorViewportClient::Draw(FViewport* Viewport, FCanvas* Canvas)
 		const float W = Rect.Dimensions.X * ZoomAmount;
 		const float H = Rect.Dimensions.Y * ZoomAmount;
 
-		FCanvasBoxItem BoxItem(FVector2D(X, Y), FVector2D(X+W, Y+H));
+		FCanvasBoxItem BoxItem(FVector2D(X, Y), FVector2D(W, H));
 		BoxItem.SetColor(Rect.Color);
 		Canvas->DrawItem(BoxItem);
 	}
