Paper2D: Partially fixed tile map rendering rot (still not rendering correctly)

diff --git a/Engine/Plugins/Experimental/Paper2D/Source/Paper2D/Private/PaperRenderSceneProxy.cpp b/Engine/Plugins/Experimental/Paper2D/Source/Paper2D/Private/PaperRenderSceneProxy.cpp
--- a/Engine/Plugins/Experimental/Paper2D/Source/Paper2D/Private/PaperRenderSceneProxy.cpp
+++ b/Engine/Plugins/Experimental/Paper2D/Source/Paper2D/Private/PaperRenderSceneProxy.cpp
@@ -190,14 +190,22 @@ FPaperRenderSceneProxy::FPaperRenderSceneProxy(const UPrimitiveComponent* InComp
 
 FPaperRenderSceneProxy::~FPaperRenderSceneProxy()
 {
-	FPaperBatchSceneProxy* Batcher = FPaperBatchManager::GetBatcher(GetScene());
-	Batcher->UnregisterManagedProxy(this);
+#if TEST_BATCHING
+	if (FPaperBatchSceneProxy* Batcher = FPaperBatchManager::GetBatcher(GetScene()))
+	{
+		Batcher->UnregisterManagedProxy(this);
+	}
+#endif
 }
 
 void FPaperRenderSceneProxy::CreateRenderThreadResources()
 {
-	FPaperBatchSceneProxy* Batcher = FPaperBatchManager::GetBatcher(GetScene());
-	Batcher->RegisterManagedProxy(this);
+#if TEST_BATCHING
+	if (FPaperBatchSceneProxy* Batcher = FPaperBatchManager::GetBatcher(GetScene()))
+	{
+		Batcher->RegisterManagedProxy(this);
+	}
+#endif
 }
 
 void FPaperRenderSceneProxy::DrawDynamicElements(FPrimitiveDrawInterface* PDI, const FSceneView* View)
@@ -318,7 +326,7 @@ void FPaperRenderSceneProxy::DrawDynamicElements_RichMesh(FPrimitiveDrawInterfac
 			{
 				const FVector4& SourceVert = Record.RenderVerts[SVI];
 
-				const FVector Pos((PaperAxisX * SourceVert.X) + (PaperAxisY * SourceVert.Y));
+				const FVector Pos((PaperAxisX * SourceVert.X) + (PaperAxisY * SourceVert.Y) + EffectiveOrigin);
 				const FVector2D UV(SourceVert.Z, SourceVert.W);
 
 				new (Vertices) FPaperSpriteVertex(Pos, UV, SpriteColor);
diff --git a/Engine/Plugins/Experimental/Paper2D/Source/Paper2D/Private/PaperTileMapRenderSceneProxy.cpp b/Engine/Plugins/Experimental/Paper2D/Source/Paper2D/Private/PaperTileMapRenderSceneProxy.cpp
--- a/Engine/Plugins/Experimental/Paper2D/Source/Paper2D/Private/PaperTileMapRenderSceneProxy.cpp
+++ b/Engine/Plugins/Experimental/Paper2D/Source/Paper2D/Private/PaperTileMapRenderSceneProxy.cpp
@@ -177,7 +177,14 @@ void FPaperTileMapRenderSceneProxy::DrawDynamicElements(FPrimitiveDrawInterface*
 				TileHeight = TileSet->TileHeight;
 				SourceTexture = TileSet->TileSheet;
 			}
-			FVector2D SourceDimensionsUV(TileWidth, TileHeight);
+
+			FVector2D InverseTextureSize(1.0f, 1.0f);
+			if (SourceTexture != nullptr)
+			{
+				InverseTextureSize = FVector2D(1.0f / SourceTexture->GetSizeX(), 1.0f / SourceTexture->GetSizeY());
+			}
+
+			const FVector2D SourceDimensionsUV(TileWidth * InverseTextureSize.X, TileHeight * InverseTextureSize.Y);
 
 			UPaperTileMapRenderComponent* TileMap = Layer->GetTileMap();
 			FVector2D TileSizeXY(TileMap->TileWidth, TileMap->TileHeight);
@@ -210,15 +217,31 @@ void FPaperTileMapRenderSceneProxy::DrawDynamicElements(FPrimitiveDrawInterface*
 							continue;
 						}
 					}
+					SourceUV.X *= InverseTextureSize.X;
+					SourceUV.Y *= InverseTextureSize.Y;
 
-					const FVector DrawPos(TileSeparationX * X, LayerIndex * LayerSeparation, -TileSeparationY * Y);
+					const FVector DrawPos(TileSeparationX * (X - 0.5f) * PaperAxisX + TileSeparationY * -(Y - 0.5f) * PaperAxisY + ((LayerIndex * LayerSeparation) * PaperAxisZ));
 
 					FSpriteDrawCallRecord& NewTile = *(new (BatchedSprites) FSpriteDrawCallRecord());
 					NewTile.Destination = DrawPos;
 					NewTile.Texture = SourceTexture;
 					NewTile.Color = DrawColor;
 					
-					//@TODO: PAPER: TILE EDITOR: Build data
+					const float W = TileSeparationX;
+					const float H = TileSeparationY;
+
+					const FVector4 BottomLeft(0.0f, 0.0f, SourceUV.X, SourceUV.Y);
+					const FVector4 BottomRight(W, 0.0f, SourceUV.X + SourceDimensionsUV.X, SourceUV.Y);
+					const FVector4 TopRight(W, H, SourceUV.X + SourceDimensionsUV.X, SourceUV.Y + SourceDimensionsUV.Y);
+					const FVector4 TopLeft(0.0f, H, SourceUV.X, SourceUV.Y + SourceDimensionsUV.Y);
+
+					new (NewTile.RenderVerts) FVector4(BottomLeft);
+					new (NewTile.RenderVerts) FVector4(TopRight);
+					new (NewTile.RenderVerts) FVector4(BottomRight);
+
+					new (NewTile.RenderVerts) FVector4(BottomLeft);
+					new (NewTile.RenderVerts) FVector4(TopLeft);
+					new (NewTile.RenderVerts) FVector4(TopRight);
 				}
 			}
 		}
