Fix for non-unity compile error.

diff --git a/Engine/Plugins/Experimental/Paper2D/Source/Paper2D/Private/Rendering/PaperBatchManager.cpp b/Engine/Plugins/Experimental/Paper2D/Source/Paper2D/Private/Rendering/PaperBatchManager.cpp
--- a/Engine/Plugins/Experimental/Paper2D/Source/Paper2D/Private/Rendering/PaperBatchManager.cpp
+++ b/Engine/Plugins/Experimental/Paper2D/Source/Paper2D/Private/Rendering/PaperBatchManager.cpp
@@ -3,6 +3,7 @@
 #include "Paper2DPrivatePCH.h"
 #include "PaperBatchManager.h"
 #include "PaperBatchComponent.h"
+#include "PaperBatchSceneProxy.h"
 
 namespace FPaperBatchManagerNS
 {
