Paper2D: Box2D integration WIP (using actual collision geometry, start of 2D BodySetup implementation)

diff --git a/Engine/Plugins/Experimental/Paper2D/Source/Paper2D/Classes/BodyInstance2D.h b/Engine/Plugins/Experimental/Paper2D/Source/Paper2D/Classes/BodyInstance2D.h
--- a/Engine/Plugins/Experimental/Paper2D/Source/Paper2D/Classes/BodyInstance2D.h
+++ b/Engine/Plugins/Experimental/Paper2D/Source/Paper2D/Classes/BodyInstance2D.h
@@ -14,12 +14,6 @@ struct PAPER2D_API FBodyInstance2D
 	UPROPERTY(EditAnywhere, BlueprintReadWrite, Category=Physics)
 	uint32 bSimulatePhysics:1;
 
-	UPROPERTY(EditAnywhere, Category=Physics)
-	float HackWidth;
-
-	UPROPERTY(EditAnywhere, Category=Physics)
-	float HackHeight;
-
 
 	FBodyInstance2D();
 
diff --git a/Engine/Plugins/Experimental/Paper2D/Source/Paper2D/Classes/BodySetup2D.h b/Engine/Plugins/Experimental/Paper2D/Source/Paper2D/Classes/BodySetup2D.h
--- a/Engine/Plugins/Experimental/Paper2D/Source/Paper2D/Classes/BodySetup2D.h
+++ b/Engine/Plugins/Experimental/Paper2D/Source/Paper2D/Classes/BodySetup2D.h
@@ -5,15 +5,97 @@
 #include "BodyInstance2D.h"
 #include "BodySetup2D.generated.h"
 
+
+
+
+
+/** Container for an aggregate of collision shapes */
+USTRUCT()
+struct FAggregateGeometry2D
+{
+	GENERATED_USTRUCT_BODY()
+
+	UPROPERTY(EditAnywhere, editfixedsize, Category = KAggregateGeom)
+	TArray<struct FKSphereElem> SphereElems;
+
+	UPROPERTY(EditAnywhere, editfixedsize, Category = KAggregateGeom)
+	TArray<struct FKBoxElem> BoxElems;
+
+	UPROPERTY(EditAnywhere, editfixedsize, Category = KAggregateGeom)
+	TArray<struct FKSphylElem> SphylElems;
+
+	UPROPERTY(EditAnywhere, editfixedsize, Category = KAggregateGeom)
+	TArray<struct FKConvexElem> ConvexElems;
+
+	//class FKConvexGeomRenderInfo* RenderInfo;
+
+	FAggregateGeometry2D()
+//		: RenderInfo(NULL)
+	{}
+
+// 	int32 GetElementCount() const
+// 	{
+// 		return SphereElems.Num() + SphylElems.Num() + BoxElems.Num() + ConvexElems.Num();
+// 	}
+// 
+// 	int32 GetElementCount(int32 Type) const;
+// 
+	void EmptyElements()
+	{
+		BoxElems.Empty();
+		ConvexElems.Empty();
+		SphylElems.Empty();
+		SphereElems.Empty();
+
+//		FreeRenderInfo();
+	}
+// 
+// 
+// 
+// 	void Serialize(const FArchive& Ar);
+// 
+// 	void DrawAggGeom(class FPrimitiveDrawInterface* PDI, const FTransform& Transform, const FColor Color, const FMaterialRenderProxy* MatInst, bool bPerHullColor, bool bDrawSolid);
+// 
+// 	/** Release the RenderInfo (if its there) and safely clean up any resources. Call on the game thread. */
+// 	void FreeRenderInfo();
+// 
+// 	FBox CalcAABB(const FTransform& Transform);
+// 
+// 	/**
+// 	* Calculates a tight box-sphere bounds for the aggregate geometry; this is more expensive than CalcAABB
+// 	* (tight meaning the sphere may be smaller than would be required to encompass the AABB, but all individual components lie within both the box and the sphere)
+// 	*
+// 	* @param Output The output box-sphere bounds calculated for this set of aggregate geometry
+// 	*	@param LocalToWorld Transform
+// 	*/
+// 	void CalcBoxSphereBounds(FBoxSphereBounds& Output, const FTransform& LocalToWorld);
+// 
+// 	/** Returns the volume of this element */
+// 	float GetVolume(const FVector& Scale) const;
+};
+
+
+
 UCLASS(DependsOn=UEngineTypes)
 class PAPER2D_API UBodySetup2D : public UObject
 {
 	GENERATED_UCLASS_BODY()
 
+	UPROPERTY()
+	FAggregateGeometry2D AggGeom2D;
+
 	/** Default properties of the body instance, copied into objects on instantiation */
 	UPROPERTY(EditAnywhere, Category=Collision, meta=(FullyExpand = "true"))
 	struct FBodyInstance2D DefaultInstance;
 
+	/** Create Physics meshes (ConvexMeshes, TriMesh & TriMeshNegX) from cooked data */
+	void CreatePhysicsMeshes();
+
+
+#if WITH_EDITOR
+	/** Release Physics meshes (ConvexMeshes, TriMesh & TriMeshNegX). Must be called before the BodySetup is destroyed */
+	void InvalidatePhysicsData();
+#endif
 
 protected:
 	/** Physics triangle mesh, created from cooked data in CreatePhysicsMeshes */
diff --git a/Engine/Plugins/Experimental/Paper2D/Source/Paper2D/Classes/PaperRenderComponent.h b/Engine/Plugins/Experimental/Paper2D/Source/Paper2D/Classes/PaperRenderComponent.h
--- a/Engine/Plugins/Experimental/Paper2D/Source/Paper2D/Classes/PaperRenderComponent.h
+++ b/Engine/Plugins/Experimental/Paper2D/Source/Paper2D/Classes/PaperRenderComponent.h
@@ -64,7 +64,7 @@ public:
 	// End of UPrimitiveComponent interface
 
 	/** Return the BodySetup to use for this PrimitiveComponent (single body case) */
-	virtual class UBodySetup2D* GetRBBodySetup2D();
+	virtual class UBodySetup2D* GetBodySetup2D();
 
 	// Returns the wireframe color to use for this component.
 	FLinearColor GetWireframeColor() const;
diff --git a/Engine/Plugins/Experimental/Paper2D/Source/Paper2D/Classes/PaperSprite.h b/Engine/Plugins/Experimental/Paper2D/Source/Paper2D/Classes/PaperSprite.h
--- a/Engine/Plugins/Experimental/Paper2D/Source/Paper2D/Classes/PaperSprite.h
+++ b/Engine/Plugins/Experimental/Paper2D/Source/Paper2D/Classes/PaperSprite.h
@@ -85,6 +85,10 @@ public:
 	UPROPERTY() //@TODO: Is anything in this worth exposing to edit? Category=Physics, EditAnywhere, EditInline
 	class UBodySetup* BodySetup3D;
 
+	// Baked physics data (2D domain).
+	UPROPERTY() //@TODO: Is anything in this worth exposing to edit? Category=Physics, EditAnywhere, EditInline
+	class UBodySetup2D* BodySetup2D;
+
 #if WITH_EDITORONLY_DATA
 
 protected:
diff --git a/Engine/Plugins/Experimental/Paper2D/Source/Paper2D/Private/PaperRenderComponent.cpp b/Engine/Plugins/Experimental/Paper2D/Source/Paper2D/Private/PaperRenderComponent.cpp
--- a/Engine/Plugins/Experimental/Paper2D/Source/Paper2D/Private/PaperRenderComponent.cpp
+++ b/Engine/Plugins/Experimental/Paper2D/Source/Paper2D/Private/PaperRenderComponent.cpp
@@ -102,17 +102,17 @@ void UPaperRenderComponent::CreatePhysicsState2D()
 {
 	if (!BodyInstance2D.IsValidBodyInstance() && (SourceSprite != NULL) && (SourceSprite->SpriteCollisionDomain == ESpriteCollisionMode::Use2DPhysics))
 	{
-		UBodySetup2D* BodySetup = GetRBBodySetup2D();
-		//@TODO: if (UBodySetup2D* BodySetup = GetRBBodySetup2D())
+		UBodySetup2D* BodySetup = GetBodySetup2D();
+		//@TODO: if (UBodySetup2D* BodySetup = GetBodySetup2D())
 		{
 			BodyInstance2D.InitBody(BodySetup, SourceSprite, ComponentToWorld, this);
 		}
 	}
 }
 
-UBodySetup2D* UPaperRenderComponent::GetRBBodySetup2D()
+UBodySetup2D* UPaperRenderComponent::GetBodySetup2D()
 {
-	return NULL;
+	return (SourceSprite != nullptr) ? SourceSprite->BodySetup2D : nullptr;
 }
 
 void UPaperRenderComponent::OnUpdateTransform(bool bSkipPhysicsMove)
diff --git a/Engine/Plugins/Experimental/Paper2D/Source/Paper2D/Private/PaperRenderSceneProxy.cpp b/Engine/Plugins/Experimental/Paper2D/Source/Paper2D/Private/PaperRenderSceneProxy.cpp
--- a/Engine/Plugins/Experimental/Paper2D/Source/Paper2D/Private/PaperRenderSceneProxy.cpp
+++ b/Engine/Plugins/Experimental/Paper2D/Source/Paper2D/Private/PaperRenderSceneProxy.cpp
@@ -216,13 +216,6 @@ void FPaperRenderSceneProxy::DrawDynamicElements(FPrimitiveDrawInterface* PDI, c
 
 	if (SourceSprite != NULL)//const UPaperRenderComponent* TypedFoo = Cast<const UPaperRenderComponent>(OwnerComponent))
 	{
-		// Show 2D physics
-		//@TODO: SO MUCH EVIL: Hack to debug draw the physics scene
-// 		if (TypedFoo->BodyInstance2D.bDrawPhysicsSceneMaster)
-// 		{
-// 			FPhysicsIntegration2D::DrawDebugPhysics(TypedFoo->GetWorld(), PDI, View);
-// 		}
-
 		// Show 3D physics
 		if ((View->Family->EngineShowFlags.Collision /*@TODO: && bIsCollisionEnabled*/) && AllowDebugViewmodes())
 		{
diff --git a/Engine/Plugins/Experimental/Paper2D/Source/Paper2D/Private/PaperSprite.cpp b/Engine/Plugins/Experimental/Paper2D/Source/Paper2D/Private/PaperSprite.cpp
--- a/Engine/Plugins/Experimental/Paper2D/Source/Paper2D/Private/PaperSprite.cpp
+++ b/Engine/Plugins/Experimental/Paper2D/Source/Paper2D/Private/PaperSprite.cpp
@@ -402,6 +402,7 @@ void UPaperSprite::PostEditChangeProperty(FPropertyChangedEvent& PropertyChanged
 
 	if ((PropertyName == GET_MEMBER_NAME_CHECKED(UPaperSprite, SpriteCollisionDomain)) ||
 		(PropertyName == GET_MEMBER_NAME_CHECKED(UPaperSprite, BodySetup3D)) ||
+		(PropertyName == GET_MEMBER_NAME_CHECKED(UPaperSprite, BodySetup2D)) ||
 		(PropertyName == GET_MEMBER_NAME_CHECKED(UPaperSprite, CollisionGeometry)) )
 	{
 		bCollisionDataModified = true;
@@ -448,16 +449,22 @@ void UPaperSprite::RebuildCollisionData()
 	switch (SpriteCollisionDomain)
 	{
 	case ESpriteCollisionMode::Use3DPhysics:
-		if (BodySetup3D == NULL)
+		BodySetup2D = nullptr;
+		if (BodySetup3D == nullptr)
 		{
-			BodySetup3D = ConstructObject<UBodySetup>(UBodySetup::StaticClass(), this);
+			BodySetup3D = NewObject<UBodySetup>(this);
 		}
 		break;
 	case ESpriteCollisionMode::Use2DPhysics:
-		BodySetup3D = NULL;
+		BodySetup3D = nullptr;
+		if (BodySetup2D == nullptr)
+		{
+			BodySetup2D = NewObject<UBodySetup2D>(this);
+		}
 		break;
 	case ESpriteCollisionMode::None:
-		BodySetup3D = NULL;
+		BodySetup2D = nullptr;
+		BodySetup3D = nullptr;
 		CollisionGeometry.Reset();
 		break;
 	}
@@ -506,7 +513,7 @@ void UPaperSprite::BuildCustomCollisionData()
 	{
 	case ESpriteCollisionMode::Use3DPhysics:
 		{
-			check(BodySetup3D);
+			checkSlow(BodySetup3D);
 			BodySetup3D->AggGeom.EmptyElements();
 
 			//@TODO: Observe if the custom data is really a hand-edited bounding box, and generate box geom instead of convex geom!
@@ -536,6 +543,33 @@ void UPaperSprite::BuildCustomCollisionData()
 		}
 		break;
 	case ESpriteCollisionMode::Use2DPhysics:
+		{
+			checkSlow(BodySetup2D);
+			BodySetup2D->AggGeom2D.EmptyElements();
+
+			//@TODO: Observe if the custom data is really a hand-edited bounding box, and generate box geom instead of convex geom!
+
+			//@TODO: Use this guy instead: DecomposeMeshToHulls
+			//@TODO: Merge triangles that are convex together!
+			int32 RunningIndex = 0;
+			for (int32 TriIndex = 0; TriIndex < CollisionData.Num() / 3; ++TriIndex)
+			{
+				FKConvexElem& ConvexTri = *new (BodySetup2D->AggGeom2D.ConvexElems) FKConvexElem();
+				ConvexTri.VertexData.Empty(3);
+				for (int32 Index = 0; Index < 3; ++Index)
+				{
+					const FVector2D& Pos2D = CollisionData[RunningIndex++];
+
+					const FVector Pos3D = (PaperAxisX * Pos2D.X) + (PaperAxisY * Pos2D.X);
+
+					new (ConvexTri.VertexData) FVector(Pos3D);
+				}
+				ConvexTri.UpdateElemBox();
+			}
+
+			BodySetup2D->InvalidatePhysicsData();
+			BodySetup2D->CreatePhysicsMeshes();
+	}
 		break;
 	default:
 		check(false);
@@ -554,7 +588,7 @@ void UPaperSprite::BuildBoundingBoxCollisionData(bool bUseTightBounds)
 	case ESpriteCollisionMode::Use3DPhysics:
 		{
 			// Store the bounding box as an actual box for 3D physics
-			check(BodySetup3D);
+			checkSlow(BodySetup3D);
 			BodySetup3D->AggGeom.EmptyElements();
 
 			// Determine the box center in pivot space
@@ -574,7 +608,26 @@ void UPaperSprite::BuildBoundingBoxCollisionData(bool bUseTightBounds)
 		}
 		break;
 	case ESpriteCollisionMode::Use2DPhysics:
-		//@TODO:!!!
+		{
+			//@TODO:!!!
+			checkSlow(BodySetup2D);
+			BodySetup2D->AggGeom2D.EmptyElements();
+
+			// Determine the box center in pivot space
+			const FVector2D& BoxSize2D = CollisionGeometry.Polygons[0].BoxSize;
+			const FVector2D& BoxPos = CollisionGeometry.Polygons[0].BoxPosition;
+			const FVector2D CenterInTextureSpace = BoxPos + (BoxSize2D * 0.5f);
+			const FVector2D CenterInPivotSpace = ConvertTextureSpaceToPivotSpace(CenterInTextureSpace);
+
+			// Create a new box primitive
+			const FVector BoxSize3D = (PaperAxisX * BoxSize2D.X) + (PaperAxisY * BoxSize2D.Y) + (PaperAxisZ * CollisionThickness);
+
+			FKBoxElem& Box = *new (BodySetup2D->AggGeom2D.BoxElems) FKBoxElem(FMath::Abs(BoxSize3D.X), FMath::Abs(BoxSize3D.Y), FMath::Abs(BoxSize3D.Z));
+			Box.Center = (PaperAxisX * CenterInPivotSpace.X) + (PaperAxisY * CenterInPivotSpace.Y);
+
+			BodySetup2D->InvalidatePhysicsData();
+			BodySetup2D->CreatePhysicsMeshes();
+		}
 		break;
 	default:
 		check(false);
diff --git a/Engine/Plugins/Experimental/Paper2D/Source/Paper2D/Private/Physics/BodyInstance2D.cpp b/Engine/Plugins/Experimental/Paper2D/Source/Paper2D/Private/Physics/BodyInstance2D.cpp
--- a/Engine/Plugins/Experimental/Paper2D/Source/Paper2D/Private/Physics/BodyInstance2D.cpp
+++ b/Engine/Plugins/Experimental/Paper2D/Source/Paper2D/Private/Physics/BodyInstance2D.cpp
@@ -17,10 +17,74 @@ FBodyInstance2D::FBodyInstance2D()
 //@TODO: Everything about this block of code...
 void FBodyInstance2D::InitBody(class UBodySetup2D* Setup, const UPaperSprite* SpriteSetupHack, const FTransform& Transform, class UPrimitiveComponent* PrimComp)
 {
-#if WITH_BOX2D && 0 //@TODO: DISABLED BOX2D FOR NOW!!!
+#if WITH_BOX2D
 	if (b2World* World = FPhysicsIntegration2D::FindAssociatedWorld(PrimComp->GetWorld()))
 	{
-		if ((SpriteSetupHack->CollisionData.Num() > 0) && ((SpriteSetupHack->CollisionData.Num() % 3) == 0))
+		const FVector Scale3D = Transform.GetScale3D();
+
+		if (Setup != nullptr)
+		{
+			OwnerComponentPtr = PrimComp;
+
+			b2BodyDef BodyDefinition;
+			BodyDefinition.type = bSimulatePhysics ? b2_dynamicBody : b2_kinematicBody;
+
+			BodyInstancePtr = World->CreateBody(&BodyDefinition);
+			BodyInstancePtr->SetUserData(this);
+
+// 			const FVector2D PivotPos = FVector2D::ZeroVector;
+// 			const FVector2D OffsetPos = FVector2D(SpriteSetupHack->GetSourceSize().X*-0.5f, SpriteSetupHack->GetSourceSize().Y*-0.5f) + PivotPos;
+
+			// Boxes
+			for (const FKBoxElem& Box : Setup->AggGeom2D.BoxElems)
+			{
+				const b2Vec2 HalfBoxSize = FPhysicsIntegration2D::ConvertUnrealVectorToBox(FVector(Box.X, Box.Y, Box.Z) * 0.5f  * Scale3D);
+				const b2Vec2 BoxCenter = FPhysicsIntegration2D::ConvertUnrealVectorToBox(Box.Center * Scale3D);
+
+				b2PolygonShape DynamicBox;
+				DynamicBox.SetAsBox(HalfBoxSize.x, HalfBoxSize.y, BoxCenter, 0.0f);
+
+				b2FixtureDef FixtureDef;
+				FixtureDef.shape = &DynamicBox;
+				FixtureDef.density = 1.0f; //@TODO:
+				FixtureDef.friction = 0.3f;//@TODO:
+
+				BodyInstancePtr->CreateFixture(&FixtureDef);
+			}
+
+			// Convex hulls
+			for (const FKConvexElem& Convex : Setup->AggGeom2D.ConvexElems)
+			{
+				const int32 NumVerts = Convex.VertexData.Num();
+
+				if (NumVerts <= b2_maxPolygonVertices)
+				{
+					TArray<b2Vec2, TInlineAllocator<b2_maxPolygonVertices>> Verts;
+
+					for (int32 VertexIndex = 0; VertexIndex < Convex.VertexData.Num(); ++VertexIndex)
+					{
+						const FVector SourceVert = Convex.VertexData[VertexIndex];
+						new (Verts) b2Vec2(FPhysicsIntegration2D::ConvertUnrealVectorToBox(SourceVert * Scale3D));
+					}
+
+					b2PolygonShape ConvexPoly;
+					ConvexPoly.Set(Verts.GetTypedData(), Verts.Num());
+
+					b2FixtureDef FixtureDef;
+					FixtureDef.shape = &ConvexPoly;
+					FixtureDef.density = 1.0f; //@TODO:
+					FixtureDef.friction = 0.3f;//@TODO:
+
+					BodyInstancePtr->CreateFixture(&FixtureDef);
+				}
+				else
+				{
+					UE_LOG(LogPaper2D, Warning, TEXT("Too many vertices in a 2D convex body")); //@TODO: Create a better error message that indicates the asset
+				}
+			}
+
+#if 0
+		if ((SpriteSetupHack->CollisionGeometry.Num() > 0) && ((SpriteSetupHack->CollisionGeometry.Num() % 3) == 0))
 		{
 			OwnerComponentPtr = PrimComp;
 
@@ -31,14 +95,14 @@ void FBodyInstance2D::InitBody(class UBodySetup2D* Setup, const UPaperSprite* Sp
 			BodyInstancePtr->SetUserData(this);
 
 			const FVector2D PivotPos = FVector2D::ZeroVector;
-			const FVector2D OffsetPos = FVector2D(SpriteSetupHack->SourceDimension.X*-0.5f, SpriteSetupHack->SourceDimension.Y*-0.5f) + PivotPos;
+			const FVector2D OffsetPos = FVector2D(SpriteSetupHack->GetSourceSize().X*-0.5f, SpriteSetupHack->GetSourceSize().Y*-0.5f) + PivotPos;
 
-			for (int32 TriangleIndex = 0; TriangleIndex < SpriteSetupHack->CollisionData.Num() / 3; ++TriangleIndex)
+			for (int32 TriangleIndex = 0; TriangleIndex < SpriteSetupHack->CollisionGeometry.Num() / 3; ++TriangleIndex)
 			{
 				b2Vec2 Verts[3];
 				for (int32 VertexIndex = 0; VertexIndex < 3; ++VertexIndex)
 				{
-					const FVector2D SourceVert = SpriteSetupHack->CollisionData[(TriangleIndex * 3) + VertexIndex] + OffsetPos;
+					const FVector2D SourceVert = SpriteSetupHack->CollisionGeometry[(TriangleIndex * 3) + VertexIndex] + OffsetPos;
 					Verts[VertexIndex] = FPhysicsIntegration2D::ConvertUnrealVectorToBox(FVector(SourceVert.X, 0.0f, SourceVert.Y));
 				}
 				
@@ -52,6 +116,10 @@ void FBodyInstance2D::InitBody(class UBodySetup2D* Setup, const UPaperSprite* Sp
 
 				BodyInstancePtr->CreateFixture(&FixtureDef);
 			}
+#endif
+
+
+
 
 #if 0
 			const b2Vec2 HalfBoxSize = FPhysicsIntegration2D::ConvertUnrealVectorToBox(FVector(SpriteSetupHack->BoundingBoxSize.X * 0.5f, 0.0f, SpriteSetupHack->BoundingBoxSize.Y * 0.5f));
diff --git a/Engine/Plugins/Experimental/Paper2D/Source/Paper2D/Private/Physics/BodySetup2D.cpp b/Engine/Plugins/Experimental/Paper2D/Source/Paper2D/Private/Physics/BodySetup2D.cpp
--- a/Engine/Plugins/Experimental/Paper2D/Source/Paper2D/Private/Physics/BodySetup2D.cpp
+++ b/Engine/Plugins/Experimental/Paper2D/Source/Paper2D/Private/Physics/BodySetup2D.cpp
@@ -9,3 +9,20 @@ UBodySetup2D::UBodySetup2D(const FPostConstructInitializeProperties& PCIP)
 	: Super(PCIP)
 {
 }
+
+
+#if WITH_EDITOR
+
+void UBodySetup2D::InvalidatePhysicsData()
+{
+	//ClearPhysicsMeshes();
+	//BodySetupGuid = FGuid::NewGuid(); // change the guid
+	//CookedFormatData.FlushData();
+}
+
+#endif
+
+
+void UBodySetup2D::CreatePhysicsMeshes()
+{
+}
