Paper2D: Adding tick to batch component

diff --git a/Engine/Plugins/Experimental/Paper2D/Source/Paper2D/Private/Rendering/PaperBatchComponent.h b/Engine/Plugins/Experimental/Paper2D/Source/Paper2D/Private/Rendering/PaperBatchComponent.h
--- a/Engine/Plugins/Experimental/Paper2D/Source/Paper2D/Private/Rendering/PaperBatchComponent.h
+++ b/Engine/Plugins/Experimental/Paper2D/Source/Paper2D/Private/Rendering/PaperBatchComponent.h
@@ -10,6 +10,10 @@ class UPaperBatchComponent : public UPrimitiveComponent
 {
 	GENERATED_UCLASS_BODY()
 
+	// UActorComponent interface
+	virtual void TickComponent(float DeltaTime, enum ELevelTick TickType, FActorComponentTickFunction* ThisTickFunction) OVERRIDE;
+	// End of UActorComponent interface
+
 	// UPrimitiveComponent interface
 	virtual FPrimitiveSceneProxy* CreateSceneProxy() OVERRIDE;
 	virtual FBoxSphereBounds CalcBounds(const FTransform & LocalToWorld) const OVERRIDE;
