Paper2D: Box2D debug visualization WIP

diff --git a/Engine/Plugins/Experimental/Paper2D/Source/Paper2D/Private/Rendering/PaperBatchComponent.cpp b/Engine/Plugins/Experimental/Paper2D/Source/Paper2D/Private/Rendering/PaperBatchComponent.cpp
--- a/Engine/Plugins/Experimental/Paper2D/Source/Paper2D/Private/Rendering/PaperBatchComponent.cpp
+++ b/Engine/Plugins/Experimental/Paper2D/Source/Paper2D/Private/Rendering/PaperBatchComponent.cpp
@@ -11,6 +11,10 @@ UPaperBatchComponent::UPaperBatchComponent(const FPostConstructInitializePropert
 	: Super(PCIP)
 {
 	BodyInstance.SetCollisionEnabled(ECollisionEnabled::NoCollision);
+
+	bAutoActivate = true;
+	PrimaryComponentTick.TickGroup = TG_PostPhysics;
+	PrimaryComponentTick.bCanEverTick = true;
 }
 
 FPrimitiveSceneProxy* UPaperBatchComponent::CreateSceneProxy()
@@ -25,3 +29,14 @@ FBoxSphereBounds UPaperBatchComponent::CalcBounds(const FTransform & LocalToWorl
 	FBoxSphereBounds Bounds(FVector::ZeroVector, FVector(HALF_WORLD_MAX, HALF_WORLD_MAX, HALF_WORLD_MAX), HALF_WORLD_MAX);
 	return Bounds;
 }
+
+void UPaperBatchComponent::TickComponent(float DeltaTime, enum ELevelTick TickType, FActorComponentTickFunction* ThisTickFunction)
+{
+	// Draw all the debug rendering for the 2D physics scene
+	//if (View->Family->EngineShowFlags.Collision && AllowDebugViewmodes())
+	{
+	//	FPhysicsIntegration2D::DrawDebugPhysics(GetWorld(), PDI, View); //@TODO: GWorld is worse than the disease
+	}
+
+	
+}
diff --git a/Engine/Plugins/Experimental/Paper2D/Source/Paper2D/Private/Rendering/PaperBatchSceneProxy.cpp b/Engine/Plugins/Experimental/Paper2D/Source/Paper2D/Private/Rendering/PaperBatchSceneProxy.cpp
--- a/Engine/Plugins/Experimental/Paper2D/Source/Paper2D/Private/Rendering/PaperBatchSceneProxy.cpp
+++ b/Engine/Plugins/Experimental/Paper2D/Source/Paper2D/Private/Rendering/PaperBatchSceneProxy.cpp
@@ -4,6 +4,8 @@
 #include "PaperBatchSceneProxy.h"
 #include "PaperRenderSceneProxy.h"
 
+#include "Physics/Box2DIntegration.h"
+
 //////////////////////////////////////////////////////////////////////////
 // FPaperBatchSceneProxy
 
@@ -61,6 +63,17 @@ struct FSortSpriteBatchesByTexture
 
 void FPaperBatchSceneProxy::DrawDynamicElements(FPrimitiveDrawInterface* PDI, const FSceneView* View)
 {
+#if 0
+	// Not thread safe right now (not to mention use of GWorld...)
+
+	// Draw all the debug rendering for the 2D physics scene
+	if ((View->Family->EngineShowFlags.Collision /*@TODO: && bIsCollisionEnabled*/) && AllowDebugViewmodes())
+	{
+		FPhysicsIntegration2D::DrawDebugPhysics(GWorld, PDI, View); //@TODO: GWorld is worse than the disease
+	}
+#endif
+
+
 #if TEST_BATCHING
 	// First sort by material
 	ManagedProxies.Sort(FSortSpriteProxiesByMaterial());
