Paper2D: Split basic 2D physics support out of UPaperRenderComponent into UPrimitiveComponent2D
Split FAggregateGeometry2D into another file and changed elements to actually be 2D

diff --git a/Engine/Plugins/Experimental/Paper2D/Source/Paper2D/Classes/AggregateGeometry2D.h b/Engine/Plugins/Experimental/Paper2D/Source/Paper2D/Classes/AggregateGeometry2D.h
new file mode 100644
--- /dev/null
+++ b/Engine/Plugins/Experimental/Paper2D/Source/Paper2D/Classes/AggregateGeometry2D.h
@@ -0,0 +1,135 @@
+// Copyright 1998-2014 Epic Games, Inc. All Rights Reserved.
+
+#pragma once
+
+#include "AggregateGeometry2D.generated.h"
+
+
+/** One convex hull, used for simplified collision. */
+USTRUCT()
+struct FConvexElement2D
+{
+	GENERATED_USTRUCT_BODY()
+
+	// Vertices that lie on the convex hull
+	UPROPERTY()
+	TArray<FVector2D> VertexData;
+
+	FConvexElement2D()
+	{}
+
+	void DrawElemWire(class FPrimitiveDrawInterface* PDI, const FTransform& ElemTM, const FColor Color);
+	void DrawElemSolid(class FPrimitiveDrawInterface* PDI, const FTransform& ElemTM, float Scale, const FMaterialRenderProxy* MaterialRenderProxy);
+
+	// Returns the area of this element
+	float GetArea(const FVector& Scale) const;
+};
+
+/** Sphere shape used for collision */
+USTRUCT()
+struct FCircleElement2D
+{
+	GENERATED_USTRUCT_BODY()
+
+	UPROPERTY(Category=KSphereElem, VisibleAnywhere)
+	FVector2D Center;
+
+	UPROPERTY(Category=KSphereElem, VisibleAnywhere)
+	float Radius;
+
+public:
+	FCircleElement2D()
+		: Center(FVector2D::ZeroVector)
+		, Radius(1.0f)
+	{
+
+	}
+
+	FCircleElement2D(float InitialRadius)
+		: Center(FVector2D::ZeroVector)
+		, Radius(InitialRadius)
+	{
+	}
+
+	friend bool operator==(const FCircleElement2D& LHS, const FCircleElement2D& RHS)
+	{
+		return (LHS.Center == RHS.Center) && (LHS.Radius == RHS.Radius);
+	}
+
+	// Returns the area of this element
+	float GetArea(const FVector& Scale) const;
+	
+	void DrawElemWire(class FPrimitiveDrawInterface* PDI, const FTransform& ElemTM, float Scale, const FColor Color);
+	void DrawElemSolid(class FPrimitiveDrawInterface* PDI, const FTransform& ElemTM, float Scale, const FMaterialRenderProxy* MaterialRenderProxy);
+};
+
+/** Box shape used for collision */
+USTRUCT()
+struct FBoxElement2D
+{
+	GENERATED_USTRUCT_BODY()
+
+	// Center of the box
+	UPROPERTY(Category=KBoxElem, VisibleAnywhere)
+	FVector2D Center;
+
+	// Width of the box
+	UPROPERTY(Category=KBoxElem, VisibleAnywhere)
+	float Width;
+
+	// Height of the box
+	UPROPERTY(Category=KBoxElem, VisibleAnywhere)
+	float Height;
+
+	// Rotation of the box (in degreens)
+	UPROPERTY(Category = KBoxElem, VisibleAnywhere)
+	float Angle;
+
+	FBoxElement2D()
+		: Center(FVector::ZeroVector)
+		, Width(1.0f)
+		, Height(1.0f)
+		, Angle(0.0f)
+	{
+	}
+
+	friend bool operator==(const FBoxElement2D& LHS, const FBoxElement2D& RHS)
+	{
+		return (LHS.Center == RHS.Center) &&
+			(LHS.Width == RHS.Width) &&
+			(LHS.Height == RHS.Height) &&
+			(LHS.Angle == RHS.Angle);
+	};
+
+	// Returns the area of this element
+	float GetArea(const FVector& Scale) const;
+
+	void DrawElemWire(class FPrimitiveDrawInterface* PDI, const FTransform& ElemTM, float Scale, const FColor Color);
+	void DrawElemSolid(class FPrimitiveDrawInterface* PDI, const FTransform& ElemTM, float Scale, const FMaterialRenderProxy* MaterialRenderProxy);
+};
+
+
+// Container for an aggregate of 2D collision shapes
+USTRUCT()
+struct FAggregateGeometry2D
+{
+	GENERATED_USTRUCT_BODY()
+
+public:
+	UPROPERTY(EditAnywhere, editfixedsize, Category = KAggregateGeom)
+	TArray<struct FCircleElement2D> CircleElements;
+
+	UPROPERTY(EditAnywhere, editfixedsize, Category = KAggregateGeom)
+	TArray<struct FBoxElement2D> BoxElements;
+
+	UPROPERTY(EditAnywhere, editfixedsize, Category = KAggregateGeom)
+	TArray<struct FConvexElement2D> ConvexElements;
+
+public:
+	void EmptyElements();
+
+	void DrawAggGeom(class FPrimitiveDrawInterface* PDI, const FTransform& Transform, const FColor Color, const FMaterialRenderProxy* MatInst, bool bPerHullColor, bool bDrawSolid) const;
+
+	// Returns the area of all contained elements
+ 	float GetArea(const FVector& Scale) const;
+};
diff --git a/Engine/Plugins/Experimental/Paper2D/Source/Paper2D/Classes/BodyInstance2D.h b/Engine/Plugins/Experimental/Paper2D/Source/Paper2D/Classes/BodyInstance2D.h
--- a/Engine/Plugins/Experimental/Paper2D/Source/Paper2D/Classes/BodyInstance2D.h
+++ b/Engine/Plugins/Experimental/Paper2D/Source/Paper2D/Classes/BodyInstance2D.h
@@ -9,15 +9,13 @@ struct PAPER2D_API FBodyInstance2D
 {
 	GENERATED_USTRUCT_BODY()
 
-
-	/** If true, this body will use simulation. If false, will be 'fixed' (ie kinematic) and move where it is told. */
+	/** If true, this body will use simulation. If false, will be 'fixed' (i.e., kinematic) and move where it is told. */
 	UPROPERTY(EditAnywhere, BlueprintReadWrite, Category=Physics)
 	uint32 bSimulatePhysics:1;
 
-
 	FBodyInstance2D();
 
-	void InitBody(class UBodySetup2D* Setup, const class UPaperSprite* SpriteSetupHack, const FTransform& Transform, class UPrimitiveComponent* PrimComp);//,class FPhysScene* InRBScene, class physx::PxAggregate* InAggregate = NULL);
+	void InitBody(class UBodySetup2D* Setup, const FTransform& Transform, class UPrimitiveComponent* PrimComp);
 	void TermBody();
 
 	/** See if this body is valid. */
@@ -39,5 +37,8 @@ public:
 	TWeakObjectPtr<class UPrimitiveComponent> OwnerComponentPtr;
 
 protected:
+
+#if WITH_BOX2D
 	class b2Body* BodyInstancePtr;
+#endif
 };
diff --git a/Engine/Plugins/Experimental/Paper2D/Source/Paper2D/Classes/BodySetup2D.h b/Engine/Plugins/Experimental/Paper2D/Source/Paper2D/Classes/BodySetup2D.h
--- a/Engine/Plugins/Experimental/Paper2D/Source/Paper2D/Classes/BodySetup2D.h
+++ b/Engine/Plugins/Experimental/Paper2D/Source/Paper2D/Classes/BodySetup2D.h
@@ -3,77 +3,9 @@
 #pragma once
 
 #include "BodyInstance2D.h"
-#include "BodySetup2D.generated.h"
-
-
-
-
-
-/** Container for an aggregate of collision shapes */
-USTRUCT()
-struct FAggregateGeometry2D
-{
-	GENERATED_USTRUCT_BODY()
-
-	UPROPERTY(EditAnywhere, editfixedsize, Category = KAggregateGeom)
-	TArray<struct FKSphereElem> SphereElems;
-
-	UPROPERTY(EditAnywhere, editfixedsize, Category = KAggregateGeom)
-	TArray<struct FKBoxElem> BoxElems;
-
-	UPROPERTY(EditAnywhere, editfixedsize, Category = KAggregateGeom)
-	TArray<struct FKSphylElem> SphylElems;
-
-	UPROPERTY(EditAnywhere, editfixedsize, Category = KAggregateGeom)
-	TArray<struct FKConvexElem> ConvexElems;
-
-	//class FKConvexGeomRenderInfo* RenderInfo;
-
-	FAggregateGeometry2D()
-//		: RenderInfo(NULL)
-	{}
-
-// 	int32 GetElementCount() const
-// 	{
-// 		return SphereElems.Num() + SphylElems.Num() + BoxElems.Num() + ConvexElems.Num();
-// 	}
-// 
-// 	int32 GetElementCount(int32 Type) const;
-// 
-	void EmptyElements()
-	{
-		BoxElems.Empty();
-		ConvexElems.Empty();
-		SphylElems.Empty();
-		SphereElems.Empty();
-
-//		FreeRenderInfo();
-	}
-// 
-// 
-// 
-// 	void Serialize(const FArchive& Ar);
-// 
-// 	void DrawAggGeom(class FPrimitiveDrawInterface* PDI, const FTransform& Transform, const FColor Color, const FMaterialRenderProxy* MatInst, bool bPerHullColor, bool bDrawSolid);
-// 
-// 	/** Release the RenderInfo (if its there) and safely clean up any resources. Call on the game thread. */
-// 	void FreeRenderInfo();
-// 
-// 	FBox CalcAABB(const FTransform& Transform);
-// 
-// 	/**
-// 	* Calculates a tight box-sphere bounds for the aggregate geometry; this is more expensive than CalcAABB
-// 	* (tight meaning the sphere may be smaller than would be required to encompass the AABB, but all individual components lie within both the box and the sphere)
-// 	*
-// 	* @param Output The output box-sphere bounds calculated for this set of aggregate geometry
-// 	*	@param LocalToWorld Transform
-// 	*/
-// 	void CalcBoxSphereBounds(FBoxSphereBounds& Output, const FTransform& LocalToWorld);
-// 
-// 	/** Returns the volume of this element */
-// 	float GetVolume(const FVector& Scale) const;
-};
+#include "AggregateGeometry2D.h"
 
+#include "BodySetup2D.generated.h"
 
 
 UCLASS(DependsOn=UEngineTypes)
@@ -100,5 +32,4 @@ class PAPER2D_API UBodySetup2D : public UObject
 protected:
 	/** Physics triangle mesh, created from cooked data in CreatePhysicsMeshes */
 	//class PxTriangleMesh* TriMesh;
-
 };
diff --git a/Engine/Plugins/Experimental/Paper2D/Source/Paper2D/Classes/PaperRenderComponent.h b/Engine/Plugins/Experimental/Paper2D/Source/Paper2D/Classes/PaperRenderComponent.h
--- a/Engine/Plugins/Experimental/Paper2D/Source/Paper2D/Classes/PaperRenderComponent.h
+++ b/Engine/Plugins/Experimental/Paper2D/Source/Paper2D/Classes/PaperRenderComponent.h
@@ -2,11 +2,11 @@
 
 #pragma once
 
-#include "BodyInstance2D.h"
+#include "PrimitiveComponent2D.h"
 #include "PaperRenderComponent.generated.h"
 
 UCLASS(DependsOn=UPaperSprite, MinimalAPI, ShowCategories=(Mobility), meta=(BlueprintSpawnableComponent))
-class UPaperRenderComponent : public UPrimitiveComponent
+class UPaperRenderComponent : public UPrimitiveComponent2D
 {
 	GENERATED_UCLASS_BODY()
 
@@ -19,10 +19,6 @@ protected:
 	UPROPERTY(Category=Sprite, EditAnywhere, BlueprintReadOnly)
 	UMaterialInterface* MaterialOverride;
 
-	// Physics scene information for this component, holds a single rigid body with multiple shapes.
-	UPROPERTY(EditAnywhere, BlueprintReadOnly, Category=Collision)//, meta=(ShowOnlyInnerProperties))
-	FBodyInstance2D BodyInstance2D;
-
 	// The color of the sprite (passed to the sprite material as a vertex color)
 	UPROPERTY(BlueprintReadOnly, Interp, Category=Sprite)
 	FLinearColor SpriteColor;
@@ -44,13 +40,10 @@ public:
 
 	// UActorComponent interface
 	virtual void SendRenderDynamicData_Concurrent() OVERRIDE;
-	virtual void CreatePhysicsState() OVERRIDE;
-	virtual void DestroyPhysicsState() OVERRIDE;
 	virtual const UObject* AdditionalStatObject() const OVERRIDE;
 	// End of UActorComponent interface
 
 	// USceneComponent interface
-	virtual void OnUpdateTransform(bool bSkipPhysicsMove) OVERRIDE;
 	virtual bool HasAnySockets() const OVERRIDE;
 	virtual FTransform GetSocketTransform(FName InSocketName, ERelativeTransformSpace TransformSpace = RTS_World) const OVERRIDE;
 	virtual void QuerySupportedSockets(TArray<FComponentSocketDescription>& OutSockets) const OVERRIDE;
@@ -59,20 +52,13 @@ public:
 	// UPrimitiveComponent interface
 	virtual FPrimitiveSceneProxy* CreateSceneProxy() OVERRIDE;
 	virtual FBoxSphereBounds CalcBounds(const FTransform & LocalToWorld) const OVERRIDE;
-	virtual void SetSimulatePhysics(bool bSimulate) OVERRIDE;
 	virtual class UBodySetup* GetBodySetup() OVERRIDE;
 	// End of UPrimitiveComponent interface
 
-	/** Return the BodySetup to use for this PrimitiveComponent (single body case) */
-	virtual class UBodySetup2D* GetBodySetup2D();
+	// UPrimitiveComponent2D interface
+	virtual class UBodySetup2D* GetBodySetup2D() const OVERRIDE;
+	// End of UPrimitiveComponent2D interface
 
 	// Returns the wireframe color to use for this component.
 	FLinearColor GetWireframeColor() const;
-
-protected:
-	//@TODO: Document
-	virtual void CreatePhysicsState2D();
-
-	//@TODO: Document
-	virtual void DestroyPhysicsState2D();
 };
diff --git a/Engine/Plugins/Experimental/Paper2D/Source/Paper2D/Classes/PrimitiveComponent2D.h b/Engine/Plugins/Experimental/Paper2D/Source/Paper2D/Classes/PrimitiveComponent2D.h
new file mode 100644
--- /dev/null
+++ b/Engine/Plugins/Experimental/Paper2D/Source/Paper2D/Classes/PrimitiveComponent2D.h
@@ -0,0 +1,41 @@
+// Copyright 1998-2014 Epic Games, Inc. All Rights Reserved.
+
+#pragma once
+
+#include "BodyInstance2D.h"
+#include "PrimitiveComponent2D.generated.h"
+
+UCLASS(MinimalAPI)
+class UPrimitiveComponent2D : public UPrimitiveComponent
+{
+	GENERATED_UCLASS_BODY()
+
+protected:
+	// Physics scene information for this component, holds a single rigid body with multiple shapes.
+	UPROPERTY(EditAnywhere, BlueprintReadOnly, Category=Collision)//, meta=(ShowOnlyInnerProperties))
+	FBodyInstance2D BodyInstance2D;
+
+public:
+	// UActorComponent interface
+	virtual void CreatePhysicsState() OVERRIDE;
+	virtual void DestroyPhysicsState() OVERRIDE;
+	// End of UActorComponent interface
+
+	// USceneComponent interface
+	virtual void OnUpdateTransform(bool bSkipPhysicsMove) OVERRIDE;
+	// End of USceneComponent interface
+
+	// UPrimitiveComponent interface
+	virtual void SetSimulatePhysics(bool bSimulate) OVERRIDE;
+	// End of UPrimitiveComponent interface
+
+	/** Return the BodySetup to use for this PrimitiveComponent (single body case) */
+	virtual class UBodySetup2D* GetBodySetup2D() const;
+
+protected:
+	//@TODO: Document
+	virtual void CreatePhysicsState2D();
+
+	//@TODO: Document
+	virtual void DestroyPhysicsState2D();
+};
diff --git a/Engine/Plugins/Experimental/Paper2D/Source/Paper2D/Private/PaperRenderComponent.cpp b/Engine/Plugins/Experimental/Paper2D/Source/Paper2D/Private/PaperRenderComponent.cpp
--- a/Engine/Plugins/Experimental/Paper2D/Source/Paper2D/Private/PaperRenderComponent.cpp
+++ b/Engine/Plugins/Experimental/Paper2D/Source/Paper2D/Private/PaperRenderComponent.cpp
@@ -77,58 +77,6 @@ void UPaperRenderComponent::SendRenderDynamicData_Concurrent()
 	}
 }
 
-void UPaperRenderComponent::CreatePhysicsState()
-{
-	Super::CreatePhysicsState();
-	CreatePhysicsState2D();
-}
-
-void UPaperRenderComponent::DestroyPhysicsState()
-{
-	Super::DestroyPhysicsState();
-	DestroyPhysicsState2D();
-}
-
-void UPaperRenderComponent::DestroyPhysicsState2D()
-{
-	// clean up physics engine representation
-	if (BodyInstance2D.IsValidBodyInstance())
-	{
-		BodyInstance2D.TermBody();
-	}
-}
-
-void UPaperRenderComponent::CreatePhysicsState2D()
-{
-	if (!BodyInstance2D.IsValidBodyInstance() && (SourceSprite != NULL) && (SourceSprite->SpriteCollisionDomain == ESpriteCollisionMode::Use2DPhysics))
-	{
-		UBodySetup2D* BodySetup = GetBodySetup2D();
-		//@TODO: if (UBodySetup2D* BodySetup = GetBodySetup2D())
-		{
-			BodyInstance2D.InitBody(BodySetup, SourceSprite, ComponentToWorld, this);
-		}
-	}
-}
-
-UBodySetup2D* UPaperRenderComponent::GetBodySetup2D()
-{
-	return (SourceSprite != nullptr) ? SourceSprite->BodySetup2D : nullptr;
-}
-
-void UPaperRenderComponent::OnUpdateTransform(bool bSkipPhysicsMove)
-{
-	Super::OnUpdateTransform(bSkipPhysicsMove);
-
-	// Always send new transform to physics
-	if (bPhysicsStateCreated && !bSkipPhysicsMove)
-	{
-		// @todo UE4 rather than always pass false, this function should know if it is being teleported or not!
-		SendPhysicsTransform(false);
-
-		BodyInstance2D.SetBodyTransform(ComponentToWorld);
-	}
-}
-
 bool UPaperRenderComponent::HasAnySockets() const
 {
 	if (SourceSprite != NULL)
@@ -178,15 +126,30 @@ void UPaperRenderComponent::QuerySupportedSockets(TArray<FComponentSocketDescrip
 	}
 }
 
-void UPaperRenderComponent::SetSimulatePhysics(bool bSimulate)
+UBodySetup* UPaperRenderComponent::GetBodySetup()
 {
-	Super::SetSimulatePhysics(bSimulate);
-	BodyInstance2D.SetInstanceSimulatePhysics(bSimulate);
+	if (SourceSprite != nullptr)
+	{
+		if (SourceSprite->SpriteCollisionDomain == ESpriteCollisionMode::Use3DPhysics)
+		{
+			return SourceSprite->BodySetup3D;
+		}
+	}
+	
+	return nullptr;
 }
 
-UBodySetup* UPaperRenderComponent::GetBodySetup()
+UBodySetup2D* UPaperRenderComponent::GetBodySetup2D() const
 {
-	return (SourceSprite != NULL) ? SourceSprite->BodySetup3D : NULL;
+	if (SourceSprite != nullptr)
+	{
+		if (SourceSprite->SpriteCollisionDomain == ESpriteCollisionMode::Use2DPhysics)
+		{
+			return SourceSprite->BodySetup2D;
+		}
+	}
+
+	return nullptr;
 }
 
 bool UPaperRenderComponent::SetSprite(class UPaperSprite* NewSprite)
diff --git a/Engine/Plugins/Experimental/Paper2D/Source/Paper2D/Private/PaperSprite.cpp b/Engine/Plugins/Experimental/Paper2D/Source/Paper2D/Private/PaperSprite.cpp
--- a/Engine/Plugins/Experimental/Paper2D/Source/Paper2D/Private/PaperSprite.cpp
+++ b/Engine/Plugins/Experimental/Paper2D/Source/Paper2D/Private/PaperSprite.cpp
@@ -508,6 +508,10 @@ void UPaperSprite::BuildCustomCollisionData()
 		Point = ConvertTextureSpaceToPivotSpace(Point);
 	}
 
+	//@TODO: Observe if the custom data is really a hand-edited bounding box, and generate box geom instead of convex geom!
+	//@TODO: Use this guy instead: DecomposeMeshToHulls
+	//@TODO: Merge triangles that are convex together!
+
 	// Bake it to the runtime structure
 	switch (SpriteCollisionDomain)
 	{
@@ -516,11 +520,8 @@ void UPaperSprite::BuildCustomCollisionData()
 			checkSlow(BodySetup3D);
 			BodySetup3D->AggGeom.EmptyElements();
 
-			//@TODO: Observe if the custom data is really a hand-edited bounding box, and generate box geom instead of convex geom!
-
 			const FVector HalfThicknessVector = PaperAxisZ * 0.5f * CollisionThickness;
-			//@TODO: Use this guy instead: DecomposeMeshToHulls
-			//@TODO: Merge triangles that are convex together!
+
 			int32 RunningIndex = 0;
 			for (int32 TriIndex = 0; TriIndex < CollisionData.Num() / 3; ++TriIndex)
 			{
@@ -547,29 +548,21 @@ void UPaperSprite::BuildCustomCollisionData()
 			checkSlow(BodySetup2D);
 			BodySetup2D->AggGeom2D.EmptyElements();
 
-			//@TODO: Observe if the custom data is really a hand-edited bounding box, and generate box geom instead of convex geom!
-
-			//@TODO: Use this guy instead: DecomposeMeshToHulls
-			//@TODO: Merge triangles that are convex together!
 			int32 RunningIndex = 0;
 			for (int32 TriIndex = 0; TriIndex < CollisionData.Num() / 3; ++TriIndex)
 			{
-				FKConvexElem& ConvexTri = *new (BodySetup2D->AggGeom2D.ConvexElems) FKConvexElem();
+				FConvexElement2D& ConvexTri = *new (BodySetup2D->AggGeom2D.ConvexElements) FConvexElement2D();
 				ConvexTri.VertexData.Empty(3);
 				for (int32 Index = 0; Index < 3; ++Index)
 				{
 					const FVector2D& Pos2D = CollisionData[RunningIndex++];
-
-					const FVector Pos3D = (PaperAxisX * Pos2D.X) + (PaperAxisY * Pos2D.X);
-
-					new (ConvexTri.VertexData) FVector(Pos3D);
+					new (ConvexTri.VertexData) FVector2D(Pos2D);
 				}
-				ConvexTri.UpdateElemBox();
 			}
 
 			BodySetup2D->InvalidatePhysicsData();
 			BodySetup2D->CreatePhysicsMeshes();
-	}
+		}
 		break;
 	default:
 		check(false);
@@ -609,7 +602,6 @@ void UPaperSprite::BuildBoundingBoxCollisionData(bool bUseTightBounds)
 		break;
 	case ESpriteCollisionMode::Use2DPhysics:
 		{
-			//@TODO:!!!
 			checkSlow(BodySetup2D);
 			BodySetup2D->AggGeom2D.EmptyElements();
 
@@ -620,10 +612,11 @@ void UPaperSprite::BuildBoundingBoxCollisionData(bool bUseTightBounds)
 			const FVector2D CenterInPivotSpace = ConvertTextureSpaceToPivotSpace(CenterInTextureSpace);
 
 			// Create a new box primitive
-			const FVector BoxSize3D = (PaperAxisX * BoxSize2D.X) + (PaperAxisY * BoxSize2D.Y) + (PaperAxisZ * CollisionThickness);
-
-			FKBoxElem& Box = *new (BodySetup2D->AggGeom2D.BoxElems) FKBoxElem(FMath::Abs(BoxSize3D.X), FMath::Abs(BoxSize3D.Y), FMath::Abs(BoxSize3D.Z));
-			Box.Center = (PaperAxisX * CenterInPivotSpace.X) + (PaperAxisY * CenterInPivotSpace.Y);
+			FBoxElement2D& Box = *new (BodySetup2D->AggGeom2D.BoxElements) FBoxElement2D();
+			Box.Width = FMath::Abs(BoxSize2D.X);
+			Box.Height = FMath::Abs(BoxSize2D.Y);
+			Box.Center.X = CenterInPivotSpace.X;
+			Box.Center.Y = CenterInPivotSpace.Y;
 
 			BodySetup2D->InvalidatePhysicsData();
 			BodySetup2D->CreatePhysicsMeshes();
diff --git a/Engine/Plugins/Experimental/Paper2D/Source/Paper2D/Private/Physics/AggregateGeometry2D.cpp b/Engine/Plugins/Experimental/Paper2D/Source/Paper2D/Private/Physics/AggregateGeometry2D.cpp
new file mode 100644
--- /dev/null
+++ b/Engine/Plugins/Experimental/Paper2D/Source/Paper2D/Private/Physics/AggregateGeometry2D.cpp
@@ -0,0 +1,140 @@
+// Copyright 1998-2014 Epic Games, Inc. All Rights Reserved.
+
+#include "Paper2DPrivatePCH.h"
+
+static const int32 DrawCollisionSides = 16;
+
+//////////////////////////////////////////////////////////////////////////
+// FCircleElement2D
+
+float FCircleElement2D::GetArea(const FVector& Scale) const
+{
+	return PI * FMath::Square(Radius * Scale.Size());
+}
+
+// Note: ElemTM is assumed to have no scaling in it!
+void FCircleElement2D::DrawElemWire(FPrimitiveDrawInterface* PDI, const FTransform& ElemTM, float Scale, const FColor Color)
+{
+	const FVector ElemCenter = ElemTM.GetLocation();
+	const FVector X = ElemTM.TransformVector(PaperAxisX);
+	const FVector Y = ElemTM.TransformVector(PaperAxisY);
+
+	DrawCircle(PDI, ElemCenter, X, Y, Color, Scale * Radius, DrawCollisionSides, SDPG_World);
+}
+
+void FCircleElement2D::DrawElemSolid(class FPrimitiveDrawInterface* PDI, const FTransform& ElemTM, float Scale, const FMaterialRenderProxy* MaterialRenderProxy)
+{
+	const FVector ElemCenter = ElemTM.GetLocation();
+	const FVector X = ElemTM.TransformVector(PaperAxisX);
+	const FVector Y = ElemTM.TransformVector(PaperAxisY);
+
+	DrawDisc(PDI, ElemCenter, X, Y, FColor::White, Scale * Radius, DrawCollisionSides, MaterialRenderProxy, SDPG_World);
+}
+
+//////////////////////////////////////////////////////////////////////////
+// FBoxElement2D
+
+float FBoxElement2D::GetArea(const FVector& Scale) const
+{
+	const float Scale1D = Scale.Size();
+	return (Width * Scale1D) * (Height * Scale1D);
+}
+
+// Note: ElemTM is assumed to have no scaling in it!
+void FBoxElement2D::DrawElemWire(FPrimitiveDrawInterface* PDI, const FTransform& ElemTM, float Scale, const FColor Color)
+{
+	//@TODO: PAPER2D: Implement this method
+
+// 	FVector	B[2], P, Q, Radii;
+// 
+// 	// X,Y,Z member variables are LENGTH not RADIUS
+// 	Radii.X = Scale*0.5f*X;
+// 	Radii.Y = Scale*0.5f*Y;
+// 	Radii.Z = Scale*0.5f*Z;
+// 
+// 	B[0] = Radii; // max
+// 	B[1] = -1.0f * Radii; // min
+// 
+// 	for (int32 i = 0; i < 2; i++)
+// 	{
+// 		for (int32 j = 0; j < 2; j++)
+// 		{
+// 			P.X = B[i].X; Q.X = B[i].X;
+// 			P.Y = B[j].Y; Q.Y = B[j].Y;
+// 			P.Z = B[0].Z; Q.Z = B[1].Z;
+// 			PDI->DrawLine(ElemTM.TransformPosition(P), ElemTM.TransformPosition(Q), Color, SDPG_World);
+// 
+// 			P.Y = B[i].Y; Q.Y = B[i].Y;
+// 			P.Z = B[j].Z; Q.Z = B[j].Z;
+// 			P.X = B[0].X; Q.X = B[1].X;
+// 			PDI->DrawLine(ElemTM.TransformPosition(P), ElemTM.TransformPosition(Q), Color, SDPG_World);
+// 
+// 			P.Z = B[i].Z; Q.Z = B[i].Z;
+// 			P.X = B[j].X; Q.X = B[j].X;
+// 			P.Y = B[0].Y; Q.Y = B[1].Y;
+// 			PDI->DrawLine(ElemTM.TransformPosition(P), ElemTM.TransformPosition(Q), Color, SDPG_World);
+// 		}
+// 	}
+}
+
+void FBoxElement2D::DrawElemSolid(class FPrimitiveDrawInterface* PDI, const FTransform& ElemTM, float Scale, const FMaterialRenderProxy* MaterialRenderProxy)
+{
+	//@TODO: PAPER2D: Implement this method
+	// 	DrawBox(PDI, ElemTM.ToMatrixWithScale(), 0.5f * FVector(X, Y, Z), MaterialRenderProxy, SDPG_World);
+}
+
+//////////////////////////////////////////////////////////////////////////
+// FConvexElement2D
+
+float FConvexElement2D::GetArea(const FVector& Scale) const
+{
+	//@TODO: PAPER2D: Implement this method properly
+	return 1.0f * Scale.Size();
+}
+
+void FConvexElement2D::DrawElemWire(class FPrimitiveDrawInterface* PDI, const FTransform& ElemTM, const FColor Color)
+{
+	//@TODO: PAPER2D: Implement this method
+}
+
+void FConvexElement2D::DrawElemSolid(class FPrimitiveDrawInterface* PDI, const FTransform& ElemTM, float Scale, const FMaterialRenderProxy* MaterialRenderProxy)
+{
+	//@TODO: PAPER2D: Implement this method
+}
+
+//////////////////////////////////////////////////////////////////////////
+// FAggregateGeometry2D
+
+void FAggregateGeometry2D::EmptyElements()
+{
+	CircleElements.Empty();
+	BoxElements.Empty();
+	ConvexElements.Empty();
+}
+
+void FAggregateGeometry2D::DrawAggGeom(FPrimitiveDrawInterface* PDI, const FTransform& Transform, const FColor Color, const FMaterialRenderProxy* MatInst, bool bPerHullColor, bool bDrawSolid) const
+{
+	//@TODO: PAPER2D: Implement this method
+}
+
+float FAggregateGeometry2D::GetArea(const FVector& Scale) const
+{
+	float SumArea = 0.0f;
+
+	for (const FCircleElement2D& Circle : CircleElements)
+	{
+		SumArea += Circle.GetArea(Scale);
+	}
+
+	for (const FBoxElement2D& Box : BoxElements)
+	{
+		SumArea += Box.GetArea(Scale);
+	}
+
+	for (const FConvexElement2D& Convex : ConvexElements)
+	{
+		SumArea += Convex.GetArea(Scale);
+	}
+
+	return SumArea;
+}
diff --git a/Engine/Plugins/Experimental/Paper2D/Source/Paper2D/Private/Physics/BodyInstance2D.cpp b/Engine/Plugins/Experimental/Paper2D/Source/Paper2D/Private/Physics/BodyInstance2D.cpp
--- a/Engine/Plugins/Experimental/Paper2D/Source/Paper2D/Private/Physics/BodyInstance2D.cpp
+++ b/Engine/Plugins/Experimental/Paper2D/Source/Paper2D/Private/Physics/BodyInstance2D.cpp
@@ -15,12 +15,13 @@ FBodyInstance2D::FBodyInstance2D()
 }
 
 //@TODO: Everything about this block of code...
-void FBodyInstance2D::InitBody(class UBodySetup2D* Setup, const UPaperSprite* SpriteSetupHack, const FTransform& Transform, class UPrimitiveComponent* PrimComp)
+void FBodyInstance2D::InitBody(class UBodySetup2D* Setup, const FTransform& Transform, class UPrimitiveComponent* PrimComp)
 {
 #if WITH_BOX2D
 	if (b2World* World = FPhysicsIntegration2D::FindAssociatedWorld(PrimComp->GetWorld()))
 	{
 		const FVector Scale3D = Transform.GetScale3D();
+		const b2Vec2 Scale2D = FPhysicsIntegration2D::ConvertUnrealVectorToBox(Scale3D);
 
 		if (Setup != nullptr)
 		{
@@ -32,17 +33,30 @@ void FBodyInstance2D::InitBody(class UBodySetup2D* Setup, const UPaperSprite* Sp
 			BodyInstancePtr = World->CreateBody(&BodyDefinition);
 			BodyInstancePtr->SetUserData(this);
 
-// 			const FVector2D PivotPos = FVector2D::ZeroVector;
-// 			const FVector2D OffsetPos = FVector2D(SpriteSetupHack->GetSourceSize().X*-0.5f, SpriteSetupHack->GetSourceSize().Y*-0.5f) + PivotPos;
+			// Circles
+			for (const FCircleElement2D& Circle : Setup->AggGeom2D.CircleElements)
+			{
+				b2CircleShape CircleShape;
+				CircleShape.m_radius = Circle.Radius * Scale3D.Size() / UnrealUnitsPerMeter;
+				CircleShape.m_p.x = Circle.Center.X;
+				CircleShape.m_p.y = Circle.Center.Y;
+
+				b2FixtureDef FixtureDef;
+				FixtureDef.shape = &CircleShape;
+				FixtureDef.density = 1.0f; //@TODO:
+				FixtureDef.friction = 0.3f;//@TODO:
+
+				BodyInstancePtr->CreateFixture(&FixtureDef);
+			}
 
 			// Boxes
-			for (const FKBoxElem& Box : Setup->AggGeom2D.BoxElems)
+			for (const FBoxElement2D& Box : Setup->AggGeom2D.BoxElements)
 			{
-				const b2Vec2 HalfBoxSize = FPhysicsIntegration2D::ConvertUnrealVectorToBox(FVector(Box.X, Box.Y, Box.Z) * 0.5f  * Scale3D);
-				const b2Vec2 BoxCenter = FPhysicsIntegration2D::ConvertUnrealVectorToBox(Box.Center * Scale3D);
+				const b2Vec2 HalfBoxSize(Box.Width * 0.5f * Scale2D.x, Box.Height * 0.5f * Scale2D.y);
+				const b2Vec2 BoxCenter(Box.Center.X * Scale2D.x, Box.Center.Y * Scale2D.y);
 
 				b2PolygonShape DynamicBox;
-				DynamicBox.SetAsBox(HalfBoxSize.x, HalfBoxSize.y, BoxCenter, 0.0f);
+				DynamicBox.SetAsBox(HalfBoxSize.x, HalfBoxSize.y, BoxCenter, FMath::DegreesToRadians(Box.Angle));
 
 				b2FixtureDef FixtureDef;
 				FixtureDef.shape = &DynamicBox;
@@ -53,7 +67,7 @@ void FBodyInstance2D::InitBody(class UBodySetup2D* Setup, const UPaperSprite* Sp
 			}
 
 			// Convex hulls
-			for (const FKConvexElem& Convex : Setup->AggGeom2D.ConvexElems)
+			for (const FConvexElement2D& Convex : Setup->AggGeom2D.ConvexElements)
 			{
 				const int32 NumVerts = Convex.VertexData.Num();
 
@@ -63,8 +77,8 @@ void FBodyInstance2D::InitBody(class UBodySetup2D* Setup, const UPaperSprite* Sp
 
 					for (int32 VertexIndex = 0; VertexIndex < Convex.VertexData.Num(); ++VertexIndex)
 					{
-						const FVector SourceVert = Convex.VertexData[VertexIndex];
-						new (Verts) b2Vec2(FPhysicsIntegration2D::ConvertUnrealVectorToBox(SourceVert * Scale3D));
+						const FVector2D SourceVert = Convex.VertexData[VertexIndex];
+						new (Verts) b2Vec2(SourceVert.X * Scale2D.x, SourceVert.Y * Scale2D.y);
 					}
 
 					b2PolygonShape ConvexPoly;
@@ -83,62 +97,6 @@ void FBodyInstance2D::InitBody(class UBodySetup2D* Setup, const UPaperSprite* Sp
 				}
 			}
 
-#if 0
-		if ((SpriteSetupHack->CollisionGeometry.Num() > 0) && ((SpriteSetupHack->CollisionGeometry.Num() % 3) == 0))
-		{
-			OwnerComponentPtr = PrimComp;
-
-			b2BodyDef BodyDefinition;
-			BodyDefinition.type = bSimulatePhysics ? b2_dynamicBody : b2_kinematicBody;
-
-			BodyInstancePtr = World->CreateBody(&BodyDefinition);
-			BodyInstancePtr->SetUserData(this);
-
-			const FVector2D PivotPos = FVector2D::ZeroVector;
-			const FVector2D OffsetPos = FVector2D(SpriteSetupHack->GetSourceSize().X*-0.5f, SpriteSetupHack->GetSourceSize().Y*-0.5f) + PivotPos;
-
-			for (int32 TriangleIndex = 0; TriangleIndex < SpriteSetupHack->CollisionGeometry.Num() / 3; ++TriangleIndex)
-			{
-				b2Vec2 Verts[3];
-				for (int32 VertexIndex = 0; VertexIndex < 3; ++VertexIndex)
-				{
-					const FVector2D SourceVert = SpriteSetupHack->CollisionGeometry[(TriangleIndex * 3) + VertexIndex] + OffsetPos;
-					Verts[VertexIndex] = FPhysicsIntegration2D::ConvertUnrealVectorToBox(FVector(SourceVert.X, 0.0f, SourceVert.Y));
-				}
-				
-				b2PolygonShape Triangle;
-				Triangle.Set(Verts, 3);
-
-				b2FixtureDef FixtureDef;
-				FixtureDef.shape = &Triangle;
-				FixtureDef.density = 1.0f; //@TODO:
-				FixtureDef.friction = 0.3f;//@TODO:
-
-				BodyInstancePtr->CreateFixture(&FixtureDef);
-			}
-#endif
-
-
-
-
-#if 0
-			const b2Vec2 HalfBoxSize = FPhysicsIntegration2D::ConvertUnrealVectorToBox(FVector(SpriteSetupHack->BoundingBoxSize.X * 0.5f, 0.0f, SpriteSetupHack->BoundingBoxSize.Y * 0.5f));
-			const b2Vec2 BoxCenter = FPhysicsIntegration2D::ConvertUnrealVectorToBox(FVector(
-				SpriteSetupHack->BoundingBoxPosition.X,
-				0.0f,
-				SpriteSetupHack->BoundingBoxPosition.Y));
-
-			b2PolygonShape DynamicBox;
-			DynamicBox.SetAsBox(HalfBoxSize.x, HalfBoxSize.y, BoxCenter, 0.0f);
-
-			b2FixtureDef FixtureDef;
-			FixtureDef.shape = &DynamicBox;
-			FixtureDef.density = 1.0f; //@TODO:
-			FixtureDef.friction = 0.3f;//@TODO:
-
-			BodyInstancePtr->CreateFixture(&FixtureDef);
-#endif
-
 			SetBodyTransform(Transform);
 		}
 	}
diff --git a/Engine/Plugins/Experimental/Paper2D/Source/Paper2D/Private/Physics/PrimitiveComponent2D.cpp b/Engine/Plugins/Experimental/Paper2D/Source/Paper2D/Private/Physics/PrimitiveComponent2D.cpp
new file mode 100644
--- /dev/null
+++ b/Engine/Plugins/Experimental/Paper2D/Source/Paper2D/Private/Physics/PrimitiveComponent2D.cpp
@@ -0,0 +1,72 @@
+// Copyright 1998-2014 Epic Games, Inc. All Rights Reserved.
+
+#include "Paper2DPrivatePCH.h"
+#include "PrimitiveComponent2D.h"
+
+//////////////////////////////////////////////////////////////////////////
+// UPrimitiveComponent2D
+
+UPrimitiveComponent2D::UPrimitiveComponent2D(const FPostConstructInitializeProperties& PCIP)
+	: Super(PCIP)
+{
+	BodyInstance.SetCollisionEnabled(ECollisionEnabled::NoCollision);
+
+	SetCollisionProfileName(UCollisionProfile::BlockAllDynamic_ProfileName);
+}
+
+void UPrimitiveComponent2D::CreatePhysicsState()
+{
+	Super::CreatePhysicsState();
+	CreatePhysicsState2D();
+}
+
+void UPrimitiveComponent2D::DestroyPhysicsState()
+{
+	Super::DestroyPhysicsState();
+	DestroyPhysicsState2D();
+}
+
+void UPrimitiveComponent2D::CreatePhysicsState2D()
+{
+	if (!BodyInstance2D.IsValidBodyInstance())
+	{
+		if (UBodySetup2D* BodySetup = GetBodySetup2D())
+		{
+			BodyInstance2D.InitBody(BodySetup, ComponentToWorld, this);
+		}
+	}
+}
+
+void UPrimitiveComponent2D::DestroyPhysicsState2D()
+{
+	// clean up physics engine representation
+	if (BodyInstance2D.IsValidBodyInstance())
+	{
+		BodyInstance2D.TermBody();
+	}
+}
+
+void UPrimitiveComponent2D::OnUpdateTransform(bool bSkipPhysicsMove)
+{
+	Super::OnUpdateTransform(bSkipPhysicsMove);
+
+	// Always send new transform to physics
+	if (bPhysicsStateCreated && !bSkipPhysicsMove)
+	{
+		// @todo UE4 rather than always pass false, this function should know if it is being teleported or not!
+		SendPhysicsTransform(false);
+
+		BodyInstance2D.SetBodyTransform(ComponentToWorld);
+	}
+}
+
+void UPrimitiveComponent2D::SetSimulatePhysics(bool bSimulate)
+{
+	Super::SetSimulatePhysics(bSimulate);
+	BodyInstance2D.SetInstanceSimulatePhysics(bSimulate);
+}
+
+class UBodySetup2D* UPrimitiveComponent2D::GetBodySetup2D() const
+{
+	return nullptr;
+}
