Paper2D source rect editor WIP

diff --git a/Engine/Plugins/Experimental/Paper2D/Source/Paper2D/Classes/PaperSprite.h b/Engine/Plugins/Experimental/Paper2D/Source/Paper2D/Classes/PaperSprite.h
--- a/Engine/Plugins/Experimental/Paper2D/Source/Paper2D/Classes/PaperSprite.h
+++ b/Engine/Plugins/Experimental/Paper2D/Source/Paper2D/Classes/PaperSprite.h
@@ -195,6 +195,7 @@ public:
 	friend class FSpriteDetailsCustomization;
 	friend class FSpriteSelectedVertex;
 	friend class FSpriteSelectedEdge;
+	friend class FSpriteSelectedSourceRegion;
 	friend struct FPaperAtlasGenerator;
 };
 
diff --git a/Engine/Plugins/Experimental/Paper2D/Source/Paper2DEditor/Private/PaperStyle.cpp b/Engine/Plugins/Experimental/Paper2D/Source/Paper2DEditor/Private/PaperStyle.cpp
--- a/Engine/Plugins/Experimental/Paper2D/Source/Paper2DEditor/Private/PaperStyle.cpp
+++ b/Engine/Plugins/Experimental/Paper2D/Source/Paper2DEditor/Private/PaperStyle.cpp
@@ -71,6 +71,9 @@ TSharedRef<FSlateStyleSet> FPaperStyle::Create()
 		Style.Set("SpriteEditor.EnterCollisionEditMode", new FSlateImageBrush(TEXT("texture://Paper2D/SlateBrushes/EditCollisionGeom.EditCollisionGeom"), Icon40x40));
 		Style.Set("SpriteEditor.EnterCollisionEditMode.Small", new FSlateImageBrush(TEXT("texture://Paper2D/SlateBrushes/EditCollisionGeom.EditCollisionGeom"), Icon20x20));
 
+		Style.Set("SpriteEditor.EnterSourceRegionEditMode", new FSlateImageBrush(TEXT("texture://Paper2D/SlateBrushes/EditCollisionGeom.EditSourceRegion"), Icon40x40));
+		Style.Set("SpriteEditor.EnterSourceRegionEditMode.Small", new FSlateImageBrush(TEXT("texture://Paper2D/SlateBrushes/EditCollisionGeom.EditSourceRegion"), Icon20x20));
+
 		Style.Set("SpriteEditor.EnterRenderingEditMode", new FSlateImageBrush(TEXT("texture://Paper2D/SlateBrushes/EditRenderGeomButton.EditRenderGeomButton"), Icon40x40));
 		Style.Set("SpriteEditor.EnterRenderingEditMode.Small", new FSlateImageBrush(TEXT("texture://Paper2D/SlateBrushes/EditRenderGeomButton.EditRenderGeomButton"), Icon20x20));
 
diff --git a/Engine/Plugins/Experimental/Paper2D/Source/Paper2DEditor/Private/SpriteEditor/SpriteEditor.cpp b/Engine/Plugins/Experimental/Paper2D/Source/Paper2DEditor/Private/SpriteEditor/SpriteEditor.cpp
--- a/Engine/Plugins/Experimental/Paper2D/Source/Paper2DEditor/Private/SpriteEditor/SpriteEditor.cpp
+++ b/Engine/Plugins/Experimental/Paper2D/Source/Paper2DEditor/Private/SpriteEditor/SpriteEditor.cpp
@@ -132,9 +132,14 @@ void SSpriteEditorViewport::BindCommands()
 	// Editing modes
 	CommandList->MapAction(
 		Commands.EnterViewMode,
-		FExecuteAction::CreateSP( EditorViewportClientRef, &FSpriteEditorViewportClient::EnterViewMode ),
+		FExecuteAction::CreateSP(EditorViewportClientRef, &FSpriteEditorViewportClient::EnterViewMode),
 		FCanExecuteAction(),
-		FIsActionChecked::CreateSP( EditorViewportClientRef, &FSpriteEditorViewportClient::IsInViewMode ) );
+		FIsActionChecked::CreateSP(EditorViewportClientRef, &FSpriteEditorViewportClient::IsInViewMode));
+	CommandList->MapAction(
+		Commands.EnterSourceRegionEditMode,
+		FExecuteAction::CreateSP(EditorViewportClientRef, &FSpriteEditorViewportClient::EnterSourceRegionEditMode),
+		FCanExecuteAction(),
+		FIsActionChecked::CreateSP(EditorViewportClientRef, &FSpriteEditorViewportClient::IsInSourceRegionEditMode));
 	CommandList->MapAction(
 		Commands.EnterCollisionEditMode,
 		FExecuteAction::CreateSP( EditorViewportClientRef, &FSpriteEditorViewportClient::EnterCollisionEditMode ),
@@ -466,6 +471,7 @@ void FSpriteEditor::ExtendToolbar()
 			ToolbarBuilder.BeginSection("Mode");
 			{
 				ToolbarBuilder.AddToolBarButton(FSpriteEditorCommands::Get().EnterViewMode);
+				ToolbarBuilder.AddToolBarButton(FSpriteEditorCommands::Get().EnterSourceRegionEditMode);
 				ToolbarBuilder.AddToolBarButton(FSpriteEditorCommands::Get().EnterCollisionEditMode);
 				ToolbarBuilder.AddToolBarButton(FSpriteEditorCommands::Get().EnterRenderingEditMode);
 				ToolbarBuilder.AddToolBarButton(FSpriteEditorCommands::Get().EnterAddSpriteMode);
diff --git a/Engine/Plugins/Experimental/Paper2D/Source/Paper2DEditor/Private/SpriteEditor/SpriteEditorCommands.cpp b/Engine/Plugins/Experimental/Paper2D/Source/Paper2DEditor/Private/SpriteEditor/SpriteEditorCommands.cpp
--- a/Engine/Plugins/Experimental/Paper2D/Source/Paper2DEditor/Private/SpriteEditor/SpriteEditorCommands.cpp
+++ b/Engine/Plugins/Experimental/Paper2D/Source/Paper2DEditor/Private/SpriteEditor/SpriteEditorCommands.cpp
@@ -21,6 +21,7 @@ void FSpriteEditorCommands::RegisterCommands()
 
 	// Editing modes
 	UI_COMMAND(EnterViewMode, "View", "View the sprite.", EUserInterfaceActionType::ToggleButton, FInputGesture());
+	UI_COMMAND(EnterSourceRegionEditMode, "Edit Source Region", "Edit the sprite source region.", EUserInterfaceActionType::ToggleButton, FInputGesture());
 	UI_COMMAND(EnterCollisionEditMode, "Edit Collision", "Edit the collision geometry.", EUserInterfaceActionType::ToggleButton, FInputGesture());
 	UI_COMMAND(EnterRenderingEditMode, "Edit RenderGeom", "Edit the rendering geometry (useful to reduce overdraw).", EUserInterfaceActionType::ToggleButton, FInputGesture());
 	UI_COMMAND(EnterAddSpriteMode, "Add Sprite", "Create new sprites quickly.", EUserInterfaceActionType::ToggleButton, FInputGesture());
diff --git a/Engine/Plugins/Experimental/Paper2D/Source/Paper2DEditor/Private/SpriteEditor/SpriteEditorCommands.h b/Engine/Plugins/Experimental/Paper2D/Source/Paper2DEditor/Private/SpriteEditor/SpriteEditorCommands.h
--- a/Engine/Plugins/Experimental/Paper2D/Source/Paper2DEditor/Private/SpriteEditor/SpriteEditorCommands.h
+++ b/Engine/Plugins/Experimental/Paper2D/Source/Paper2DEditor/Private/SpriteEditor/SpriteEditorCommands.h
@@ -33,6 +33,7 @@ public:
 
 	// Editing modes
 	TSharedPtr<FUICommandInfo> EnterViewMode;
+	TSharedPtr<FUICommandInfo> EnterSourceRegionEditMode;
 	TSharedPtr<FUICommandInfo> EnterCollisionEditMode;
 	TSharedPtr<FUICommandInfo> EnterRenderingEditMode;
 	TSharedPtr<FUICommandInfo> EnterAddSpriteMode;
diff --git a/Engine/Plugins/Experimental/Paper2D/Source/Paper2DEditor/Private/SpriteEditor/SpriteEditorSelections.h b/Engine/Plugins/Experimental/Paper2D/Source/Paper2DEditor/Private/SpriteEditor/SpriteEditorSelections.h
--- a/Engine/Plugins/Experimental/Paper2D/Source/Paper2DEditor/Private/SpriteEditor/SpriteEditorSelections.h
+++ b/Engine/Plugins/Experimental/Paper2D/Source/Paper2DEditor/Private/SpriteEditor/SpriteEditorSelections.h
@@ -12,6 +12,7 @@ public:
 	static const FName Edge;
 	static const FName Pivot;
 	static const FName Socket;
+	static const FName SourceRegion;
 private:
 	FSelectionTypes() {}
 };
@@ -77,6 +78,130 @@ inline bool operator==(const FSelectedItem& V1, const FSelectedItem& V2)
 	return V1.Equals(V2);
 }
 
+//////////////////////////////////////////////////////////////////////////
+// FSpriteSelectedSourceRegion
+
+class FSpriteSelectedSourceRegion : public FSelectedItem
+{
+public:
+	int32 VertexIndex;
+	TWeakObjectPtr<UPaperSprite> SpritePtr;
+
+public:
+	FSpriteSelectedSourceRegion()
+		: FSelectedItem(FSelectionTypes::SourceRegion)
+		, VertexIndex(0)
+	{
+	}
+
+	virtual uint32 GetTypeHash() const OVERRIDE
+	{
+		return VertexIndex;
+	}
+
+	virtual bool Equals(const FSelectedItem& OtherItem) const OVERRIDE
+	{
+		if (OtherItem.IsA(FSelectionTypes::SourceRegion))
+		{
+			const FSpriteSelectedSourceRegion& V1 = *this;
+			const FSpriteSelectedSourceRegion& V2 = *(FSpriteSelectedSourceRegion*)(&OtherItem);
+
+			return (V1.VertexIndex == V2.VertexIndex) && (V1.SpritePtr == V2.SpritePtr);
+		}
+		else
+		{
+			return false;
+		}
+	}
+
+	virtual void ApplyDeltaIndexed(const FVector2D& Delta, int32 TargetVertexIndex)
+	{
+		if (UPaperSprite* Sprite = SpritePtr.Get())
+		{
+			UTexture2D* SourceTexture = Sprite->GetSourceTexture();
+			const FVector2D SourceDims = (SourceTexture != NULL) ? FVector2D(SourceTexture->GetSurfaceWidth(), SourceTexture->GetSurfaceHeight()) : FVector2D::ZeroVector;
+			float AllowedDelta = 0;
+
+			switch (VertexIndex) 
+			{
+			case 0:
+				AllowedDelta = FMath::Clamp(Delta.Y, -Sprite->SourceUV.Y, Sprite->SourceDimension.Y - 1);
+				Sprite->SourceUV.Y += AllowedDelta;
+				Sprite->SourceDimension.Y -= AllowedDelta;
+				break;
+			case 1:
+				Sprite->SourceDimension.X = FMath::Clamp(Sprite->SourceDimension.X + Delta.X, 1.0f, SourceDims.X - Sprite->SourceUV.X);
+				break;
+			case 2:
+				Sprite->SourceDimension.Y = FMath::Clamp(Sprite->SourceDimension.Y + Delta.Y, 1.0f, SourceDims.Y - Sprite->SourceUV.Y);
+				break;
+			case 3:
+				AllowedDelta = FMath::Clamp(Delta.X, -Sprite->SourceUV.X, Sprite->SourceDimension.X - 1);
+				Sprite->SourceUV.X += AllowedDelta;
+				Sprite->SourceDimension.X -= AllowedDelta;
+				break;
+			default:
+				break;
+			}
+		}
+	}
+
+	FVector GetWorldPosIndexed(int32 TargetVertexIndex) const
+	{
+		FVector Result = FVector::ZeroVector;
+
+		if (UPaperSprite* Sprite = SpritePtr.Get())
+		{
+			UTexture2D* SourceTexture = Sprite->GetSourceTexture();
+			const FVector2D SourceDims = (SourceTexture != NULL) ? FVector2D(SourceTexture->GetSurfaceWidth(), SourceTexture->GetSurfaceHeight()) : FVector2D::ZeroVector;
+			FVector2D BoundsVertex = Sprite->SourceUV;
+			switch (VertexIndex)
+			{
+			case 0:
+				BoundsVertex.X += Sprite->SourceDimension.X * 0.5f;
+				break;
+			case 1:
+				BoundsVertex.X += Sprite->SourceDimension.X;
+				BoundsVertex.Y += Sprite->SourceDimension.Y * 0.5f;
+				break;
+			case 2:
+				BoundsVertex.X += Sprite->SourceDimension.X * 0.5f;
+				BoundsVertex.Y += Sprite->SourceDimension.Y;
+				break;
+			case 3:
+				BoundsVertex.Y += Sprite->SourceDimension.Y * 0.5f;
+				break;
+			default:
+				break;
+			}
+			Result = BoundsVertex.X * PaperAxisX + (SourceDims.Y - BoundsVertex.Y) * PaperAxisY;
+		}
+
+		return Result;
+	}
+
+	virtual void ApplyDelta(const FVector2D& Delta) OVERRIDE
+	{
+		ApplyDeltaIndexed(Delta, VertexIndex);
+	}
+
+	FVector GetWorldPos() const OVERRIDE
+	{
+		return GetWorldPosIndexed(VertexIndex);
+	}
+
+	virtual void Delete() OVERRIDE
+	{
+		// Cant delete anything on a bounds object
+	}
+
+		virtual void SplitEdge() OVERRIDE
+	{
+		// Nonsense operation on a vertex, do nothing
+	}
+};
+
+
 //////////////////////////////////////////////////////////////////////////
 // FSpriteSelectedVertex
 
diff --git a/Engine/Plugins/Experimental/Paper2D/Source/Paper2DEditor/Private/SpriteEditor/SpriteEditorViewportClient.cpp b/Engine/Plugins/Experimental/Paper2D/Source/Paper2DEditor/Private/SpriteEditor/SpriteEditorViewportClient.cpp
--- a/Engine/Plugins/Experimental/Paper2D/Source/Paper2DEditor/Private/SpriteEditor/SpriteEditorViewportClient.cpp
+++ b/Engine/Plugins/Experimental/Paper2D/Source/Paper2DEditor/Private/SpriteEditor/SpriteEditorViewportClient.cpp
@@ -16,6 +16,7 @@ const FName FSelectionTypes::Vertex(TEXT("Vertex"));
 const FName FSelectionTypes::Edge(TEXT("Edge"));
 const FName FSelectionTypes::Pivot(TEXT("Pivot"));
 const FName FSelectionTypes::Socket(TEXT("Socket"));
+const FName FSelectionTypes::SourceRegion(TEXT("SourceRegion"));
 
 //////////////////////////////////////////////////////////////////////////
 // HSpriteSelectableObjectHitProxy
@@ -155,6 +156,49 @@ void FSpriteEditorViewportClient::DrawTriangleList(FViewport& InViewport, FScene
 	}
 }
 
+void FSpriteEditorViewportClient::DrawSourceRegion(FViewport& InViewport, FSceneView& View, FCanvas& Canvas, const FLinearColor& GeometryVertexColor, bool bIsRenderGeometry)
+{
+	const bool bIsHitTesting = Canvas.IsHitTesting();
+	UPaperSprite* Sprite = GetSpriteBeingEdited();
+
+	const float CollisionVertexSize = 8.0f;
+
+	const float NormalLength = 15.0f;
+	const FLinearColor GeometryLineColor(GeometryVertexColor.R, GeometryVertexColor.G, GeometryVertexColor.B, 0.5f * GeometryVertexColor.A);
+	const FLinearColor GeometryNormalColor(0.0f, 1.0f, 0.0f, 0.5f);
+
+	FVector2D BoundsVertices[4];
+	BoundsVertices[0] = TextureSpaceToScreenSpace(View, Sprite->SourceUV);
+	BoundsVertices[1] = TextureSpaceToScreenSpace(View, Sprite->SourceUV + FVector2D(Sprite->SourceDimension.X, 0));
+	BoundsVertices[2] = TextureSpaceToScreenSpace(View, Sprite->SourceUV + FVector2D(Sprite->SourceDimension.X, Sprite->SourceDimension.Y));
+	BoundsVertices[3] = TextureSpaceToScreenSpace(View, Sprite->SourceUV + FVector2D(0, Sprite->SourceDimension.Y));
+	for (int32 VertexIndex = 0; VertexIndex < 4; ++VertexIndex)
+	{
+		const int32 NextVertexIndex = (VertexIndex + 1) % 4;
+
+		FCanvasLineItem LineItem(BoundsVertices[VertexIndex], BoundsVertices[NextVertexIndex]);
+		LineItem.SetColor(GeometryVertexColor);
+		Canvas.DrawItem(LineItem);
+
+		FVector2D MidPoint = (BoundsVertices[VertexIndex] + BoundsVertices[NextVertexIndex]) * 0.5f;
+
+		if (bIsHitTesting)
+		{
+			TSharedPtr<FSpriteSelectedSourceRegion> Data = MakeShareable(new FSpriteSelectedSourceRegion());
+			Data->SpritePtr = Sprite;
+			Data->VertexIndex = VertexIndex;
+			Canvas.SetHitProxy(new HSpriteSelectableObjectHitProxy(Data));
+		}
+
+		Canvas.DrawTile(MidPoint.X - CollisionVertexSize*0.5f, MidPoint.Y - CollisionVertexSize*0.5f, CollisionVertexSize, CollisionVertexSize, 0.f, 0.f, 1.f, 1.f, GeometryVertexColor, GWhiteTexture);
+
+		if (bIsHitTesting)
+		{
+			Canvas.SetHitProxy(NULL);
+		}
+	}
+}
+
 void FSpriteEditorViewportClient::DrawGeometry(FViewport& InViewport, FSceneView& View, FCanvas& Canvas, const FSpritePolygonCollection& Geometry, const FLinearColor& GeometryVertexColor, bool bIsRenderGeometry)
 {
 	const bool bIsHitTesting = Canvas.IsHitTesting();
@@ -401,6 +445,13 @@ void FSpriteEditorViewportClient::DrawCanvas(FViewport& Viewport, FSceneView& Vi
 			DrawGeometryStats(Viewport, View, Canvas, Sprite->RenderGeometry, true, /*inout*/ YPos);
 		}
 		break;
+	case ESpriteEditorMode::EditSourceRegionMode:
+		{
+			// Draw the custom render geometry
+			const FLinearColor BoundsColor(1.0f, 1.0f, 1.0f, 0.8f);
+			DrawSourceRegion(Viewport, View, Canvas, BoundsColor, true);
+		}
+		break;
 	case ESpriteEditorMode::AddSpriteMode:
 		//@TODO:
 		break;
diff --git a/Engine/Plugins/Experimental/Paper2D/Source/Paper2DEditor/Private/SpriteEditor/SpriteEditorViewportClient.h b/Engine/Plugins/Experimental/Paper2D/Source/Paper2DEditor/Private/SpriteEditor/SpriteEditorViewportClient.h
--- a/Engine/Plugins/Experimental/Paper2D/Source/Paper2DEditor/Private/SpriteEditor/SpriteEditorViewportClient.h
+++ b/Engine/Plugins/Experimental/Paper2D/Source/Paper2DEditor/Private/SpriteEditor/SpriteEditorViewportClient.h
@@ -19,6 +19,7 @@ namespace ESpriteEditorMode
 	enum Type
 	{
 		ViewMode,
+		EditSourceRegionMode,
 		EditCollisionMode,
 		EditRenderingGeomMode,
 		AddSpriteMode
@@ -67,11 +68,13 @@ public:
 	bool IsShowNormalsChecked() const { return bShowNormals; }
 
 	void EnterViewMode() { CurrentMode = ESpriteEditorMode::ViewMode; }
+	void EnterSourceRegionEditMode() { CurrentMode = ESpriteEditorMode::EditSourceRegionMode; }
 	void EnterCollisionEditMode() { CurrentMode = ESpriteEditorMode::EditCollisionMode; }
 	void EnterRenderingEditMode() { CurrentMode = ESpriteEditorMode::EditRenderingGeomMode; }
 	void EnterAddSpriteMode() { CurrentMode = ESpriteEditorMode::AddSpriteMode; }
 
 	bool IsInViewMode() const { return CurrentMode == ESpriteEditorMode::ViewMode; }
+	bool IsInSourceRegionEditMode() const { return CurrentMode == ESpriteEditorMode::EditSourceRegionMode; }
 	bool IsInCollisionEditMode() const { return CurrentMode == ESpriteEditorMode::EditCollisionMode; }
 	bool IsInRenderingEditMode() const { return CurrentMode == ESpriteEditorMode::EditRenderingGeomMode; }
 	bool IsInAddSpriteMode() const { return CurrentMode == ESpriteEditorMode::AddSpriteMode; }
@@ -158,6 +161,7 @@ private:
 	void DrawGeometryStats(FViewport& InViewport, FSceneView& View, FCanvas& Canvas, const FSpritePolygonCollection& Geometry, bool bIsRenderGeometry, int32& YPos);
 	void DrawSockets(const FSceneView* View, FPrimitiveDrawInterface* PDI);
 	void DrawSocketNames(FViewport& InViewport, FSceneView& View, FCanvas& Canvas);
+	void DrawSourceRegion(FViewport& InViewport, FSceneView& View, FCanvas& Canvas, const FLinearColor& GeometryVertexColor, bool bIsRenderGeometry);
 
 	void BeginTransaction(const FText& SessionName);
 	void EndTransaction();
