Paper2D: 2D Physics WIP (start of mass support, improvements in simulation control)

diff --git a/Engine/Plugins/Experimental/Paper2D/Source/Paper2D/Classes/BodyInstance2D.h b/Engine/Plugins/Experimental/Paper2D/Source/Paper2D/Classes/BodyInstance2D.h
--- a/Engine/Plugins/Experimental/Paper2D/Source/Paper2D/Classes/BodyInstance2D.h
+++ b/Engine/Plugins/Experimental/Paper2D/Source/Paper2D/Classes/BodyInstance2D.h
@@ -13,9 +13,13 @@ struct PAPER2D_API FBodyInstance2D
 	UPROPERTY(EditAnywhere, BlueprintReadWrite, Category=Physics)
 	uint32 bSimulatePhysics:1;
 
+	/**	Allows you to override the PhysicalMaterial to use for simple collision on this body. */
+	UPROPERTY(EditAnywhere, AdvancedDisplay, BlueprintReadOnly, Category = Collision)
+	class UPhysicalMaterial* PhysMaterialOverride;
+
 	FBodyInstance2D();
 
-	void InitBody(class UBodySetup2D* Setup, const FTransform& Transform, class UPrimitiveComponent* PrimComp);
+	void InitBody(class UBodySetup2D* Setup, const FTransform& Transform, class UPrimitiveComponent2D* PrimComp);
 	void TermBody();
 
 	/** See if this body is valid. */
@@ -32,9 +36,15 @@ struct PAPER2D_API FBodyInstance2D
 	 */
 	void SetBodyTransform(const FTransform& NewTransform);
 
+	/** Should Simulate Physics **/
+	bool ShouldInstanceSimulatingPhysics() const;
+
+	/** Returns true if this body is simulating, false if it is fixed (kinematic) */
+	bool IsInstanceSimulatingPhysics() const;
+
 public:
 	/** PrimitiveComponent containing this body.   */
-	TWeakObjectPtr<class UPrimitiveComponent> OwnerComponentPtr;
+	TWeakObjectPtr<class UPrimitiveComponent2D> OwnerComponentPtr;
 
 protected:
 
diff --git a/Engine/Plugins/Experimental/Paper2D/Source/Paper2D/Classes/BodySetup2D.h b/Engine/Plugins/Experimental/Paper2D/Source/Paper2D/Classes/BodySetup2D.h
--- a/Engine/Plugins/Experimental/Paper2D/Source/Paper2D/Classes/BodySetup2D.h
+++ b/Engine/Plugins/Experimental/Paper2D/Source/Paper2D/Classes/BodySetup2D.h
@@ -16,6 +16,10 @@ class PAPER2D_API UBodySetup2D : public UObject
 	UPROPERTY()
 	FAggregateGeometry2D AggGeom2D;
 
+	/** Physical material to use for simple collision on this body. Encodes information about density, friction etc. */
+	UPROPERTY(EditAnywhere, Category = Physics, meta = (DisplayName = "Simple Collision Physical Material"))
+	class UPhysicalMaterial* PhysMaterial;
+
 	/** Default properties of the body instance, copied into objects on instantiation */
 	UPROPERTY(EditAnywhere, Category=Collision, meta=(FullyExpand = "true"))
 	struct FBodyInstance2D DefaultInstance;
@@ -23,6 +27,11 @@ class PAPER2D_API UBodySetup2D : public UObject
 	/** Create Physics meshes (ConvexMeshes, TriMesh & TriMeshNegX) from cooked data */
 	void CreatePhysicsMeshes();
 
+	/** Calculates the mass. You can pass in the component where additional information is pulled from ( Scale, PhysMaterialOverride ) */
+	float CalculateMass(const class UPrimitiveComponent2D* Component = nullptr) const;
+
+	/** Returns the physics material used for this body. If none, specified, returns the default engine material. */
+	class UPhysicalMaterial* GetPhysMaterial() const;
 
 #if WITH_EDITOR
 	/** Release Physics meshes (ConvexMeshes, TriMesh & TriMeshNegX). Must be called before the BodySetup is destroyed */
diff --git a/Engine/Plugins/Experimental/Paper2D/Source/Paper2D/Classes/PrimitiveComponent2D.h b/Engine/Plugins/Experimental/Paper2D/Source/Paper2D/Classes/PrimitiveComponent2D.h
--- a/Engine/Plugins/Experimental/Paper2D/Source/Paper2D/Classes/PrimitiveComponent2D.h
+++ b/Engine/Plugins/Experimental/Paper2D/Source/Paper2D/Classes/PrimitiveComponent2D.h
@@ -23,6 +23,7 @@ public:
 
 	// USceneComponent interface
 	virtual void OnUpdateTransform(bool bSkipPhysicsMove) OVERRIDE;
+	virtual bool IsSimulatingPhysics(FName BoneName = NAME_None) const OVERRIDE;
 	// End of USceneComponent interface
 
 	// UPrimitiveComponent interface
@@ -32,6 +33,9 @@ public:
 	/** Return the BodySetup to use for this PrimitiveComponent (single body case) */
 	virtual class UBodySetup2D* GetBodySetup2D() const;
 
+	/** Get a BodyInstance from this component. The supplied name is used in the SkeletalMeshComponent case. A name of NAME_None in the skeletal case gives the root body instance. */
+	virtual const FBodyInstance2D* GetBodyInstance2D(FName BoneName = NAME_None) const;
+
 protected:
 	//@TODO: Document
 	virtual void CreatePhysicsState2D();
diff --git a/Engine/Plugins/Experimental/Paper2D/Source/Paper2D/Private/Physics/BodyInstance2D.cpp b/Engine/Plugins/Experimental/Paper2D/Source/Paper2D/Private/Physics/BodyInstance2D.cpp
--- a/Engine/Plugins/Experimental/Paper2D/Source/Paper2D/Private/Physics/BodyInstance2D.cpp
+++ b/Engine/Plugins/Experimental/Paper2D/Source/Paper2D/Private/Physics/BodyInstance2D.cpp
@@ -14,8 +14,7 @@ FBodyInstance2D::FBodyInstance2D()
 	bSimulatePhysics = true;
 }
 
-//@TODO: Everything about this block of code...
-void FBodyInstance2D::InitBody(class UBodySetup2D* Setup, const FTransform& Transform, class UPrimitiveComponent* PrimComp)
+void FBodyInstance2D::InitBody(class UBodySetup2D* Setup, const FTransform& Transform, class UPrimitiveComponent2D* PrimComp)
 {
 #if WITH_BOX2D
 	if (b2World* World = FPhysicsIntegration2D::FindAssociatedWorld(PrimComp->GetWorld()))
@@ -38,13 +37,14 @@ void FBodyInstance2D::InitBody(class UBodySetup2D* Setup, const FTransform& Tran
 			{
 				b2CircleShape CircleShape;
 				CircleShape.m_radius = Circle.Radius * Scale3D.Size() / UnrealUnitsPerMeter;
-				CircleShape.m_p.x = Circle.Center.X;
-				CircleShape.m_p.y = Circle.Center.Y;
+				CircleShape.m_p.x = Circle.Center.X * Scale2D.x;
+				CircleShape.m_p.y = Circle.Center.Y * Scale2D.y;
 
 				b2FixtureDef FixtureDef;
 				FixtureDef.shape = &CircleShape;
 				FixtureDef.density = 1.0f; //@TODO:
 				FixtureDef.friction = 0.3f;//@TODO:
+				//FixtureDef.restitution = ;
 
 				BodyInstancePtr->CreateFixture(&FixtureDef);
 			}
@@ -148,12 +148,12 @@ FTransform FBodyInstance2D::GetUnrealWorldTransform() const
 	if (BodyInstancePtr != NULL)
 	{
 		const b2Vec2 Pos2D = BodyInstancePtr->GetPosition();
-		const float Rotation = BodyInstancePtr->GetAngle();
+		const float RotationInRadians = BodyInstancePtr->GetAngle();
 		//@TODO: Should I use get transform?
 		//@TODO: What about scale?
 
 		const FVector Translation3D(FPhysicsIntegration2D::ConvertBoxVectorToUnreal(Pos2D));
-		const FRotator Rotation3D(FMath::RadiansToDegrees(Rotation), 0.0f, 0.0f);
+		const FRotator Rotation3D(FMath::RadiansToDegrees(RotationInRadians), 0.0f, 0.0f);
 		const FVector Scale3D(1.0f, 1.0f, 1.0f);
 		
 		return FTransform(Rotation3D, Translation3D, Scale3D);
@@ -182,6 +182,27 @@ void FBodyInstance2D::SetBodyTransform(const FTransform& NewTransform)
 #endif
 }
 
+bool FBodyInstance2D::ShouldInstanceSimulatingPhysics() const
+{
+	// if I'm simulating or owner is simulating
+	UPrimitiveComponent2D* OwnerComponent = OwnerComponentPtr.Get();
+	
+	// if derive from owner, and owner is simulating, this should simulate
+	if ((OwnerComponent != nullptr) && OwnerComponent->GetBodyInstance2D()->bSimulatePhysics)
+	{
+		return true;
+	}
+
+	// or else, it should look its own setting
+	return bSimulatePhysics;
+}
+
+bool FBodyInstance2D::IsInstanceSimulatingPhysics() const
+{
+	// if I'm simulating or owner is simulating
+	return ShouldInstanceSimulatingPhysics() && IsValidBodyInstance();
+}
+
 void FBodyInstance2D::SetInstanceSimulatePhysics(bool bSimulate)
 {
 #if WITH_BOX2D
diff --git a/Engine/Plugins/Experimental/Paper2D/Source/Paper2D/Private/Physics/BodySetup2D.cpp b/Engine/Plugins/Experimental/Paper2D/Source/Paper2D/Private/Physics/BodySetup2D.cpp
--- a/Engine/Plugins/Experimental/Paper2D/Source/Paper2D/Private/Physics/BodySetup2D.cpp
+++ b/Engine/Plugins/Experimental/Paper2D/Source/Paper2D/Private/Physics/BodySetup2D.cpp
@@ -26,3 +26,67 @@ void UBodySetup2D::InvalidatePhysicsData()
 void UBodySetup2D::CreatePhysicsMeshes()
 {
 }
+
+class UPhysicalMaterial* UBodySetup2D::GetPhysMaterial() const
+{
+	UPhysicalMaterial* PhysMat = PhysMaterial;
+
+	if (PhysMat == NULL && GEngine != NULL)
+	{
+		PhysMat = GEngine->DefaultPhysMaterial;
+	}
+	return PhysMat;
+}
+
+float UBodySetup2D::CalculateMass(const UPrimitiveComponent2D* Component) const
+{
+	FVector ComponentScale(1.0f, 1.0f, 1.0f);
+	UPhysicalMaterial* PhysMat = GetPhysMaterial();
+	const FBodyInstance2D* BodyInstance = NULL;
+
+	const UPrimitiveComponent2D* OuterComp = Component != NULL ? Component : Cast<UPrimitiveComponent2D>(GetOuter());
+	if (OuterComp)
+	{
+		ComponentScale = OuterComp->GetComponentScale();
+
+		BodyInstance = OuterComp->GetBodyInstance2D();
+
+		//@TODO:
+// 		const USkinnedMeshComponent* SkinnedMeshComp = Cast<const USkinnedMeshComponent>(OuterComp);
+// 		if (SkinnedMeshComp != NULL)
+// 		{
+// 			const FBodyInstance* Body = SkinnedMeshComp->GetBodyInstance2D(BoneName);
+// 
+// 			if (Body != NULL)
+// 			{
+// 				BodyInstance = Body;
+// 			}
+// 		}
+	}
+	else
+	{
+		BodyInstance = &DefaultInstance;
+	}
+
+	PhysMat = BodyInstance->PhysMaterialOverride != NULL ? BodyInstance->PhysMaterialOverride : PhysMat;
+	const float MassScale = 1.0f;//@TODO: BodyInstance->MassScale;
+
+	// physical material - nothing can weigh less than hydrogen (0.09 kg/m^3)
+	float DensityKGPerCubicUU = 1.0f;
+	float RaiseMassToPower = 0.75f;
+	float FictionalThickness = 1.0f;
+	if (PhysMat)
+	{
+		DensityKGPerCubicUU = FMath::Max(0.00009f, PhysMat->Density * 0.001f);
+		RaiseMassToPower = PhysMat->RaiseMassToPower;
+	}
+
+	// Then scale mass to avoid big differences between big and small objects.
+	const float BasicMass = AggGeom2D.GetArea(ComponentScale) * FictionalThickness * DensityKGPerCubicUU;
+	ensureMsgf(BasicMass >= 0.0f, TEXT("UBodySetup2D::CalculateMass(%s) - The volume of the aggregate geometry is negative"), *Component->GetReadableName());
+
+	const float UsePow = FMath::Clamp<float>(RaiseMassToPower, KINDA_SMALL_NUMBER, 1.f);
+	const float RealMass = FMath::Pow(FMath::Max<float>(BasicMass, 0.0f), UsePow);
+
+	return RealMass * MassScale;
+}
diff --git a/Engine/Plugins/Experimental/Paper2D/Source/Paper2D/Private/Physics/Box2DIntegration.cpp b/Engine/Plugins/Experimental/Paper2D/Source/Paper2D/Private/Physics/Box2DIntegration.cpp
--- a/Engine/Plugins/Experimental/Paper2D/Source/Paper2D/Private/Physics/Box2DIntegration.cpp
+++ b/Engine/Plugins/Experimental/Paper2D/Source/Paper2D/Private/Physics/Box2DIntegration.cpp
@@ -144,15 +144,18 @@ void FStartPhysics2DTickFunction::ExecuteTick(float DeltaTime, enum ELevelTick T
 
 			if (FBodyInstance2D* UnrealBodyInstance = (FBodyInstance2D*)Body->GetUserData())
 			{
-				// See if the transform is actually different, and if so, move the component to match physics
-				const FTransform NewTransform = UnrealBodyInstance->GetUnrealWorldTransform();	
-				const FTransform& OldTransform = UnrealBodyInstance->OwnerComponentPtr->ComponentToWorld;
-				if (!NewTransform.Equals(OldTransform))
+				if (UnrealBodyInstance->bSimulatePhysics)
 				{
-					const FVector MoveBy = NewTransform.GetLocation() - OldTransform.GetLocation();
-					const FRotator NewRotation = NewTransform.Rotator();
-
-					UnrealBodyInstance->OwnerComponentPtr->MoveComponent(MoveBy, NewRotation, false, NULL, MOVECOMP_SkipPhysicsMove);
+					// See if the transform is actually different, and if so, move the component to match physics
+					const FTransform NewTransform = UnrealBodyInstance->GetUnrealWorldTransform();
+					const FTransform& OldTransform = UnrealBodyInstance->OwnerComponentPtr->ComponentToWorld;
+					if (!NewTransform.Equals(OldTransform))
+					{
+						const FVector MoveBy = NewTransform.GetLocation() - OldTransform.GetLocation();
+						const FRotator NewRotation = NewTransform.Rotator();
+
+						UnrealBodyInstance->OwnerComponentPtr->MoveComponent(MoveBy, NewRotation, false, NULL, MOVECOMP_SkipPhysicsMove);
+					}
 				}
 			}
 
diff --git a/Engine/Plugins/Experimental/Paper2D/Source/Paper2D/Private/Physics/PrimitiveComponent2D.cpp b/Engine/Plugins/Experimental/Paper2D/Source/Paper2D/Private/Physics/PrimitiveComponent2D.cpp
--- a/Engine/Plugins/Experimental/Paper2D/Source/Paper2D/Private/Physics/PrimitiveComponent2D.cpp
+++ b/Engine/Plugins/Experimental/Paper2D/Source/Paper2D/Private/Physics/PrimitiveComponent2D.cpp
@@ -60,6 +60,18 @@ void UPrimitiveComponent2D::OnUpdateTransform(bool bSkipPhysicsMove)
 	}
 }
 
+bool UPrimitiveComponent2D::IsSimulatingPhysics(FName BoneName) const
+{
+	if (const FBodyInstance2D* BodyInst2D = GetBodyInstance2D(BoneName))
+	{
+		return BodyInst2D->IsInstanceSimulatingPhysics();
+	}
+	else
+	{
+		return Super::IsSimulatingPhysics(BoneName);
+	}
+}
+
 void UPrimitiveComponent2D::SetSimulatePhysics(bool bSimulate)
 {
 	Super::SetSimulatePhysics(bSimulate);
@@ -70,3 +82,8 @@ class UBodySetup2D* UPrimitiveComponent2D::GetBodySetup2D() const
 {
 	return nullptr;
 }
+
+const FBodyInstance2D* UPrimitiveComponent2D::GetBodyInstance2D(FName BoneName) const
+{
+	return &BodyInstance2D;
+}
