Paper2D: Dummy filter implementation

diff --git a/Engine/Plugins/Experimental/Paper2D/Source/Paper2D/Private/Physics/Box2DIntegration.cpp b/Engine/Plugins/Experimental/Paper2D/Source/Paper2D/Private/Physics/Box2DIntegration.cpp
--- a/Engine/Plugins/Experimental/Paper2D/Source/Paper2D/Private/Physics/Box2DIntegration.cpp
+++ b/Engine/Plugins/Experimental/Paper2D/Source/Paper2D/Private/Physics/Box2DIntegration.cpp
@@ -122,6 +122,11 @@ FPhysicsScene2D::~FPhysicsScene2D()
 	EndPhysicsTickFunction.UnRegisterTickFunction();
 }
 
+bool FPhysicsScene2D::ShouldCollide(b2Fixture* FixtureA, b2Fixture* FixtureB)
+{
+	return b2ContactFilter::ShouldCollide(FixtureA, FixtureB);
+}
+
 //////////////////////////////////////////////////////////////////////////
 // FStartPhysics2DTickFunction
 
diff --git a/Engine/Plugins/Experimental/Paper2D/Source/Paper2D/Private/Physics/Box2DIntegration.h b/Engine/Plugins/Experimental/Paper2D/Source/Paper2D/Private/Physics/Box2DIntegration.h
--- a/Engine/Plugins/Experimental/Paper2D/Source/Paper2D/Private/Physics/Box2DIntegration.h
+++ b/Engine/Plugins/Experimental/Paper2D/Source/Paper2D/Private/Physics/Box2DIntegration.h
@@ -33,7 +33,7 @@ struct FEndPhysics2DTickFunction : public FTickFunction
 //////////////////////////////////////////////////////////////////////////
 // FPhysicsWorld2D
 
-struct FPhysicsScene2D
+struct FPhysicsScene2D : public b2ContactFilter
 {
 	class b2World* World;
 	UWorld* UnrealWorld;
@@ -44,8 +44,13 @@ struct FPhysicsScene2D
 	// Tick function for ending physics
 	FEndPhysics2DTickFunction EndPhysicsTickFunction;
 
+public:
 	FPhysicsScene2D(UWorld* AssociatedWorld);
-	~FPhysicsScene2D();
+	virtual ~FPhysicsScene2D();
+
+	// b2ContactFilter interface
+	virtual bool ShouldCollide(b2Fixture* FixtureA, b2Fixture* FixtureB) OVERRIDE;
+	// End of b2ContactFilter interface
 };
 
 #endif
