Paper2D: Double clicking in source region edit mode detects a sprite at mouse cursor and adjusts the source region to tightly fit it
BitmapUtils to share code with this and one-click splitting up a sprite sheet

diff --git a/Engine/Plugins/Experimental/Paper2D/Source/Paper2D/Classes/PaperSprite.h b/Engine/Plugins/Experimental/Paper2D/Source/Paper2D/Classes/PaperSprite.h
--- a/Engine/Plugins/Experimental/Paper2D/Source/Paper2D/Classes/PaperSprite.h
+++ b/Engine/Plugins/Experimental/Paper2D/Source/Paper2D/Classes/PaperSprite.h
@@ -139,6 +139,7 @@ public:
 	FVector2D ConvertPivotSpaceToTextureSpace(FVector2D Input) const;
 	FVector ConvertTextureSpaceToPivotSpace(FVector Input) const;
 	FVector ConvertPivotSpaceToTextureSpace(FVector Input) const;
+	FVector2D ConvertWorldSpaceToTextureSpace(const FVector& WorldPoint) const;
 
 	// World space WRT the sprite editor *only*
 	FVector ConvertTextureSpaceToWorldSpace(const FVector2D& SourcePoint) const;
@@ -151,6 +152,8 @@ public:
 	void RebuildCollisionData();
 	void RebuildRenderData();
 
+	void ExtractSourceRegionFromTexturePoint(const FVector2D& Point);
+
 	void FindTextureBoundingBox(float AlphaThreshold, /*out*/ FVector2D& OutBoxPosition, /*out*/ FVector2D& OutBoxSize);
 	static void FindContours(const FIntPoint& ScanPos, const FIntPoint& ScanSize, float AlphaThreshold, UTexture2D* Texture, TArray< TArray<FIntPoint> >& OutPoints);
 	void BuildGeometryFromContours(FSpritePolygonCollection& GeomOwner);
diff --git a/Engine/Plugins/Experimental/Paper2D/Source/Paper2D/Private/BitmapUtils.h b/Engine/Plugins/Experimental/Paper2D/Source/Paper2D/Private/BitmapUtils.h
new file mode 100644
--- /dev/null
+++ b/Engine/Plugins/Experimental/Paper2D/Source/Paper2D/Private/BitmapUtils.h
@@ -0,0 +1,375 @@
+// Copyright 1998-2014 Epic Games, Inc. All Rights Reserved.
+
+#pragma once
+
+//////////////////////////////////////////////////////////////////////////
+// FBitmap
+
+struct FBitmap
+{
+	FBitmap(UTexture2D* SourceTexture, int AlphaThreshold, int DefaultValue)
+		: Width(0)
+		, Height(0)
+		, DefaultValue(DefaultValue)
+	{
+		ExtractFromTexture(SourceTexture, AlphaThreshold);
+	}
+
+	void ExtractFromTexture(UTexture2D* SourceTexture, int AlphaThreshold)
+	{
+		// use the source art if it exists
+		FTextureSource* TextureSource = NULL;
+		if ((SourceTexture != NULL) && SourceTexture->Source.IsValid())
+		{
+			switch (SourceTexture->Source.GetFormat())
+			{
+			case TSF_G8:
+			case TSF_BGRA8:
+				TextureSource = &(SourceTexture->Source);
+				break;
+			default:
+				break;
+			};
+		}
+
+		if (TextureSource != NULL)
+		{
+			TArray<uint8> TextureRawData;
+			TextureSource->GetMipData(TextureRawData, 0);
+			int32 BytesPerPixel = TextureSource->GetBytesPerPixel();
+			ETextureSourceFormat PixelFormat = TextureSource->GetFormat();
+			Width = TextureSource->GetSizeX();
+			Height = TextureSource->GetSizeY();
+			RawData.SetNumZeroed(Width * Height);
+
+			for (int Y = 0; Y < Height; ++Y)
+			{
+				for (int X = 0; X < Width; ++X)
+				{
+					int32 PixelByteOffset = (X + Y * Width) * BytesPerPixel;
+					const uint8* PixelPtr = TextureRawData.GetTypedData() + PixelByteOffset;
+					FColor Color;
+					if (PixelFormat == TSF_BGRA8)
+					{
+						Color = *((FColor*)PixelPtr);
+					}
+					else
+					{
+						checkSlow(PixelFormat == TSF_G8);
+						const uint8 Intensity = *PixelPtr;
+						Color = FColor(Intensity, Intensity, Intensity, Intensity);
+					}
+
+					if (Color.A > AlphaThreshold)
+					{
+						RawData[Y * Width + X] = 1;
+					}
+					else
+					{
+						RawData[Y * Width + X] = DefaultValue;
+					}
+				}
+			}
+		}
+	}
+
+	// Create an empty bitmap
+	FBitmap(int Width, int Height, int DefaultValue)
+		: Width(Width)
+		, Height(Height)
+		, DefaultValue(DefaultValue)
+	{
+		RawData.SetNumZeroed(Width * Height);
+		ClearToDefaultValue();
+	}
+
+	bool IsValid() const
+	{
+		return Width != 0 && Height != 0;
+	}
+
+	uint8 GetPixel(int X, int Y)
+	{
+		if (X >= 0 && X < Width && Y >= 0 && Y < Height)
+		{
+			return RawData[Y * Width + X];
+		}
+		else
+		{
+			return DefaultValue;
+		}
+	}
+
+	void SetPixel(int X, int Y, uint8 Value)
+	{
+		if (X >= 0 && X < Width && Y >= 0 && Y < Height)
+		{
+			RawData[Y * Width + X] = Value;
+		}
+	}
+
+	void ClearToDefaultValue()
+	{
+		int RawDataCount = Width * Height;
+		for (int RawDataIndex = 0; RawDataIndex < RawDataCount; ++RawDataIndex)
+		{
+			RawData[RawDataIndex] = DefaultValue;
+		}
+	}
+
+	bool IsColumnEmpty(int X, int Y0, int Y1)
+	{
+		for (int Y = Y0; Y < Y1; ++Y)
+		{
+			if (GetPixel(X, Y) != 0)
+			{
+				return false;
+			}
+		}
+		return true;
+	}
+
+	bool IsRowEmpty(int X0, int X1, int Y)
+	{
+		for (int X = X0; X < X1; ++X)
+		{
+			if (GetPixel(X, Y) != 0)
+			{
+				return false;
+			}
+		}
+		return true;
+	}
+
+	// Returns the tight bounding box around pixels that are not 0
+	void GetTightBounds(int& OutX, int& OutY, int& OutWidth, int& OutHeight)
+	{
+		int Top = 0;
+		int Bottom = Height - 1;
+		int Left = 0;
+		int Right = Width - 1;
+		while (Top < Bottom && IsRowEmpty(Left, Right, Top))
+		{
+			++Top;
+		}
+		while (Bottom >= Top && IsRowEmpty(Left, Right, Bottom))
+		{
+			--Bottom;
+		}
+		while (Left < Right && IsColumnEmpty(Left, Top, Bottom))
+		{
+			++Left;
+		}
+		while (Right >= Left && IsColumnEmpty(Right, Top, Bottom))
+		{
+			--Right;
+		}
+
+		OutX = Left;
+		OutY = Top;
+		OutWidth = Right - Left + 1;
+		OutHeight = Bottom - Top + 1;
+	}
+
+	// Performs a flood fill on the target bitmap, with the boundary defined
+	// by the current bitmap.
+	void FloodFill(FBitmap& MaskBitmap, int StartX, int StartY)
+	{
+		TArray<FIntPoint> QueuedPoints;
+		QueuedPoints.Reserve(Width);
+
+		QueuedPoints.Add(FIntPoint(StartX, StartY));
+		while (QueuedPoints.Num() > 0)
+		{
+			FIntPoint Point = QueuedPoints.Pop(false);
+			if (MaskBitmap.GetPixel(Point.X, Point.Y) == 0 && GetPixel(Point.X, Point.Y) == 1)
+			{
+				MaskBitmap.SetPixel(Point.X, Point.Y, 1);
+				if (Point.X > 0) 
+				{
+					QueuedPoints.Add(FIntPoint(Point.X - 1, Point.Y));
+				}
+				if (Point.X < Width - 1) 
+				{
+					QueuedPoints.Add(FIntPoint(Point.X + 1, Point.Y));
+				}
+				if (Point.Y > 0) 
+				{
+					QueuedPoints.Add(FIntPoint(Point.X, Point.Y - 1));
+				}
+				if (Point.Y < Height - 1) 
+				{
+					QueuedPoints.Add(FIntPoint(Point.X, Point.Y + 1));
+				}
+			}
+		}
+	}
+
+	// Walk the border pixels to find any intersecting islands we haven't already filled in the mask bitmap
+	bool HasOverlappingIsland(FBitmap& MaskBitmap, int X0, int Y0, int Width, int Height, int& FillX, int& FillY)
+	{
+		FillX = 0;
+		FillY = 0;
+		int X1 = X0 + Width - 1;
+		int Y1 = Y0 + Height - 1;
+
+		for (int X = X0; X <= X1; ++X)
+		{
+			if (MaskBitmap.GetPixel(X, Y0) == 0 && GetPixel(X, Y0) != 0)
+			{
+				FillX = X;
+				FillY = Y0;
+				return true;
+			}
+			if (MaskBitmap.GetPixel(X, Y1) == 0 && GetPixel(X, Y1) != 0)
+			{
+				FillX = X;
+				FillY = Y1;
+				return true;
+			}
+		}
+
+		for (int Y = Y0; Y <= Y1; ++Y)
+		{
+			if (MaskBitmap.GetPixel(X0, Y) == 0 && GetPixel(X0, Y) != 0)
+			{
+				FillX = X0;
+				FillY = Y;
+				return true;
+			}
+			if (MaskBitmap.GetPixel(X1, Y) == 0 && GetPixel(X1, Y) != 0)
+			{
+				FillX = X1;
+				FillY = Y;
+				return true;
+			}
+		}
+
+		return false;
+	}
+
+	// Sets pixels to 1 marking the rectangle outline
+	void DrawRectOutline(int StartX, int StartY, int Width, int Height)
+	{
+		int X0 = StartX;
+		int Y0 = StartY;
+		int X1 = StartX + Width - 1;
+		int Y1 = StartY + Height - 1;
+		for (int Y = Y0; Y <= Y1; ++Y)
+		{
+			SetPixel(X0, Y, 1);
+			SetPixel(X1, Y, 1);
+		}
+		for (int X = X0; X <= X1; ++X)
+		{
+			SetPixel(X, Y0, 1);
+			SetPixel(X, Y1, 1);
+		}
+	}
+
+	// Finds the rect of the contour of the shape clicked on, extended by rectangles to support
+	// separated but intersecting islands.
+	bool HasConnectedRect(int X, int Y, bool Extend, int& OutOriginX, int& OutOriginY, int& OutDimensionX, int& OutDimensionY)
+	{
+		if (GetPixel(X, Y) == 0) 
+		{
+			// selected an empty pixel
+			OutOriginX = 0;
+			OutOriginY = 0;
+			OutDimensionX = 0;
+			OutDimensionY = 0;
+			return false;
+		}
+
+		// Tmp get rid of this
+		// This whole thing can be much more efficiently using the 8 bpp
+		FBitmap MaskBitmap(Width, Height, 0);
+		if (Extend)
+		{
+			MaskBitmap.DrawRectOutline(OutOriginX, OutOriginY, OutDimensionX, OutDimensionY);
+		}
+
+		// MaxPasses shouldn't be nessecary, but worst case interlocked pixel
+		// patterns can cause problems. Dilating the bitmap before processing will
+		// reduce these problems, but may not be desirable in all cases.
+		int NumPasses = 0;
+		int MaxPasses = 40;
+		OutOriginX = 0;
+		OutOriginY = 0;
+		OutDimensionX = 0;
+		OutDimensionY = 0;
+		int FillPointX = X;
+		int FillPointY = Y;
+		do
+		{
+			// This is probably going to be a botleneck at larger texture sizes
+			// A contour tracing algorithm will probably suffice here.
+			FloodFill(MaskBitmap, FillPointX, FillPointY);
+			MaskBitmap.GetTightBounds(OutOriginX, OutOriginY, OutDimensionX, OutDimensionY);
+		} while (NumPasses++ < MaxPasses && HasOverlappingIsland(MaskBitmap, OutOriginX, OutOriginY, OutDimensionX, OutDimensionY, /*out*/FillPointX, /*out*/FillPointY));
+
+		return true;
+	}
+
+	// Wind through the bitmap from StartX, StartY to find the closest hit point
+	bool FoundClosestValidPoint(int StartX, int StartY, int MaxAllowedSearchDistance, int& HitX, int& HitY)
+	{
+		// Constrain the point within the image bounds
+		StartX = FMath::Clamp(StartX, 0, Width - 1);
+		StartY = FMath::Clamp(StartY, 0, Height - 1);
+		HitX = StartX;
+		HitY = StartY;
+
+		// Should probably calculate a better max based on StartX and StartY
+		int RequiredSearchDistance = (FMath::Max(Width, Height) + 1) / 2;
+		int MaxSearchDistance = FMath::Min(RequiredSearchDistance, MaxAllowedSearchDistance);
+		int SearchDistance = 0;
+		while (SearchDistance < MaxSearchDistance)
+		{
+			int X0 = FMath::Max(StartX - SearchDistance, 0);
+			int X1 = FMath::Min(StartX + SearchDistance, Width - 1);
+			int Y0 = FMath::Max(StartY - SearchDistance, 0);
+			int Y1 = FMath::Min(StartY + SearchDistance, Height - 1);
+			// Search along the rectangular edge
+			for (int Y = Y0; Y <= Y1; ++Y)
+			{
+				if (GetPixel(X0, Y) != 0) 
+				{
+					HitX = X0;
+					HitY = Y;
+					return true;
+				}
+				if (GetPixel(X1, Y) != 0) 
+				{
+					HitX = X1;
+					HitY = Y;
+					return true;
+				}
+			}
+			for (int X = X0; X <= X1; ++X)
+			{
+				if (GetPixel(X, Y0) != 0) 
+				{
+					HitX = X;
+					HitY = Y0;
+					return true;
+				}
+				if (GetPixel(X, Y1) != 0) 
+				{
+					HitX = X;
+					HitY = Y1;
+					return true;
+				}
+			}
+			SearchDistance += 1;
+		}
+		return false;
+	}
+
+	TArray<uint8> RawData;
+	int32 Width;
+	int32 Height;
+	int32 DefaultValue;
+};
+
+
diff --git a/Engine/Plugins/Experimental/Paper2D/Source/Paper2D/Private/PaperSprite.cpp b/Engine/Plugins/Experimental/Paper2D/Source/Paper2D/Private/PaperSprite.cpp
--- a/Engine/Plugins/Experimental/Paper2D/Source/Paper2D/Private/PaperSprite.cpp
+++ b/Engine/Plugins/Experimental/Paper2D/Source/Paper2D/Private/PaperSprite.cpp
@@ -9,6 +9,7 @@
 
 #include "PaperSpriteAtlas.h"
 #include "GeomTools.h"
+#include "BitmapUtils.h"
 
 //////////////////////////////////////////////////////////////////////////
 // FTextureReader
@@ -342,6 +343,32 @@ struct FBoundaryImage
 	}
 };
 
+void UPaperSprite::ExtractSourceRegionFromTexturePoint(const FVector2D& SourcePoint)
+{
+	int SourceX = FMath::RoundToInt(SourcePoint.X);
+	int SourceY = FMath::RoundToInt(SourcePoint.Y);
+	int ClosestValidX = 0;
+	int ClosestValidY = 0;
+
+	FBitmap Bitmap(SourceTexture, 0, 0);
+	if (Bitmap.IsValid() && Bitmap.FoundClosestValidPoint(SourceX, SourceY, 10, /*out*/ ClosestValidX, /*out*/ ClosestValidY))
+	{
+		int OriginX = 0;
+		int OriginY = 0;
+		int DimensionX = 0;
+		int DimensionY = 0;
+		if (Bitmap.HasConnectedRect(ClosestValidX, ClosestValidY, false, /*out*/ OriginX, /*out*/ OriginY, /*out*/ DimensionX, /*out*/ DimensionY))
+		{
+			if (DimensionX > 0 && DimensionY > 0)
+			{
+				SourceUV = FVector2D(OriginX, OriginY);
+				SourceDimension = FVector2D(DimensionX, DimensionY);
+				PostEditChange();
+			}
+		}
+	}
+}
+
 #endif
 
 //////////////////////////////////////////////////////////////////////////
@@ -1053,6 +1080,14 @@ FVector UPaperSprite::ConvertTextureSpaceToWorldSpace(const FVector2D& SourcePoi
 // 	return FVector(SourcePoint.X, SourcePoint.Y, SourceDims.Y - SourcePoint.Z);
 // }
 
+FVector2D UPaperSprite::ConvertWorldSpaceToTextureSpace(const FVector& WorldPoint) const
+{
+	const FVector2D SourceDims = (SourceTexture != NULL) ? FVector2D(SourceTexture->GetSurfaceWidth(), SourceTexture->GetSurfaceHeight()) : FVector2D::ZeroVector;
+	const FVector ProjectionX = WorldPoint.ProjectOnTo(PaperAxisX);
+	const FVector ProjectionY = WorldPoint.ProjectOnTo(PaperAxisY);
+	return FVector2D(FMath::Sign(ProjectionX | PaperAxisX) * ProjectionX.Size(), SourceDims.Y - FMath::Sign(ProjectionY | PaperAxisY) * ProjectionY.Size());		
+}
+
 FTransform UPaperSprite::GetPivotToWorld() const
 {
 	const FVector2D Pivot = GetPivotPosition();
@@ -1182,4 +1217,4 @@ void UPaperSprite::PostLoad()
 UTexture2D* UPaperSprite::GetBakedTexture() const
 {
 	return (BakedSourceTexture != nullptr) ? BakedSourceTexture : SourceTexture;
-}
\ No newline at end of file
+}
diff --git a/Engine/Plugins/Experimental/Paper2D/Source/Paper2DEditor/Private/SpriteEditor/SpriteEditorViewportClient.cpp b/Engine/Plugins/Experimental/Paper2D/Source/Paper2DEditor/Private/SpriteEditor/SpriteEditorViewportClient.cpp
--- a/Engine/Plugins/Experimental/Paper2D/Source/Paper2DEditor/Private/SpriteEditor/SpriteEditorViewportClient.cpp
+++ b/Engine/Plugins/Experimental/Paper2D/Source/Paper2DEditor/Private/SpriteEditor/SpriteEditorViewportClient.cpp
@@ -525,11 +525,20 @@ void FSpriteEditorViewportClient::Tick(float DeltaSeconds)
 		const FVector PivotInWorldSpace = TextureSpaceToWorldSpace(PivotInTextureSpace);
 		RenderSpriteComponent->SetRelativeLocation(PivotInWorldSpace);
 
+		bool bSourceTextureViewComponentVisibility = bShowSourceTexture || (CurrentMode == ESpriteEditorMode::EditSourceRegionMode);
+		if (bSourceTextureViewComponentVisibility != SourceTextureViewComponent->IsVisible())
+		{
+			bDeferZoomToSprite = true;
+			SourceTextureViewComponent->SetVisibility(bSourceTextureViewComponentVisibility);
+		}
+
 		// Zoom in on the sprite
-		//@TODO: This doesn't work correctly, only partially zooming in or something
-		if (bDeferZoomToSprite)
+		//@TODO: Fix this properly so it doesn't need to be deferred, or wait for the viewport to initialize
+		FIntPoint Size = Viewport->GetSizeXY();
+		if (bDeferZoomToSprite && Size.X > 0 && Size.Y > 0)
 		{
-			FocusViewportOnBox(RenderSpriteComponent->Bounds.GetBox(), true);
+			UPaperRenderComponent *ComponentToFocusOn = SourceTextureViewComponent->IsVisible() ? SourceTextureViewComponent : RenderSpriteComponent;
+			FocusViewportOnBox(ComponentToFocusOn->Bounds.GetBox(), true);
 			bDeferZoomToSprite = false;
 		}		
 	}
@@ -584,6 +593,14 @@ void FSpriteEditorViewportClient::ProcessClick(FSceneView& View, HHitProxy* HitP
 	{
 		ClearSelectionSet();
 
+		if (Event == EInputEvent::IE_DoubleClick && CurrentMode == ESpriteEditorMode::EditSourceRegionMode)
+		{
+			FVector4 WorldPoint = View.PixelToWorld(HitX, HitY, 0);
+			UPaperSprite* Sprite = GetSpriteBeingEdited();
+			FVector2D TexturePoint = Sprite->ConvertWorldSpaceToTextureSpace(WorldPoint);
+			Sprite->ExtractSourceRegionFromTexturePoint(TexturePoint);
+		}
+
 		FPaperEditorViewportClient::ProcessClick(View, HitProxy, Key, Event, HitX, HitY);
 	}
 }
@@ -622,8 +639,9 @@ bool FSpriteEditorViewportClient::InputWidgetDelta(FViewport* Viewport, EAxisLis
 		if (SelectionSet.Num() > 0)
 		{
 			UPaperSprite* Sprite = GetSpriteBeingEdited();
-			Sprite->RebuildCollisionData();
-			Sprite->RebuildRenderData();
+			// Sprite->RebuildCollisionData();
+			// Sprite->RebuildRenderData();
+			Sprite->PostEditChange();
 			bManipulationDirtiedSomething = true;
 
 			Invalidate();
