Paper2D: Delete UPrimitiveComponent2D

diff --git a/Engine/Plugins/Experimental/Paper2D/Source/Paper2D/Classes/PaperRenderComponent.h b/Engine/Plugins/Experimental/Paper2D/Source/Paper2D/Classes/PaperRenderComponent.h
--- a/Engine/Plugins/Experimental/Paper2D/Source/Paper2D/Classes/PaperRenderComponent.h
+++ b/Engine/Plugins/Experimental/Paper2D/Source/Paper2D/Classes/PaperRenderComponent.h
@@ -2,11 +2,10 @@
 
 #pragma once
 
-#include "PrimitiveComponent2D.h"
 #include "PaperRenderComponent.generated.h"
 
 UCLASS(DependsOn=UPaperSprite, MinimalAPI, ShowCategories=(Mobility), meta=(BlueprintSpawnableComponent))
-class UPaperRenderComponent : public UPrimitiveComponent2D
+class UPaperRenderComponent : public UPrimitiveComponent
 {
 	GENERATED_UCLASS_BODY()
 
diff --git a/Engine/Plugins/Experimental/Paper2D/Source/Paper2D/Classes/PrimitiveComponent2D.h b/Engine/Plugins/Experimental/Paper2D/Source/Paper2D/Classes/PrimitiveComponent2D.h
deleted file mode 100644
diff --git a/Engine/Plugins/Experimental/Paper2D/Source/Paper2D/Private/Physics/PrimitiveComponent2D.cpp b/Engine/Plugins/Experimental/Paper2D/Source/Paper2D/Private/Physics/PrimitiveComponent2D.cpp
deleted file mode 100644
