Paper2D: Switched how Paper2D chooses Box2D

diff --git a/Engine/Plugins/Experimental/Paper2D/Source/Paper2D/Paper2D.Build.cs b/Engine/Plugins/Experimental/Paper2D/Source/Paper2D/Paper2D.Build.cs
--- a/Engine/Plugins/Experimental/Paper2D/Source/Paper2D/Paper2D.Build.cs
+++ b/Engine/Plugins/Experimental/Paper2D/Source/Paper2D/Paper2D.Build.cs
@@ -32,7 +32,6 @@ public class Paper2D : ModuleRules
 			PrivateDependencyModuleNames.Add("UnrealEd");
 		}
 
-		// Link Box2D in privately; don't want to expose the headers externally
-		PrivateDependencyModuleNames.Add("Box2D");
+		SetupModuleBox2DSupport(Target);
 	}
 }
