Paper2D: Replaced FTextureReader with FBitmap as they share a great deal of functionality. Added function to detect all island rects in a texture.

diff --git a/Engine/Plugins/Experimental/Paper2D/Source/Paper2D/Private/BitmapUtils.h b/Engine/Plugins/Experimental/Paper2D/Source/Paper2D/Private/BitmapUtils.h
--- a/Engine/Plugins/Experimental/Paper2D/Source/Paper2D/Private/BitmapUtils.h
+++ b/Engine/Plugins/Experimental/Paper2D/Source/Paper2D/Private/BitmapUtils.h
@@ -7,7 +7,7 @@
 
 struct FBitmap
 {
-	FBitmap(UTexture2D* SourceTexture, int AlphaThreshold, int DefaultValue)
+	FBitmap(UTexture2D* SourceTexture, int AlphaThreshold, uint8 DefaultValue = 0)
 		: Width(0)
 		, Height(0)
 		, DefaultValue(DefaultValue)
@@ -74,7 +74,7 @@ struct FBitmap
 	}
 
 	// Create an empty bitmap
-	FBitmap(int Width, int Height, int DefaultValue)
+	FBitmap(int Width, int Height, int DefaultValue = 0)
 		: Width(Width)
 		, Height(Height)
 		, DefaultValue(DefaultValue)
@@ -119,7 +119,7 @@ struct FBitmap
 
 	bool IsColumnEmpty(int X, int Y0, int Y1)
 	{
-		for (int Y = Y0; Y < Y1; ++Y)
+		for (int Y = Y0; Y <= Y1; ++Y)
 		{
 			if (GetPixel(X, Y) != 0)
 			{
@@ -131,7 +131,7 @@ struct FBitmap
 
 	bool IsRowEmpty(int X0, int X1, int Y)
 	{
-		for (int X = X0; X < X1; ++X)
+		for (int X = X0; X <= X1; ++X)
 		{
 			if (GetPixel(X, Y) != 0)
 			{
@@ -142,7 +142,7 @@ struct FBitmap
 	}
 
 	// Returns the tight bounding box around pixels that are not 0
-	void GetTightBounds(int& OutX, int& OutY, int& OutWidth, int& OutHeight)
+	void GetTightBounds(FIntPoint& Origin, FIntPoint& Dimension)
 	{
 		int Top = 0;
 		int Bottom = Height - 1;
@@ -165,10 +165,10 @@ struct FBitmap
 			--Right;
 		}
 
-		OutX = Left;
-		OutY = Top;
-		OutWidth = Right - Left + 1;
-		OutHeight = Bottom - Top + 1;
+		Origin.X = Left;
+		Origin.Y = Top;
+		Dimension.X = Right - Left + 1;
+		Dimension.Y = Bottom - Top + 1;
 	}
 
 	// Performs a flood fill on the target bitmap, with the boundary defined
@@ -206,25 +206,27 @@ struct FBitmap
 	}
 
 	// Walk the border pixels to find any intersecting islands we haven't already filled in the mask bitmap
-	bool HasOverlappingIsland(FBitmap& MaskBitmap, int X0, int Y0, int Width, int Height, int& FillX, int& FillY)
+	bool HasOverlappingIsland(FBitmap& MaskBitmap, const FIntPoint& Origin, const FIntPoint& Dimension, FIntPoint& OutFill)
 	{
-		FillX = 0;
-		FillY = 0;
-		int X1 = X0 + Width - 1;
-		int Y1 = Y0 + Height - 1;
+		OutFill.X = 0;
+		OutFill.Y = 0;
+		const int X0 = Origin.X;
+		const int Y0 = Origin.Y;
+		const int X1 = Origin.X + Dimension.X - 1;
+		const int Y1 = Origin.Y + Dimension.Y - 1;
 
 		for (int X = X0; X <= X1; ++X)
 		{
 			if (MaskBitmap.GetPixel(X, Y0) == 0 && GetPixel(X, Y0) != 0)
 			{
-				FillX = X;
-				FillY = Y0;
+				OutFill.X = X;
+				OutFill.Y = Y0;
 				return true;
 			}
 			if (MaskBitmap.GetPixel(X, Y1) == 0 && GetPixel(X, Y1) != 0)
 			{
-				FillX = X;
-				FillY = Y1;
+				OutFill.X = X;
+				OutFill.Y = Y1;
 				return true;
 			}
 		}
@@ -233,14 +235,14 @@ struct FBitmap
 		{
 			if (MaskBitmap.GetPixel(X0, Y) == 0 && GetPixel(X0, Y) != 0)
 			{
-				FillX = X0;
-				FillY = Y;
+				OutFill.X = X0;
+				OutFill.Y = Y;
 				return true;
 			}
 			if (MaskBitmap.GetPixel(X1, Y) == 0 && GetPixel(X1, Y) != 0)
 			{
-				FillX = X1;
-				FillY = Y;
+				OutFill.X = X1;
+				OutFill.Y = Y;
 				return true;
 			}
 		}
@@ -251,10 +253,10 @@ struct FBitmap
 	// Sets pixels to 1 marking the rectangle outline
 	void DrawRectOutline(int StartX, int StartY, int Width, int Height)
 	{
-		int X0 = StartX;
-		int Y0 = StartY;
-		int X1 = StartX + Width - 1;
-		int Y1 = StartY + Height - 1;
+		const int X0 = StartX;
+		const int Y0 = StartY;
+		const int X1 = StartX + Width - 1;
+		const int Y1 = StartY + Height - 1;
 		for (int Y = Y0; Y <= Y1; ++Y)
 		{
 			SetPixel(X0, Y, 1);
@@ -267,17 +269,32 @@ struct FBitmap
 		}
 	}
 
+	void FillRect(int StartX, int StartY, int Width, int Height)
+	{
+		const int X0 = StartX;
+		const int Y0 = StartY;
+		const int X1 = StartX + Width - 1;
+		const int Y1 = StartY + Height - 1;
+		for (int Y = Y0; Y <= Y1; ++Y)
+		{
+			for (int X = X0; X <= X1; ++X)
+			{
+				SetPixel(X, Y, 1);
+			}
+		}
+	}
+
 	// Finds the rect of the contour of the shape clicked on, extended by rectangles to support
 	// separated but intersecting islands.
-	bool HasConnectedRect(int X, int Y, bool Extend, int& OutOriginX, int& OutOriginY, int& OutDimensionX, int& OutDimensionY)
+	bool HasConnectedRect(const int X, const int Y, bool Extend, FIntPoint& OutOrigin, FIntPoint& OutDimension)
 	{
 		if (GetPixel(X, Y) == 0) 
 		{
 			// selected an empty pixel
-			OutOriginX = 0;
-			OutOriginY = 0;
-			OutDimensionX = 0;
-			OutDimensionY = 0;
+			OutOrigin.X = 0;
+			OutOrigin.Y = 0;
+			OutDimension.X = 0;
+			OutDimension.Y = 0;
 			return false;
 		}
 
@@ -286,7 +303,7 @@ struct FBitmap
 		FBitmap MaskBitmap(Width, Height, 0);
 		if (Extend)
 		{
-			MaskBitmap.DrawRectOutline(OutOriginX, OutOriginY, OutDimensionX, OutDimensionY);
+			MaskBitmap.DrawRectOutline(OutOrigin.X, OutOrigin.Y, OutDimension.X, OutDimension.Y);
 		}
 
 		// MaxPasses shouldn't be nessecary, but worst case interlocked pixel
@@ -294,31 +311,33 @@ struct FBitmap
 		// reduce these problems, but may not be desirable in all cases.
 		int NumPasses = 0;
 		int MaxPasses = 40;
-		OutOriginX = 0;
-		OutOriginY = 0;
-		OutDimensionX = 0;
-		OutDimensionY = 0;
-		int FillPointX = X;
-		int FillPointY = Y;
+		FIntPoint Origin(0, 0);
+		FIntPoint Dimension(0, 0);
+		FIntPoint FillPoint(X, Y);
 		do
 		{
 			// This is probably going to be a botleneck at larger texture sizes
 			// A contour tracing algorithm will probably suffice here.
-			FloodFill(MaskBitmap, FillPointX, FillPointY);
-			MaskBitmap.GetTightBounds(OutOriginX, OutOriginY, OutDimensionX, OutDimensionY);
-		} while (NumPasses++ < MaxPasses && HasOverlappingIsland(MaskBitmap, OutOriginX, OutOriginY, OutDimensionX, OutDimensionY, /*out*/FillPointX, /*out*/FillPointY));
+			FloodFill(MaskBitmap, FillPoint.X, FillPoint.Y);
+			MaskBitmap.GetTightBounds(/*out*/Origin, /*out*/Dimension);
+		} while (NumPasses++ < MaxPasses && HasOverlappingIsland(MaskBitmap, Origin, Dimension, /*out*/FillPoint));
+
+		checkSlow(Dimension.X > 0 && Dimension.Y > 0);
+
+		OutOrigin = Origin;
+		OutDimension = Dimension;
 
 		return true;
 	}
 
 	// Wind through the bitmap from StartX, StartY to find the closest hit point
-	bool FoundClosestValidPoint(int StartX, int StartY, int MaxAllowedSearchDistance, int& HitX, int& HitY)
+	bool FoundClosestValidPoint(int StartX, int StartY, const int MaxAllowedSearchDistance, FIntPoint& OutHitPoint)
 	{
 		// Constrain the point within the image bounds
 		StartX = FMath::Clamp(StartX, 0, Width - 1);
 		StartY = FMath::Clamp(StartY, 0, Height - 1);
-		HitX = StartX;
-		HitY = StartY;
+		OutHitPoint.X = StartX;
+		OutHitPoint.Y = StartY;
 
 		// Should probably calculate a better max based on StartX and StartY
 		int RequiredSearchDistance = (FMath::Max(Width, Height) + 1) / 2;
@@ -335,14 +354,14 @@ struct FBitmap
 			{
 				if (GetPixel(X0, Y) != 0) 
 				{
-					HitX = X0;
-					HitY = Y;
+					OutHitPoint.X = X0;
+					OutHitPoint.Y = Y;
 					return true;
 				}
 				if (GetPixel(X1, Y) != 0) 
 				{
-					HitX = X1;
-					HitY = Y;
+					OutHitPoint.X = X1;
+					OutHitPoint.Y = Y;
 					return true;
 				}
 			}
@@ -350,14 +369,14 @@ struct FBitmap
 			{
 				if (GetPixel(X, Y0) != 0) 
 				{
-					HitX = X;
-					HitY = Y0;
+					OutHitPoint.X = X;
+					OutHitPoint.Y = Y0;
 					return true;
 				}
 				if (GetPixel(X, Y1) != 0) 
 				{
-					HitX = X;
-					HitY = Y1;
+					OutHitPoint.X = X;
+					OutHitPoint.Y = Y1;
 					return true;
 				}
 			}
@@ -366,10 +385,60 @@ struct FBitmap
 		return false;
 	}
 
+	// Detects all valid rects in this bitmap
+	void ExtractRects(TArray<FIntRect>& OutRects)
+	{
+		FBitmap MaskBitmap(Width, Height);
+
+		for (int Y = 0; Y < Height; ++Y)
+		{
+			for (int X = 0; X < Width; ++X)
+			{
+				if (MaskBitmap.GetPixel(X, Y) == 0 && GetPixel(X, Y) != 0)
+				{
+					FIntPoint Origin;
+					FIntPoint Dimension;
+
+					// Found something we don't already know of in the mask
+					if (HasConnectedRect(X, Y, false, /*out*/Origin, /*out*/Dimension))
+					{
+						FIntRect NewRect(Origin, Origin + Dimension);
+						bool bHasOverlap = false;
+						do
+						{
+							bHasOverlap = false;
+							for (int OutRectIndex = 0; OutRectIndex < OutRects.Num(); ++OutRectIndex)
+							{
+								FIntRect& ExistingRect = OutRects[OutRectIndex];
+								bool bRectsOverlapping = NewRect.Max.X > ExistingRect.Min.X
+														&& NewRect.Min.X < ExistingRect.Max.X
+														&& NewRect.Max.Y > ExistingRect.Min.Y
+														&& NewRect.Min.Y < ExistingRect.Max.Y;
+								if (bRectsOverlapping)
+								{
+									NewRect = FIntRect(FMath::Min(NewRect.Min.X, ExistingRect.Min.X), FMath::Min(NewRect.Min.Y, ExistingRect.Min.Y),
+													   FMath::Max(NewRect.Max.X, ExistingRect.Max.X), FMath::Max(NewRect.Max.Y, ExistingRect.Max.Y));
+									OutRects.RemoveAtSwap(OutRectIndex, 1, false);
+									break;
+								}
+							}
+						} while (bHasOverlap);
+
+						// Mark in the mask, avoid checking for any more sprites within this rect
+						MaskBitmap.FillRect(NewRect.Min.X, NewRect.Min.Y, NewRect.Width(), NewRect.Height());
+						checkSlow(NewRect.Width() > 0 && NewRect.Height() > 0);
+
+						OutRects.Add(NewRect);
+					}
+				}
+			}
+		}
+	}
+
 	TArray<uint8> RawData;
 	int32 Width;
 	int32 Height;
-	int32 DefaultValue;
+	uint8 DefaultValue;
 };
 
 
diff --git a/Engine/Plugins/Experimental/Paper2D/Source/Paper2D/Private/PaperSprite.cpp b/Engine/Plugins/Experimental/Paper2D/Source/Paper2D/Private/PaperSprite.cpp
--- a/Engine/Plugins/Experimental/Paper2D/Source/Paper2D/Private/PaperSprite.cpp
+++ b/Engine/Plugins/Experimental/Paper2D/Source/Paper2D/Private/PaperSprite.cpp
@@ -12,144 +12,6 @@
 #include "GeomTools.h"
 #include "BitmapUtils.h"
 
-//////////////////////////////////////////////////////////////////////////
-// FTextureReader
-
-struct FTextureReader
-{
-	FTextureReader(UTexture2D* SourceTexture)
-		: Width(0)
-		, Height(0)
-		, BytesPerPixel(0)
-		, TransparentBlack(0,0,0,0)
-		, AlphaThreshold(0)
-		, X0(0)
-		, Y0(0)
-		, X1(0)
-		, Y1(0)
-	{
-		// use the source art if it exists
-		FTextureSource* TextureSource = NULL;
-		if ((SourceTexture != NULL) && SourceTexture->Source.IsValid())
-		{
-			switch (SourceTexture->Source.GetFormat())
-			{
-			case TSF_G8:
-			case TSF_BGRA8:
-				TextureSource = &(SourceTexture->Source);
-				break;
-			default:
-				break;
-			};
-		}
-
-		if (TextureSource != NULL)
-		{
-			TextureSource->GetMipData(RawData, 0);
-			Width = TextureSource->GetSizeX();
-			Height = TextureSource->GetSizeY();
-			BytesPerPixel = TextureSource->GetBytesPerPixel();
-			PixelFormat = TextureSource->GetFormat();
-
-			X0 = 0;
-			Y0 = 0;
-			X1 = Width-1;
-			Y1 = Height-1;
-		}
-	}
-
-	void SetBounds(int32 NewX0, int32 NewY0, int32 NewX1, int32 NewY1)
-	{
-		X0 = NewX0;
-		Y0 = NewY0;
-		X1 = NewX1;
-		Y1 = NewY1;
-	}
-
-	bool IsValid() const
-	{
-		return (Source != NULL);
-	}
-
-	FColor GetColor(int32 X, int32 Y) const
-	{
-		if ((X < 0) || (Y < 0) || (X >= Width) || (Y >= Height))
-		{
-			return TransparentBlack;
-		}
-
-		int32 PixelByteOffset = (X + Y * Width) * BytesPerPixel;
-		const uint8* PixelPtr = RawData.GetTypedData() + PixelByteOffset;
-
-		if (PixelFormat == TSF_BGRA8)
-		{
-			return *((FColor*)PixelPtr);
-		}
-		else
-		{
-			checkSlow(PixelFormat == TSF_G8);
-			const uint8 Intensity = *PixelPtr;
-			return FColor(Intensity, Intensity, Intensity, Intensity);
-		}
-	}
-
-	bool PixelPassesThreshold(int32 X, int32 Y) const
-	{
-		return GetColor(X, Y).A > AlphaThreshold;
-	}
-
-	bool PixelPassesThreshold(int32 X, int32 Y, bool bOutOfBoundsValue) const
-	{
-		if ((X < X0) || (Y < Y0) || (X > X1) || (Y >= Y1))
-		{
-			return bOutOfBoundsValue;
-		}
-		else
-		{
-			return GetColor(X, Y).A > AlphaThreshold;
-		}
-	}
-public:
-	FTextureSource* Source;
-	TArray<uint8> RawData;
-	int32 Width;
-	int32 Height;
-	int32 BytesPerPixel;
-	ETextureSourceFormat PixelFormat;
-	FColor TransparentBlack;
-	int32 AlphaThreshold;
-
-	// Bounds for the second form of PixelPassesThreshold
-	int32 X0;
-	int32 Y0;
-	int32 X1;
-	int32 Y1;
-};
-
-bool IsClearRow(FTextureReader& Reader, int32 X0, int32 X1, int32 Y)
-{
-	for (int32 X = X0; X <= X1; ++X)
-	{
-		if (Reader.PixelPassesThreshold(X, Y))
-		{
-			return false;
-		}
-	}
-	return true;
-}
-
-bool IsClearColumn(FTextureReader& Reader, int32 X, int32 Y0, int32 Y1)
-{
-	for (int32 Y = Y0; Y <= Y1; ++Y)
-	{
-		if (Reader.PixelPassesThreshold(X, Y))
-		{
-			return false;
-		}
-	}
-	return true;
-}
-
 //////////////////////////////////////////////////////////////////////////
 // maf
 
@@ -346,24 +208,20 @@ struct FBoundaryImage
 
 void UPaperSprite::ExtractSourceRegionFromTexturePoint(const FVector2D& SourcePoint)
 {
-	int SourceX = FMath::RoundToInt(SourcePoint.X);
-	int SourceY = FMath::RoundToInt(SourcePoint.Y);
-	int ClosestValidX = 0;
-	int ClosestValidY = 0;
+	FIntPoint SourceIntPoint(FMath::RoundToInt(SourcePoint.X), FMath::RoundToInt(SourcePoint.Y));
+	FIntPoint ClosestValidPoint;
 
 	FBitmap Bitmap(SourceTexture, 0, 0);
-	if (Bitmap.IsValid() && Bitmap.FoundClosestValidPoint(SourceX, SourceY, 10, /*out*/ ClosestValidX, /*out*/ ClosestValidY))
+	if (Bitmap.IsValid() && Bitmap.FoundClosestValidPoint(SourceIntPoint.X, SourceIntPoint.Y, 10, /*out*/ClosestValidPoint))
 	{
-		int OriginX = 0;
-		int OriginY = 0;
-		int DimensionX = 0;
-		int DimensionY = 0;
-		if (Bitmap.HasConnectedRect(ClosestValidX, ClosestValidY, false, /*out*/ OriginX, /*out*/ OriginY, /*out*/ DimensionX, /*out*/ DimensionY))
+		FIntPoint Origin;
+		FIntPoint Dimension;
+		if (Bitmap.HasConnectedRect(ClosestValidPoint.X, ClosestValidPoint.Y, false, /*out*/ Origin, /*out*/ Dimension))
 		{
-			if (DimensionX > 0 && DimensionY > 0)
+			if (Dimension.X > 0 && Dimension.Y > 0)
 			{
-				SourceUV = FVector2D(OriginX, OriginY);
-				SourceDimension = FVector2D(DimensionX, DimensionY);
+				SourceUV = FVector2D(Origin.X, Origin.Y);
+				SourceDimension = FVector2D(Dimension.X, Dimension.Y);
 				PostEditChange();
 			}
 		}
@@ -714,37 +572,36 @@ void UPaperSprite::FindTextureBoundingBox(float AlphaThreshold, /*out*/ FVector2
 	int32 BottomBound = (int32)(SourceUV.Y + SourceDimension.Y - 1);
 
 
-	FTextureReader Reader(SourceTexture);
-	if (Reader.IsValid())
+	int32 AlphaThresholdInt = FMath::Clamp<int32>(AlphaThreshold * 255, 0, 255);
+	FBitmap SourceBitmap(SourceTexture, AlphaThresholdInt);
+	if (SourceBitmap.IsValid())
 	{
-		Reader.AlphaThreshold = FMath::Clamp<int32>(AlphaThreshold * 255, 0, 255);
-
 		// Make sure the initial bounds starts in the texture
-		LeftBound = FMath::Clamp(LeftBound, 0, Reader.Width-1);
-		RightBound = FMath::Clamp(RightBound, 0, Reader.Width-1);
-		TopBound = FMath::Clamp(TopBound, 0, Reader.Height-1);
-		BottomBound = FMath::Clamp(BottomBound, 0, Reader.Height-1);
+		LeftBound = FMath::Clamp(LeftBound, 0, SourceBitmap.Width-1);
+		RightBound = FMath::Clamp(RightBound, 0, SourceBitmap.Width-1);
+		TopBound = FMath::Clamp(TopBound, 0, SourceBitmap.Height-1);
+		BottomBound = FMath::Clamp(BottomBound, 0, SourceBitmap.Height-1);
 
 		// Pull it in from the top
-		while ((TopBound < BottomBound) && IsClearRow(Reader, LeftBound, RightBound, TopBound))
+		while ((TopBound < BottomBound) && SourceBitmap.IsRowEmpty(LeftBound, RightBound, TopBound))
 		{
 			++TopBound;
 		}
 
 		// Pull it in from the bottom
-		while ((BottomBound > TopBound) && IsClearRow(Reader, LeftBound, RightBound, BottomBound))
+		while ((BottomBound > TopBound) && SourceBitmap.IsRowEmpty(LeftBound, RightBound, BottomBound))
 		{
 			--BottomBound;
 		}
 
 		// Pull it in from the left
-		while ((LeftBound < RightBound) && IsClearColumn(Reader, LeftBound, TopBound, BottomBound))
+		while ((LeftBound < RightBound) && SourceBitmap.IsColumnEmpty(LeftBound, TopBound, BottomBound))
 		{
 			++LeftBound;
 		}
 
 		// Pull it in from the right
-		while ((RightBound > LeftBound) && IsClearColumn(Reader, RightBound, TopBound, BottomBound))
+		while ((RightBound > LeftBound) && SourceBitmap.IsColumnEmpty(RightBound, TopBound, BottomBound))
 		{
 			--RightBound;
 		}
@@ -829,13 +686,11 @@ void UPaperSprite::FindContours(const FIntPoint& ScanPos, const FIntPoint& ScanS
 		4, //from7
 	};
 
-	FTextureReader Reader(Texture);
-	if (Reader.IsValid())
+	int32 AlphaThresholdInt = FMath::Clamp<int32>(AlphaThreshold * 255, 0, 255);
+	FBitmap SourceBitmap(Texture, AlphaThresholdInt);
+	if (SourceBitmap.IsValid())
 	{
-		checkSlow((LeftBound >= 0) && (TopBound >= 0) && (RightBound < Reader.Width) && (BottomBound < Reader.Height));
-		Reader.SetBounds(LeftBound, TopBound, RightBound, BottomBound);
-
-		Reader.AlphaThreshold = FMath::Clamp<int32>(AlphaThreshold * 255, 0, 255);
+		checkSlow((LeftBound >= 0) && (TopBound >= 0) && (RightBound < SourceBitmap.Width) && (BottomBound < SourceBitmap.Height));
 
 		// Create the 'output' boundary image
 		FBoundaryImage BoundaryImage(ScanPos, ScanSize);
@@ -847,7 +702,8 @@ void UPaperSprite::FindContours(const FIntPoint& ScanPos, const FIntPoint& ScanS
 			for (int32 X = LeftBound - 1; X < RightBound + 2; ++X)
 			{
 				const bool bAlreadyTaggedAsBoundary = BoundaryImage.GetPixel(X, Y) > 0;
-				const bool bIsFilledPixel = Reader.PixelPassesThreshold(X, Y, false);
+				const bool bPixelInsideBounds = (X >= LeftBound && X <= RightBound && Y >= TopBound && Y <= BottomBound);
+				const bool bIsFilledPixel = bPixelInsideBounds && SourceBitmap.GetPixel(X, Y) != 0;
 
 				if (bInsideBoundary)
 				{
@@ -886,7 +742,8 @@ void UPaperSprite::FindContours(const FIntPoint& ScanPos, const FIntPoint& ScanS
 							// Test pixel (clockwise from the current pixel)
 							const int32 CX = PX + NeighborX[NeighborPhase];
 							const int32 CY = PY + NeighborY[NeighborPhase];
-							const bool bTestPixelPasses = Reader.PixelPassesThreshold(CX, CY, false);
+							const bool bTestPixelInsideBounds = (CX >= LeftBound && CX <= RightBound && CY >= TopBound && CY <= BottomBound);
+							const bool bTestPixelPasses = bTestPixelInsideBounds && SourceBitmap.GetPixel(CX, CY) != 0;
 
 							UE_LOG(LogPaper2D, Log, TEXT("Outer P(%d,%d), C(%d,%d) Ph%d %s"),
 								PX, PY, CX, CY, NeighborPhase, bTestPixelPasses ? TEXT("[BORDER]") : TEXT("[]"));
