Paper2D: Fix SpriteColor not taking effect

diff --git a/Engine/Plugins/Experimental/Paper2D/Source/Paper2D/Private/PaperRenderComponent.cpp b/Engine/Plugins/Experimental/Paper2D/Source/Paper2D/Private/PaperRenderComponent.cpp
--- a/Engine/Plugins/Experimental/Paper2D/Source/Paper2D/Private/PaperRenderComponent.cpp
+++ b/Engine/Plugins/Experimental/Paper2D/Source/Paper2D/Private/PaperRenderComponent.cpp
@@ -34,6 +34,7 @@ FPrimitiveSceneProxy* UPaperRenderComponent::CreateSceneProxy()
 	FPaperRenderSceneProxy* NewProxy = new FPaperRenderSceneProxy(this);
 	FSpriteDrawCallRecord DrawCall;
 	DrawCall.BuildFromSprite(SourceSprite);
+	DrawCall.Color = SpriteColor;
 	NewProxy->SetDrawCall_RenderThread(DrawCall);
 	return NewProxy;
 }
