Paper2D: Switch mobility setup for sprite components to match static mesh components; UPaperRenderComponent defaults to Movable and APaperRenderActor defaults to Static

diff --git a/Engine/Plugins/Experimental/Paper2D/Source/Paper2D/Private/PaperRenderActor.cpp b/Engine/Plugins/Experimental/Paper2D/Source/Paper2D/Private/PaperRenderActor.cpp
--- a/Engine/Plugins/Experimental/Paper2D/Source/Paper2D/Private/PaperRenderActor.cpp
+++ b/Engine/Plugins/Experimental/Paper2D/Source/Paper2D/Private/PaperRenderActor.cpp
@@ -10,6 +10,8 @@ APaperRenderActor::APaperRenderActor(const FPostConstructInitializeProperties& P
 	: Super(PCIP)
 {
 	RenderComponent = PCIP.CreateDefaultSubobject<UPaperRenderComponent>(this, TEXT("RenderComponent"));
+	RenderComponent->SetCollisionProfileName(UCollisionProfile::BlockAll_ProfileName);
+	RenderComponent->Mobility = EComponentMobility::Static;
 
 	RootComponent = RenderComponent;
 }
diff --git a/Engine/Plugins/Experimental/Paper2D/Source/Paper2D/Private/PaperRenderComponent.cpp b/Engine/Plugins/Experimental/Paper2D/Source/Paper2D/Private/PaperRenderComponent.cpp
--- a/Engine/Plugins/Experimental/Paper2D/Source/Paper2D/Private/PaperRenderComponent.cpp
+++ b/Engine/Plugins/Experimental/Paper2D/Source/Paper2D/Private/PaperRenderComponent.cpp
@@ -16,8 +16,6 @@ UPaperRenderComponent::UPaperRenderComponent(const FPostConstructInitializePrope
 	MaterialOverride = nullptr;
 
 	SpriteColor = FLinearColor::White;
-
-	Mobility = EComponentMobility::Static;
 }
 
 #if WITH_EDITOR
