Paper2D: Fix flipbook components not rendering

diff --git a/Engine/Plugins/Experimental/Paper2D/Source/Paper2D/Private/PaperAnimatedRenderComponent.cpp b/Engine/Plugins/Experimental/Paper2D/Source/Paper2D/Private/PaperAnimatedRenderComponent.cpp
--- a/Engine/Plugins/Experimental/Paper2D/Source/Paper2D/Private/PaperAnimatedRenderComponent.cpp
+++ b/Engine/Plugins/Experimental/Paper2D/Source/Paper2D/Private/PaperAnimatedRenderComponent.cpp
@@ -30,7 +30,20 @@ UPaperAnimatedRenderComponent::UPaperAnimatedRenderComponent(const FPostConstruc
 
 FPrimitiveSceneProxy* UPaperAnimatedRenderComponent::CreateSceneProxy()
 {
-	return new FPaperAnimatedRenderSceneProxy(this);
+	FPaperAnimatedRenderSceneProxy* NewProxy = new FPaperAnimatedRenderSceneProxy(this);
+
+	CalculateCurrentFrame();
+	UPaperSprite* SpriteToSend = NULL;
+	if ((SourceFlipbook != NULL) && (CachedFrameIndex >= 0) && (CachedFrameIndex < SourceFlipbook->KeyFrames.Num()))
+	{
+		SpriteToSend = SourceFlipbook->KeyFrames[CachedFrameIndex].Sprite;
+	}
+
+	FSpriteDrawCallRecord DrawCall;
+	DrawCall.BuildFromSprite(SpriteToSend);
+	DrawCall.Color = SpriteColor;
+	NewProxy->SetDrawCall_RenderThread(DrawCall);
+	return NewProxy;
 }
 
 FBoxSphereBounds UPaperAnimatedRenderComponent::CalcBounds(const FTransform & LocalToWorld) const
diff --git a/Engine/Plugins/Experimental/Paper2D/Source/Paper2D/Private/PaperAnimatedRenderSceneProxy.cpp b/Engine/Plugins/Experimental/Paper2D/Source/Paper2D/Private/PaperAnimatedRenderSceneProxy.cpp
--- a/Engine/Plugins/Experimental/Paper2D/Source/Paper2D/Private/PaperAnimatedRenderSceneProxy.cpp
+++ b/Engine/Plugins/Experimental/Paper2D/Source/Paper2D/Private/PaperAnimatedRenderSceneProxy.cpp
@@ -9,5 +9,12 @@
 FPaperAnimatedRenderSceneProxy::FPaperAnimatedRenderSceneProxy(const UPaperAnimatedRenderComponent* InComponent)
 	: FPaperRenderSceneProxy(InComponent)
 {
+	//@TODO: PAPER2D: WireframeColor = RenderComp->GetWireframeColor();
+
 	Material = InComponent->GetSpriteMaterial();
+
+	if (Material)
+	{
+		MaterialRelevance = Material->GetRelevance();
+	}
 }
