Paper2D: Fix tile map actors not rendering

diff --git a/Engine/Plugins/Experimental/Paper2D/Source/Paper2D/Private/PaperTileMapRenderSceneProxy.cpp b/Engine/Plugins/Experimental/Paper2D/Source/Paper2D/Private/PaperTileMapRenderSceneProxy.cpp
--- a/Engine/Plugins/Experimental/Paper2D/Source/Paper2D/Private/PaperTileMapRenderSceneProxy.cpp
+++ b/Engine/Plugins/Experimental/Paper2D/Source/Paper2D/Private/PaperTileMapRenderSceneProxy.cpp
@@ -15,6 +15,11 @@ FPaperTileMapRenderSceneProxy::FPaperTileMapRenderSceneProxy(const UPaperTileMap
 	{
 		TileComponent = InTileComponent;
 		Material = TileComponent->Material;
+
+		if (Material)
+		{
+			MaterialRelevance = Material->GetRelevance();
+		}
 	}
 }
 
