Paper2D: Moved tile map data from the component to a new class UPaperTileMap, and added a per-cell tile set pointer

diff --git a/Engine/Plugins/Experimental/Paper2D/Source/Paper2D/Classes/PaperTileLayer.h b/Engine/Plugins/Experimental/Paper2D/Source/Paper2D/Classes/PaperTileLayer.h
--- a/Engine/Plugins/Experimental/Paper2D/Source/Paper2D/Classes/PaperTileLayer.h
+++ b/Engine/Plugins/Experimental/Paper2D/Source/Paper2D/Classes/PaperTileLayer.h
@@ -4,6 +4,34 @@
 
 #include "PaperTileLayer.generated.h"
 
+USTRUCT(BlueprintType)
+struct FPaperTileInfo
+{
+	GENERATED_USTRUCT_BODY()
+
+	UPROPERTY()
+	UPaperTileSet* TileSet;
+
+	UPROPERTY()
+	int32 PackedTileIndex;
+
+	FPaperTileInfo()
+		: TileSet(nullptr)
+		, PackedTileIndex(INDEX_NONE)
+	{
+	}
+
+	inline bool operator==(const FPaperTileInfo& Other) const
+	{
+		return (Other.TileSet == TileSet) && (Other.PackedTileIndex == PackedTileIndex);
+	}
+
+	inline bool operator!=(const FPaperTileInfo& Other) const
+	{
+		return !(*this == Other);
+	}
+};
+
 UCLASS(DependsOn=UPaperTileSet)
 class PAPER2D_API UPaperTileLayer : public UObject
 {
@@ -38,19 +66,24 @@ protected:
 	int32 AllocatedHeight;
 
 	UPROPERTY()
-	TArray<int32> AllocatedGrid;
+	TArray<int32> AllocatedGrid_DEPRECATED;
+
+	UPROPERTY()
+	TArray<FPaperTileInfo> AllocatedCells;
 
 public:
+	void ConvertToTileSetPerCell();
+
 	// UObject interface
 #if WITH_EDITORONLY_DATA
 	virtual void PostEditChangeProperty(struct FPropertyChangedEvent& PropertyChangedEvent) OVERRIDE;
 #endif
 	// End of UObject interface
 
-	class UPaperTileMapRenderComponent* GetTileMap() const;
+	class UPaperTileMap* GetTileMap() const;
 
-	int32 GetCell(int32 X, int32 Y) const;
-	void SetCell(int32 X, int32 Y, int32 NewValue);
+	FPaperTileInfo GetCell(int32 X, int32 Y) const;
+	void SetCell(int32 X, int32 Y, const FPaperTileInfo& NewValue);
 
 	// This is a destructive operation!
 	void DestructiveAllocateMap(int32 NewWidth, int32 NewHeight);
diff --git a/Engine/Plugins/Experimental/Paper2D/Source/Paper2D/Classes/PaperTileMap.h b/Engine/Plugins/Experimental/Paper2D/Source/Paper2D/Classes/PaperTileMap.h
new file mode 100644
--- /dev/null
+++ b/Engine/Plugins/Experimental/Paper2D/Source/Paper2D/Classes/PaperTileMap.h
@@ -0,0 +1,66 @@
+// Copyright 1998-2014 Epic Games, Inc. All Rights Reserved.
+
+#pragma once
+
+#include "PaperSprite.h"
+#include "PaperTileMap.generated.h"
+
+UCLASS(DependsOn=UEngineTypes)
+class PAPER2D_API UPaperTileMap : public UDataAsset //@TODO: Just to make it easy to spawn for now
+{
+	GENERATED_UCLASS_BODY()
+
+	// Width of map (in tiles)
+	UPROPERTY(Category=Setup, EditAnywhere, meta=(UIMin=1, ClampMin=1))
+	int32 MapWidth;
+
+	// Height of map (in tiles)
+	UPROPERTY(Category=Setup, EditAnywhere, meta=(UIMin=1, ClampMin=1))
+	int32 MapHeight;
+
+	// Width of one tile (in pixels)
+	UPROPERTY(Category=Setup, EditAnywhere, meta=(UIMin=1, ClampMin=1))
+	int32 TileWidth;
+
+	// Height of one tile (in pixels)
+	UPROPERTY(Category=Setup, EditAnywhere, meta=(UIMin=1, ClampMin=1))
+	int32 TileHeight;
+
+	// Pixels per Unreal Unit (pixels per cm)
+	UPROPERTY(Category = Setup, EditAnywhere)
+	float PixelsPerUnit;
+
+	// Default tile set to use for new layers
+	UPROPERTY(Category=Setup, EditAnywhere)
+	UPaperTileSet* DefaultLayerTileSet;
+
+	// Test material
+	UPROPERTY(Category=Setup, EditAnywhere)
+	UMaterialInterface* Material;
+
+	// The list of layers
+	UPROPERTY()
+	TArray<UPaperTileLayer*> TileLayers;
+
+	// Collision domain (no collision, 2D, or 3D)
+	UPROPERTY(Category = Collision, EditAnywhere)
+	TEnumAsByte<ESpriteCollisionMode::Type> SpriteCollisionDomain;
+
+	// Baked physics data.
+	UPROPERTY()
+	class UBodySetup* BodySetup;
+
+public:
+	// UObject interface
+#if WITH_EDITOR
+	virtual void PostEditChangeProperty(FPropertyChangedEvent& PropertyChangedEvent) OVERRIDE;
+#endif
+	// End of UObject interface
+
+
+	FVector GetTilePositionInLocalSpace(int32 TileX, int32 TileY, int32 LayerIndex = 0) const;
+
+	FBoxSphereBounds GetRenderBounds() const;
+protected:
+	virtual void UpdateBodySetup();
+};
diff --git a/Engine/Plugins/Experimental/Paper2D/Source/Paper2D/Classes/PaperTileMapRenderComponent.h b/Engine/Plugins/Experimental/Paper2D/Source/Paper2D/Classes/PaperTileMapRenderComponent.h
--- a/Engine/Plugins/Experimental/Paper2D/Source/Paper2D/Classes/PaperTileMapRenderComponent.h
+++ b/Engine/Plugins/Experimental/Paper2D/Source/Paper2D/Classes/PaperTileMapRenderComponent.h
@@ -9,45 +9,47 @@ class PAPER2D_API UPaperTileMapRenderComponent : public UPrimitiveComponent
 {
 	GENERATED_UCLASS_BODY()
 
+private:
 	// Width of map (in tiles)
-	UPROPERTY(Category=Setup, EditAnywhere, meta=(UIMin=1, ClampMin=1))
-	int32 MapWidth;
+	UPROPERTY()
+	int32 MapWidth_DEPRECATED;
 
-	// Height of map (in tiles)
-	UPROPERTY(Category=Setup, EditAnywhere, meta=(UIMin=1, ClampMin=1))
-	int32 MapHeight;
+	UPROPERTY()
+	int32 MapHeight_DEPRECATED;
 
-	// Width of one tile (in pixels)
-	UPROPERTY(Category=Setup, EditAnywhere, meta=(UIMin=1, ClampMin=1))
-	int32 TileWidth;
+	UPROPERTY()
+	int32 TileWidth_DEPRECATED;
 
-	// Height of one tile (in pixels)
-	UPROPERTY(Category=Setup, EditAnywhere, meta=(UIMin=1, ClampMin=1))
-	int32 TileHeight;
+	UPROPERTY()
+	int32 TileHeight_DEPRECATED;
 
 	// Default tile set to use for new layers
-	UPROPERTY(Category=Setup, EditAnywhere)
-	UPaperTileSet* DefaultLayerTileSet;
+	UPROPERTY()
+	UPaperTileSet* DefaultLayerTileSet_DEPRECATED;
 
 	// Test material
-	UPROPERTY(Category=Setup, EditAnywhere)
-	UMaterialInterface* Material;
+	UPROPERTY()
+	UMaterialInterface* Material_DEPRECATED;
 
 	// The list of layers
 	UPROPERTY()
-	TArray<UPaperTileLayer*> TileLayers;
+	TArray<UPaperTileLayer*> TileLayers_DEPRECATED;
+
+public:
+	// The tile map used by this component
+	UPROPERTY(Category=Setup, EditAnywhere, EditInline, BlueprintReadOnly)
+	class UPaperTileMap* TileMap;
 
 protected:
 	friend class FPaperTileMapRenderSceneProxy;
 
-	/** Description of collision */
-	UPROPERTY(transient, duplicatetransient)
-	class UBodySetup* ShapeBodySetup;
-
 public:
 	// UObject interface
-#if WITH_EDITOR
+	virtual void PostInitProperties() OVERRIDE;
+#if WITH_EDITORONLY_DATA
 	virtual void PostEditChangeProperty(FPropertyChangedEvent& PropertyChangedEvent) OVERRIDE;
+	virtual void Serialize(FArchive& Ar) OVERRIDE;
+	virtual void PostLoad() OVERRIDE;
 #endif
 	// End of UObject interface
 
@@ -61,9 +63,4 @@ public:
 	virtual FBoxSphereBounds CalcBounds(const FTransform& LocalToWorld) const OVERRIDE;
 	virtual class UBodySetup* GetBodySetup() OVERRIDE;
 	// End of UPrimitiveComponent interface
-
-	FVector ConvertTilePositionToWorldSpace(int32 TileX, int32 TileY, int32 LayerIndex = 0) const;
-
-protected:
-	virtual void UpdateBodySetup();
 };
diff --git a/Engine/Plugins/Experimental/Paper2D/Source/Paper2D/Private/PaperCustomVersion.h b/Engine/Plugins/Experimental/Paper2D/Source/Paper2D/Private/PaperCustomVersion.h
--- a/Engine/Plugins/Experimental/Paper2D/Source/Paper2D/Private/PaperCustomVersion.h
+++ b/Engine/Plugins/Experimental/Paper2D/Source/Paper2D/Private/PaperCustomVersion.h
@@ -16,6 +16,9 @@ struct FPaperCustomVersion
 		// Removed FBodyInstance2D and moved 2D physics support into Engine (requires regenerating 2D agg geom)
 		RemovedBodyInstance2D = 2,
 
+		// Moved tile map data out of the component into UPaperTileMap
+		MovedTileMapDataToSeparateClass = 3,
+
 		// -----<new versions can be added above this line>-------------------------------------------------
 		VersionPlusOne,
 		LatestVersion = VersionPlusOne - 1
diff --git a/Engine/Plugins/Experimental/Paper2D/Source/Paper2D/Private/PaperTileLayer.cpp b/Engine/Plugins/Experimental/Paper2D/Source/Paper2D/Private/PaperTileLayer.cpp
--- a/Engine/Plugins/Experimental/Paper2D/Source/Paper2D/Private/PaperTileLayer.cpp
+++ b/Engine/Plugins/Experimental/Paper2D/Source/Paper2D/Private/PaperTileLayer.cpp
@@ -20,8 +20,8 @@ void UPaperTileLayer::DestructiveAllocateMap(int32 NewWidth, int32 NewHeight)
 	check((NewWidth > 0) && (NewHeight > 0));
 
 	const int32 NumCells = NewWidth * NewHeight;
-	AllocatedGrid.Empty(NumCells);
-	AllocatedGrid.AddZeroed(NumCells);
+	AllocatedCells.Empty(NumCells);
+	AllocatedCells.AddZeroed(NumCells);
 
 	AllocatedWidth = NewWidth;
 	AllocatedHeight = NewHeight;
@@ -38,7 +38,7 @@ void UPaperTileLayer::ReallocateAndCopyMap()
 {
 	const int32 SavedWidth = AllocatedWidth;
 	const int32 SavedHeight = AllocatedHeight;
-	TArray<int32> SavedDesignedMap(AllocatedGrid);
+	TArray<FPaperTileInfo> SavedDesignedMap(AllocatedCells);
 
 	DestructiveAllocateMap(LayerWidth, LayerHeight);
 
@@ -52,7 +52,7 @@ void UPaperTileLayer::ReallocateAndCopyMap()
 			const int32 SrcIndex = Y*SavedWidth + X;
 			const int32 DstIndex = Y*LayerWidth + X;
 
-			AllocatedGrid[DstIndex] = SavedDesignedMap[SrcIndex];
+			AllocatedCells[DstIndex] = SavedDesignedMap[SrcIndex];
 		}
 	}
 
@@ -73,7 +73,7 @@ void UPaperTileLayer::PostEditChangeProperty(struct FPropertyChangedEvent& Prope
 		// Resize the map, trying to preserve existing data
 		ReallocateAndCopyMap();
 	}
-// 	else if (PropertyName == TEXT("AllocatedGrid"))
+// 	else if (PropertyName == TEXT("AllocatedCells"))
 // 	{
 // 		BakeMap();
 // 	}
@@ -82,37 +82,37 @@ void UPaperTileLayer::PostEditChangeProperty(struct FPropertyChangedEvent& Prope
 }
 #endif
 
-UPaperTileMapRenderComponent* UPaperTileLayer::GetTileMap() const
+UPaperTileMap* UPaperTileLayer::GetTileMap() const
 {
-	return CastChecked<UPaperTileMapRenderComponent>(GetOuter());
+	return CastChecked<UPaperTileMap>(GetOuter());
 }
 
-int32 UPaperTileLayer::GetCell(int32 X, int32 Y) const
+FPaperTileInfo UPaperTileLayer::GetCell(int32 X, int32 Y) const
 {
 	if ((X < 0) || (X >= LayerWidth) || (Y < 0) || (Y >= LayerHeight))
 	{
-		return INDEX_NONE;
+		return FPaperTileInfo();
 	}
 	else
 	{
-		return AllocatedGrid[X + (Y*LayerWidth)];
+		return AllocatedCells[X + (Y*LayerWidth)];
 	}
 }
 
-void UPaperTileLayer::SetCell(int32 X, int32 Y, int32 NewValue)
+void UPaperTileLayer::SetCell(int32 X, int32 Y, const FPaperTileInfo& NewValue)
 {
 	if ((X < 0) || (X >= LayerWidth) || (Y < 0) || (Y >= LayerHeight))
 	{
 	}
 	else
 	{
-		AllocatedGrid[X + (Y*LayerWidth)] = NewValue;
+		AllocatedCells[X + (Y*LayerWidth)] = NewValue;
 	}
 }
 
 void UPaperTileLayer::AugmentBodySetup(UBodySetup* ShapeBodySetup)
 {
-	if ( bCollisionLayer )
+	if (bCollisionLayer)
 	{
 		//@TODO: Tile pivot issue
 		//@TODO: Layer thickness issue
@@ -130,7 +130,7 @@ void UPaperTileLayer::AugmentBodySetup(UBodySetup* ShapeBodySetup)
 		{
 			for (int32 YValue = 0; YValue < LayerHeight; ++YValue)
 			{
-				if ( GetCell(XValue, YValue) != 0 )
+				if (GetCell(XValue, YValue).PackedTileIndex != 0)
 				{
 					FKBoxElem* NewBox = new(ShapeBodySetup->AggGeom.BoxElems) FKBoxElem(TileWidth, TileThickness, TileHeight);
 
@@ -144,4 +144,17 @@ void UPaperTileLayer::AugmentBodySetup(UBodySetup* ShapeBodySetup)
 			}
 		}
 	}
+}
+
+void UPaperTileLayer::ConvertToTileSetPerCell()
+{
+	AllocatedCells.Empty(AllocatedGrid_DEPRECATED.Num());
+
+	const int32 NumCells = AllocatedWidth * AllocatedHeight;
+	for (int32 Index = 0; Index < NumCells; ++Index)
+	{
+		FPaperTileInfo* Info = new (AllocatedCells)FPaperTileInfo();
+		Info->TileSet = TileSet;
+		Info->PackedTileIndex = AllocatedGrid_DEPRECATED[Index];
+	}
 }
\ No newline at end of file
diff --git a/Engine/Plugins/Experimental/Paper2D/Source/Paper2D/Private/PaperTileMap.cpp b/Engine/Plugins/Experimental/Paper2D/Source/Paper2D/Private/PaperTileMap.cpp
new file mode 100644
--- /dev/null
+++ b/Engine/Plugins/Experimental/Paper2D/Source/Paper2D/Private/PaperTileMap.cpp
@@ -0,0 +1,110 @@
+// Copyright 1998-2014 Epic Games, Inc. All Rights Reserved.
+
+#include "Paper2DPrivatePCH.h"
+
+//////////////////////////////////////////////////////////////////////////
+// UPaperTileMap
+
+UPaperTileMap::UPaperTileMap(const FPostConstructInitializeProperties& PCIP)
+	: Super(PCIP)
+{
+	MapWidth = 4;
+	MapHeight = 4;
+	TileWidth = 32;
+	TileHeight = 32;
+	PixelsPerUnit = 1.0f;
+	SpriteCollisionDomain = ESpriteCollisionMode::None;
+
+	static ConstructorHelpers::FObjectFinder<UMaterial> DefaultMaterial(TEXT("/Paper2D/DefaultSpriteMaterial"));
+	Material = DefaultMaterial.Object;
+}
+
+#if WITH_EDITOR
+void UPaperTileMap::PostEditChangeProperty(FPropertyChangedEvent& PropertyChangedEvent)
+{
+	const FName PropertyName = (PropertyChangedEvent.Property != NULL) ? PropertyChangedEvent.Property->GetFName() : NAME_None;
+
+	if ((PropertyName == GET_MEMBER_NAME_CHECKED(UPaperTileMap, MapWidth)) || (PropertyName == GET_MEMBER_NAME_CHECKED(UPaperTileMap, MapHeight)))
+	{
+		MapWidth = FMath::Max(MapWidth, 1);
+		MapHeight = FMath::Max(MapHeight, 1);
+
+		// Resize all of the existing layers
+		for (int32 LayerIndex = 0; LayerIndex < TileLayers.Num(); ++LayerIndex)
+		{
+			UPaperTileLayer* TileLayer = TileLayers[LayerIndex];
+			TileLayer->ResizeMap(MapWidth, MapHeight);
+		}
+	}
+
+	if (!IsTemplate())
+	{
+		UpdateBodySetup();
+	}
+
+	Super::PostEditChangeProperty(PropertyChangedEvent);
+}
+#endif
+
+
+void UPaperTileMap::UpdateBodySetup()
+{
+	// Ensure we have the data structure for the desired collision method
+	switch (SpriteCollisionDomain)
+	{
+	case ESpriteCollisionMode::Use3DPhysics:
+		BodySetup = nullptr;
+		if (BodySetup == nullptr)
+		{
+			BodySetup = NewObject<UBodySetup>(this);
+		}
+		break;
+	case ESpriteCollisionMode::Use2DPhysics:
+		BodySetup = nullptr;
+		if (BodySetup == nullptr)
+		{
+			BodySetup = NewObject<UBodySetup2D>(this);
+		}
+		break;
+	case ESpriteCollisionMode::None:
+		BodySetup = nullptr;
+		break;
+	}
+
+	if (SpriteCollisionDomain != ESpriteCollisionMode::None)
+	{
+		if (SpriteCollisionDomain == ESpriteCollisionMode::Use3DPhysics)
+		{
+			BodySetup->CollisionTraceFlag = CTF_UseSimpleAsComplex;
+
+			BodySetup->AggGeom.BoxElems.Empty();
+			for (int32 LayerIndex = 0; LayerIndex < TileLayers.Num(); ++LayerIndex)
+			{
+				TileLayers[LayerIndex]->AugmentBodySetup(BodySetup);
+			}
+		}
+
+		//@TODO: BOX2D: Add support for 2D physics on tile maps
+	}
+}
+
+FVector UPaperTileMap::GetTilePositionInLocalSpace(int32 TileX, int32 TileY, int32 LayerIndex) const
+{
+	//@TODO: Layer depth, multiply by layer depth here!
+	//@TODO: Tile pivot issue
+	const FVector LocalPos(PaperAxisX * (TileX - 0.5f) * TileWidth + PaperAxisY * -(TileY - 0.5f) * TileHeight);
+	
+	return LocalPos;
+}
+
+FBoxSphereBounds UPaperTileMap::GetRenderBounds() const
+{
+	//@TODO: Tile pivot issue
+	//@TODO: Layer thickness issue
+	const float HalfThickness = 2.0f;
+	const FVector TopLeft((-0.5f)*TileWidth, -HalfThickness, -(MapHeight - 0.5f) * TileHeight);
+	const FVector Dimenisons(MapWidth*TileWidth, 2 * HalfThickness, MapHeight * TileHeight);
+
+	const FBox Box(TopLeft, TopLeft + Dimenisons);
+	return FBoxSphereBounds(Box);
+}
\ No newline at end of file
diff --git a/Engine/Plugins/Experimental/Paper2D/Source/Paper2D/Private/PaperTileMapActor.cpp b/Engine/Plugins/Experimental/Paper2D/Source/Paper2D/Private/PaperTileMapActor.cpp
--- a/Engine/Plugins/Experimental/Paper2D/Source/Paper2D/Private/PaperTileMapActor.cpp
+++ b/Engine/Plugins/Experimental/Paper2D/Source/Paper2D/Private/PaperTileMapActor.cpp
@@ -16,9 +16,9 @@ APaperTileMapActor::APaperTileMapActor(const FPostConstructInitializeProperties&
 #if WITH_EDITOR
 bool APaperTileMapActor::GetReferencedContentObjects(TArray<UObject*>& Objects) const
 {
-	if (RenderComponent->DefaultLayerTileSet != nullptr)
+	if (const UObject* Asset = RenderComponent->AdditionalStatObject())
 	{
-		Objects.Add(RenderComponent->DefaultLayerTileSet);
+		Objects.Add(const_cast<UObject*>(Asset));
 	}
 	return true;
 }
diff --git a/Engine/Plugins/Experimental/Paper2D/Source/Paper2D/Private/PaperTileMapRenderComponent.cpp b/Engine/Plugins/Experimental/Paper2D/Source/Paper2D/Private/PaperTileMapRenderComponent.cpp
--- a/Engine/Plugins/Experimental/Paper2D/Source/Paper2D/Private/PaperTileMapRenderComponent.cpp
+++ b/Engine/Plugins/Experimental/Paper2D/Source/Paper2D/Private/PaperTileMapRenderComponent.cpp
@@ -3,6 +3,7 @@
 #include "Paper2DPrivatePCH.h"
 #include "PaperTileMapRenderSceneProxy.h"
 #include "PhysicsEngine/BodySetup.h"
+#include "PaperCustomVersion.h"
 
 #define LOCTEXT_NAMESPACE "Paper2D"
 
@@ -14,13 +15,13 @@ UPaperTileMapRenderComponent::UPaperTileMapRenderComponent(const FPostConstructI
 {
 	BodyInstance.SetCollisionEnabled(ECollisionEnabled::QueryAndPhysics);
 
-	MapWidth = 4;
-	MapHeight = 4;
-	TileWidth = 32;
-	TileHeight = 32;
+	MapWidth_DEPRECATED = 4;
+	MapHeight_DEPRECATED = 4;
+	TileWidth_DEPRECATED = 32;
+	TileHeight_DEPRECATED = 32;
 
 	static ConstructorHelpers::FObjectFinder<UMaterial> DefaultMaterial(TEXT("/Paper2D/DefaultSpriteMaterial"));
-	Material = DefaultMaterial.Object;
+	Material_DEPRECATED = DefaultMaterial.Object;
 }
 
 FPrimitiveSceneProxy* UPaperTileMapRenderComponent::CreateSceneProxy()
@@ -28,85 +29,110 @@ FPrimitiveSceneProxy* UPaperTileMapRenderComponent::CreateSceneProxy()
 	return new FPaperTileMapRenderSceneProxy(this);
 }
 
-FBoxSphereBounds UPaperTileMapRenderComponent::CalcBounds(const FTransform& LocalToWorld) const
-{
-	//@TODO: Tile pivot issue
-	//@TODO: Layer thickness issue
-	const float HalfThickness = 2.0f;	
-	const FVector TopLeft((-0.5f)*TileWidth, -HalfThickness, -(MapHeight - 0.5f) * TileHeight);
-	const FVector Dimenisons(MapWidth*TileWidth, 2*HalfThickness, MapHeight * TileHeight);
-
-	const FBox Box(TopLeft, TopLeft + Dimenisons);
-	return FBoxSphereBounds(Box.TransformBy(LocalToWorld));
-}
-
-void UPaperTileMapRenderComponent::TickComponent(float DeltaTime, enum ELevelTick TickType, FActorComponentTickFunction* ThisTickFunction)
+void UPaperTileMapRenderComponent::PostInitProperties()
 {
-	// Indicate we need to send new dynamic data.
-	MarkRenderDynamicDataDirty();
+	TileMap = NewObject<UPaperTileMap>(this);
+	Super::PostInitProperties();
 }
 
-#if WITH_EDITOR
-void UPaperTileMapRenderComponent::PostEditChangeProperty(FPropertyChangedEvent& PropertyChangedEvent)
+FBoxSphereBounds UPaperTileMapRenderComponent::CalcBounds(const FTransform& LocalToWorld) const
 {
-	const FName PropertyName = (PropertyChangedEvent.Property != NULL) ? PropertyChangedEvent.Property->GetFName() : NAME_None;
-
-	if ((PropertyName == GET_MEMBER_NAME_CHECKED(UPaperTileMapRenderComponent, MapWidth)) || (PropertyName == GET_MEMBER_NAME_CHECKED(UPaperTileMapRenderComponent, MapHeight)))
+	if (TileMap != NULL)
 	{
-		MapWidth = FMath::Max(MapWidth, 1);
-		MapHeight = FMath::Max(MapHeight, 1);
+		// Graphics bounds.
+		FBoxSphereBounds NewBounds = TileMap->GetRenderBounds().TransformBy(LocalToWorld);
 
-		// Resize all of the existing layers
-		for (int32 LayerIndex = 0; LayerIndex < TileLayers.Num(); ++LayerIndex)
+		// Add bounds of collision geometry (if present).
+		if (UBodySetup* BodySetup = TileMap->BodySetup)
 		{
-			UPaperTileLayer* TileLayer = TileLayers[LayerIndex];
-			TileLayer->ResizeMap(MapWidth, MapHeight);
+			const FBox AggGeomBox = BodySetup->AggGeom.CalcAABB(LocalToWorld);
+			if (AggGeomBox.IsValid)
+			{
+				NewBounds = Union(NewBounds, FBoxSphereBounds(AggGeomBox));
+			}
 		}
-	}
 
-	if (!IsTemplate())
+		// Apply bounds scale
+		NewBounds.BoxExtent *= BoundsScale;
+		NewBounds.SphereRadius *= BoundsScale;
+
+		return NewBounds;
+	}
+	else
 	{
-		UpdateBodySetup();
+		return FBoxSphereBounds(LocalToWorld.GetLocation(), FVector::ZeroVector, 0.f);
 	}
-
-	Super::PostEditChangeProperty(PropertyChangedEvent);
 }
-#endif
 
-FVector UPaperTileMapRenderComponent::ConvertTilePositionToWorldSpace(int32 TileX, int32 TileY, int32 LayerIndex) const
+void UPaperTileMapRenderComponent::TickComponent(float DeltaTime, enum ELevelTick TickType, FActorComponentTickFunction* ThisTickFunction)
 {
-	//@TODO: Layer depth, multiply by layer depth here!
-	//@TODO: Tile pivot issue
-	const FVector LocalPos((TileX - 0.5f) * TileWidth, 0, -(TileY - 0.5f) * TileHeight);
+	// Indicate we need to send new dynamic data.
+	MarkRenderDynamicDataDirty();
+}
 
-	return ComponentToWorld.TransformPosition(LocalPos);
+#if WITH_EDITORONLY_DATA
+void UPaperTileMapRenderComponent::PostEditChangeProperty(FPropertyChangedEvent& PropertyChangedEvent)
+{
+	Super::PostEditChangeProperty(PropertyChangedEvent);
 }
 
-UBodySetup* UPaperTileMapRenderComponent::GetBodySetup()
+void UPaperTileMapRenderComponent::Serialize(FArchive& Ar)
 {
-	UpdateBodySetup();
+	Super::Serialize(Ar);
 
-	return ShapeBodySetup;
+	Ar.UsingCustomVersion(FPaperCustomVersion::GUID);
 }
 
-void UPaperTileMapRenderComponent::UpdateBodySetup()
+void UPaperTileMapRenderComponent::PostLoad()
 {
-	if (ShapeBodySetup == NULL || ShapeBodySetup->IsPendingKill())
-	{
-		ShapeBodySetup = ConstructObject<UBodySetup>(UBodySetup::StaticClass(), this);
-		ShapeBodySetup->CollisionTraceFlag = CTF_UseSimpleAsComplex;
-	}
+	Super::PostLoad();
 
-	ShapeBodySetup->AggGeom.BoxElems.Empty();
-	for (int32 LayerIndex = 0; LayerIndex < TileLayers.Num(); ++LayerIndex)
+#if WITH_EDITORONLY_DATA
+	if (GetLinkerCustomVersion(FPaperCustomVersion::GUID) < FPaperCustomVersion::MovedTileMapDataToSeparateClass)
 	{
-		TileLayers[LayerIndex]->AugmentBodySetup(ShapeBodySetup);
+		// Create a tile map object and move our old properties over to it
+		TileMap = NewObject<UPaperTileMap>(this);
+		TileMap->MapWidth = MapWidth_DEPRECATED;
+		TileMap->MapHeight = MapHeight_DEPRECATED;
+		TileMap->TileWidth = TileWidth_DEPRECATED;
+		TileMap->TileHeight = TileHeight_DEPRECATED;
+		TileMap->PixelsPerUnit = 1.0f;
+		TileMap->DefaultLayerTileSet = DefaultLayerTileSet_DEPRECATED;
+		TileMap->Material = Material_DEPRECATED;
+		TileMap->TileLayers = TileLayers_DEPRECATED;
+
+		// Convert the layers
+		for (UPaperTileLayer* Layer : TileMap->TileLayers)
+		{
+			Layer->Rename(nullptr, TileMap, REN_ForceNoResetLoaders | REN_DontCreateRedirectors);
+			Layer->ConvertToTileSetPerCell();
+		}
+
+		// Remove references in the deprecated variables to prevent issues with deleting referenced assets, etc...
+		DefaultLayerTileSet_DEPRECATED = nullptr;
+		Material_DEPRECATED = nullptr;
+		TileLayers_DEPRECATED.Empty();
 	}
+#endif
+}
+#endif
+
+UBodySetup* UPaperTileMapRenderComponent::GetBodySetup()
+{
+	return (TileMap != nullptr) ? TileMap->BodySetup : nullptr;
 }
 
 const UObject* UPaperTileMapRenderComponent::AdditionalStatObject() const
 {
-	return DefaultLayerTileSet;
+	if (TileMap != nullptr)
+	{
+		if (TileMap->GetOuter() != this)
+		{
+			return TileMap;
+		}
+	}
+
+	return nullptr;
 }
 
 #undef LOCTEXT_NAMESPACE
\ No newline at end of file
diff --git a/Engine/Plugins/Experimental/Paper2D/Source/Paper2D/Private/PaperTileMapRenderSceneProxy.cpp b/Engine/Plugins/Experimental/Paper2D/Source/Paper2D/Private/PaperTileMapRenderSceneProxy.cpp
--- a/Engine/Plugins/Experimental/Paper2D/Source/Paper2D/Private/PaperTileMapRenderSceneProxy.cpp
+++ b/Engine/Plugins/Experimental/Paper2D/Source/Paper2D/Private/PaperTileMapRenderSceneProxy.cpp
@@ -3,18 +3,19 @@
 #include "Paper2DPrivatePCH.h"
 #include "PaperTileMapRenderSceneProxy.h"
 #include "PhysicsEngine/BodySetup.h"
+#include "PhysicsEngine/BodySetup2D.h"
 
 //////////////////////////////////////////////////////////////////////////
 // FPaperTileMapRenderSceneProxy
 
 FPaperTileMapRenderSceneProxy::FPaperTileMapRenderSceneProxy(const UPaperTileMapRenderComponent* InComponent)
 	: FPaperRenderSceneProxy(InComponent)
-	, TileComponent(NULL)
+	, TileMap(nullptr)
 {
 	if (const UPaperTileMapRenderComponent* InTileComponent = Cast<const UPaperTileMapRenderComponent>(InComponent))
 	{
-		TileComponent = InTileComponent;
-		Material = TileComponent->Material;
+		TileMap = InTileComponent->TileMap;
+		Material = (TileMap != nullptr) ? TileMap->Material : nullptr;
 
 		if (Material)
 		{
@@ -29,20 +30,23 @@ void FPaperTileMapRenderSceneProxy::DrawDynamicElements(FPrimitiveDrawInterface*
 
 	// Draw the tile maps
 	//@TODO: RenderThread race condition
-	if (TileComponent != NULL)
+	if (TileMap != NULL)
 	{
 		FColor WireframeColor = FColor(0, 255, 255, 255);
-		bool bIsCollisionEnabled = TileComponent->IsCollisionEnabled();
 
-		if((View->Family->EngineShowFlags.Collision && bIsCollisionEnabled) && AllowDebugViewmodes())
+		if ((View->Family->EngineShowFlags.Collision /*@TODO: && bIsCollisionEnabled*/) && AllowDebugViewmodes())
 		{
-			if(TileComponent->ShapeBodySetup)
+			if (UBodySetup2D* BodySetup2D = Cast<UBodySetup2D>(TileMap->BodySetup))
 			{
-				if(FMath::Abs(GetLocalToWorld().Determinant()) < SMALL_NUMBER)
+				//@TODO: Draw 2D debugging geometry
+			}
+			else if (UBodySetup* BodySetup = TileMap->BodySetup)
+			{
+				if (FMath::Abs(GetLocalToWorld().Determinant()) < SMALL_NUMBER)
 				{
 					// Catch this here or otherwise GeomTransform below will assert
 					// This spams so commented out
-					//UE_LOG(LogPaper2D, Log, TEXT("Zero scaling not supported (%s)"), *StaticMesh->GetPathName());
+					//UE_LOG(LogStaticMesh, Log, TEXT("Zero scaling not supported (%s)"), *StaticMesh->GetPathName());
 				}
 				else
 				{
@@ -68,11 +72,11 @@ void FPaperTileMapRenderSceneProxy::DrawDynamicElements(FPrimitiveDrawInterface*
 					}
 
 					// Differentiate the color based on bBlockNonZeroExtent.  Helps greatly with skimming a level for optimization opportunities.
-					FColor collisionColor = FColor(157,149,223,255);
+					FColor CollisionColor = FColor(157,149,223,255);
 
 					const bool bPerHullColor = false;
 					const bool bDrawSimpleSolid = false;
-					TileComponent->ShapeBodySetup->AggGeom.DrawAggGeom(PDI, GeomTransform, GetSelectionColor(collisionColor, bDrawWireSelected, IsHovered()), &CollisionMaterialInstance, bPerHullColor, bDrawSimpleSolid);
+					BodySetup->AggGeom.DrawAggGeom(PDI, GeomTransform, GetSelectionColor(CollisionColor, bDrawWireSelected, IsHovered()), &CollisionMaterialInstance, bPerHullColor, bDrawSimpleSolid);
 				}
 			}
 		}
@@ -99,31 +103,31 @@ void FPaperTileMapRenderSceneProxy::DrawDynamicElements(FPrimitiveDrawInterface*
 
 			FTransform LocalToWorld(GetLocalToWorld());
 
-			const float TW = TileComponent->TileWidth;
-			const float TH = TileComponent->TileHeight;
+			const float TW = TileMap ->TileWidth;
+			const float TH = TileMap->TileHeight;
 			//@TODO: PAPER: Tiles probably shouldn't be drawn with their pivot at the center - RE: all the 0.5f in here
 
 			if (bUseOverrideColor)
 			{
 				// Draw horizontal lines
-				for (int32 Y = 0; Y <= TileComponent->MapHeight; ++Y)
+				for (int32 Y = 0; Y <= TileMap->MapHeight; ++Y)
 				{
 					int32 X = 0;
 					const FVector Start((X - 0.5f) * TW, 0.0f, -(Y - 0.5f) * TH);
 
-					X = TileComponent->MapWidth;
+					X = TileMap->MapWidth;
 					const FVector End((X - 0.5f) * TW, 0.0f, -(Y - 0.5f) * TH);
 
 					PDI->DrawLine(LocalToWorld.TransformPosition(Start), LocalToWorld.TransformPosition(End), OverrideColor, DPG);
 				}
 
 				// Draw vertical lines
-				for (int32 X = 0; X <= TileComponent->MapWidth; ++X)
+				for (int32 X = 0; X <= TileMap->MapWidth; ++X)
 				{
 					int32 Y = 0;
 					const FVector Start((X - 0.5f) * TW, 0.0f, -(Y - 0.5f) * TH);
 
-					Y = TileComponent->MapHeight;
+					Y = TileMap->MapHeight;
 					const FVector End((X - 0.5f) * TW, 0.0f, -(Y - 0.5f) * TH);
 
 					PDI->DrawLine(LocalToWorld.TransformPosition(Start), LocalToWorld.TransformPosition(End), OverrideColor, DPG);
@@ -133,7 +137,7 @@ void FPaperTileMapRenderSceneProxy::DrawDynamicElements(FPrimitiveDrawInterface*
 
 			// Create a local space bounding box
 			const FVector TopLeft(-TW*0.5f, -2.0f, TH*0.5f);
-			const FVector Dimensions(TileComponent->MapWidth * TW, 4.0f, -TileComponent->MapHeight * TH);
+			const FVector Dimensions(TileMap->MapWidth * TW, 4.0f, -TileMap->MapHeight * TH);
 			const FBox OutlineBox(TopLeft, TopLeft+Dimensions);
 		
 			// Draw it		
@@ -143,12 +147,12 @@ void FPaperTileMapRenderSceneProxy::DrawDynamicElements(FPrimitiveDrawInterface*
 
 		const float LayerSeparation = 64.0f;
 
-		const int32 TileSeparationX = TileComponent->TileWidth;
-		const int32 TileSeparationY = TileComponent->TileHeight;
+		const int32 TileSeparationX = TileMap->TileWidth;
+		const int32 TileSeparationY = TileMap->TileHeight;
 
-		for (int32 LayerIndex = 0; LayerIndex < TileComponent->TileLayers.Num(); ++LayerIndex)
+		for (int32 LayerIndex = 0; LayerIndex < TileMap->TileLayers.Num(); ++LayerIndex)
 		{
-			UPaperTileLayer* Layer = TileComponent->TileLayers[LayerIndex];
+			UPaperTileLayer* Layer = TileMap->TileLayers[LayerIndex];
 			if (Layer == NULL)
 			{
 				continue;
@@ -164,34 +168,18 @@ void FPaperTileMapRenderSceneProxy::DrawDynamicElements(FPrimitiveDrawInterface*
 				continue;
 			}
 
-			ensure(Layer->LayerWidth == TileComponent->MapWidth);
-			ensure(Layer->LayerHeight == TileComponent->MapHeight);
+			ensure(Layer->LayerWidth == TileMap->MapWidth);
+			ensure(Layer->LayerHeight == TileMap->MapHeight);
 
 			UPaperTileSet* TileSet = Layer->TileSet;
 				
 			int32 TileWidth = 4;
 			int32 TileHeight = 4;
-			UTexture2D* SourceTexture = NULL;
-			if ( Layer->bCollisionLayer )
-			{
-				SourceTexture = UCanvas::StaticClass()->GetDefaultObject<UCanvas>()->DefaultTexture;
-			}
-			else
-			{
-				TileWidth = TileSet->TileWidth;
-				TileHeight = TileSet->TileHeight;
-				SourceTexture = TileSet->TileSheet;
-			}
 
+			UTexture2D* LastSourceTexture = nullptr;
 			FVector2D InverseTextureSize(1.0f, 1.0f);
-			if (SourceTexture != nullptr)
-			{
-				InverseTextureSize = FVector2D(1.0f / SourceTexture->GetSizeX(), 1.0f / SourceTexture->GetSizeY());
-			}
-
-			const FVector2D SourceDimensionsUV(TileWidth * InverseTextureSize.X, TileHeight * InverseTextureSize.Y);
+			FVector2D SourceDimensionsUV(1.0f, 1.0f);
 
-			UPaperTileMapRenderComponent* TileMap = Layer->GetTileMap();
 			FVector2D TileSizeXY(TileMap->TileWidth, TileMap->TileHeight);
 			FLinearColor DrawColor = FLinearColor::White;
 			EBlendMode BlendMode = BLEND_Translucent;
@@ -205,23 +193,48 @@ void FPaperTileMapRenderSceneProxy::DrawDynamicElements(FPrimitiveDrawInterface*
 			{
 				for (int32 X = 0; X < Layer->LayerWidth; ++X)
 				{
-					const int32 TileIndex = Layer->GetCell(X, Y);
+					const FPaperTileInfo TileInfo = Layer->GetCell(X, Y);
+					UTexture2D* SourceTexture = nullptr;
 
 					FVector2D SourceUV = FVector2D::ZeroVector;
 					if (Layer->bCollisionLayer)
 					{
-						if ( TileIndex == 0 )
+						if (TileInfo.PackedTileIndex == 0)
 						{
 							continue;
 						}
+						SourceTexture = UCanvas::StaticClass()->GetDefaultObject<UCanvas>()->DefaultTexture;
 					}
 					else
 					{
-						if (!TileSet->GetTileUV(TileIndex, /*out*/ SourceUV))
+						if (TileInfo.TileSet == nullptr)
+						{
+							continue;
+						}
+
+						if (!TileInfo.TileSet->GetTileUV(TileInfo.PackedTileIndex, /*out*/ SourceUV))
+						{
+							continue;
+						}
+
+						SourceTexture = TileInfo.TileSet->TileSheet;
+						if (SourceTexture == nullptr)
 						{
 							continue;
 						}
 					}
+
+
+					if (SourceTexture != LastSourceTexture)
+					{
+						TileWidth = TileSet->TileWidth;
+						TileHeight = TileSet->TileHeight;
+						InverseTextureSize = FVector2D(1.0f / SourceTexture->GetSizeX(), 1.0f / SourceTexture->GetSizeY());
+						SourceDimensionsUV = FVector2D(TileWidth * InverseTextureSize.X, TileHeight * InverseTextureSize.Y);
+						LastSourceTexture = SourceTexture;
+					}
+
+
 					SourceUV.X *= InverseTextureSize.X;
 					SourceUV.Y *= InverseTextureSize.Y;
 
diff --git a/Engine/Plugins/Experimental/Paper2D/Source/Paper2D/Private/PaperTileMapRenderSceneProxy.h b/Engine/Plugins/Experimental/Paper2D/Source/Paper2D/Private/PaperTileMapRenderSceneProxy.h
--- a/Engine/Plugins/Experimental/Paper2D/Source/Paper2D/Private/PaperTileMapRenderSceneProxy.h
+++ b/Engine/Plugins/Experimental/Paper2D/Source/Paper2D/Private/PaperTileMapRenderSceneProxy.h
@@ -18,5 +18,5 @@ public:
 
 protected:
 	//@TODO: Not thread safe
-	const UPaperTileMapRenderComponent* TileComponent;
+	const UPaperTileMap* TileMap;
 };
diff --git a/Engine/Plugins/Experimental/Paper2D/Source/Paper2DEditor/Private/TileMapEditing/EdModeTileMap.cpp b/Engine/Plugins/Experimental/Paper2D/Source/Paper2DEditor/Private/TileMapEditing/EdModeTileMap.cpp
--- a/Engine/Plugins/Experimental/Paper2D/Source/Paper2DEditor/Private/TileMapEditing/EdModeTileMap.cpp
+++ b/Engine/Plugins/Experimental/Paper2D/Source/Paper2DEditor/Private/TileMapEditing/EdModeTileMap.cpp
@@ -34,7 +34,7 @@ struct FHorizontalSpan
 	}
 
 	// Grows a span horizontally until it reaches something that doesn't match
-	void GrowSpan(int32 RequiredInk, UPaperTileLayer* Layer, TBitArray<>& Reachability)
+	void GrowSpan(const FPaperTileInfo& RequiredInk, UPaperTileLayer* Layer, TBitArray<>& Reachability)
 	{
 		// Go left
 		for (int32 TestX = X0 - 1; TestX >= 0; --TestX)
@@ -315,7 +315,15 @@ UPaperTileLayer* FEdModeTileMap::GetSelectedLayerUnderCursor(const FViewportCurs
 	for (FSelectionIterator Iter(*SelectedActors); Iter; ++Iter)
 	{
 		AActor* Actor = CastChecked<AActor>(*Iter);
-		if (UPaperTileMapRenderComponent* TileMap = Actor->FindComponentByClass<UPaperTileMapRenderComponent>())
+
+		UPaperTileMap* TileMap = nullptr;
+		UPaperTileMapRenderComponent* TileMapComponent = Actor->FindComponentByClass<UPaperTileMapRenderComponent>();
+		if (TileMapComponent != nullptr)
+		{
+			TileMap = TileMapComponent->TileMap;
+		}
+
+		if (TileMap != nullptr)
 		{
 			// Find the first visible layer
 			int32 LayerIndex = 0;
@@ -336,7 +344,7 @@ UPaperTileLayer* FEdModeTileMap::GetSelectedLayerUnderCursor(const FViewportCurs
 				const float WX = TileMap->MapWidth * TileMap->TileWidth;
 				const float WY = TileMap->MapHeight * TileMap->TileHeight;
 
-				const FTransform ComponentToWorld = TileMap->ComponentToWorld;
+				ComponentToWorld = (TileMapComponent != nullptr) ? TileMapComponent->ComponentToWorld : FTransform::Identity;
 				FVector LocalStart = ComponentToWorld.InverseTransformPosition(TraceStart);
 				FVector LocalDirection = ComponentToWorld.InverseTransformVector(TraceDir);
 
@@ -414,7 +422,7 @@ bool FEdModeTileMap::PaintTiles(const FViewportCursorLocation& Ray)
 
 	if (UPaperTileLayer* Layer = GetSelectedLayerUnderCursor(Ray, /*out*/ DestTileX, /*out*/ DestTileY))
 	{
-		UPaperTileMapRenderComponent* TileMap = CastChecked<UPaperTileMapRenderComponent>(Layer->GetOuter());
+		UPaperTileMap* TileMap = Layer->GetTileMap();
 
 		FScopedTransaction Transaction( LOCTEXT("TileMapPaintAction", "Tile Painting") );
 
@@ -436,11 +444,11 @@ bool FEdModeTileMap::PaintTiles(const FViewportCursorLocation& Ray)
 					continue;
 				}
 
-				int32 Ink = 0;
+				FPaperTileInfo Ink;
 				if (GetActiveLayerPaintingMode() == ETileMapLayerPaintingMode::CollisionLayers)
 				{
 					// 1 Means collision, 0 means no collision
-					Ink = 1;
+					Ink.PackedTileIndex = 1;
 				}
 				else
 				{
@@ -457,7 +465,8 @@ bool FEdModeTileMap::PaintTiles(const FViewportCursorLocation& Ray)
 						continue;
 					}
 
-					Ink = SX + (SY * InkSource->GetTileCountX());
+					Ink.PackedTileIndex = SX + (SY * InkSource->GetTileCountX());
+					Ink.TileSet = Layer->TileSet;
 				}
 
 				if (Layer->GetCell(DX, DY) != Ink)
@@ -494,14 +503,14 @@ bool FEdModeTileMap::EraseTiles(const FViewportCursorLocation& Ray)
 	bool bPaintedOnSomething = false;
 	bool bChangedSomething = false;
 
-	const int32 EmptyCellValue = 0; //@TODO: Would really like to use INDEX_NONE for this...
+	const FPaperTileInfo EmptyCellValue;
 
 	int32 DestTileX;
 	int32 DestTileY;
 
 	if (UPaperTileLayer* Layer = GetSelectedLayerUnderCursor(Ray, /*out*/ DestTileX, /*out*/ DestTileY))
 	{
-		UPaperTileMapRenderComponent* TileMap = CastChecked<UPaperTileMapRenderComponent>(Layer->GetOuter());
+		UPaperTileMap* TileMap = Layer->GetTileMap();
 
 		FScopedTransaction Transaction( LOCTEXT("TileMapEraseAction", "Tile Erasing") );
 
@@ -586,9 +595,9 @@ bool FEdModeTileMap::FloodFillTiles(const FViewportCursorLocation& Ray)
 		}
 
 		// The kind of ink we'll replace, starting at the seed point
-		const int32 RequiredInk = Layer->GetCell(DestTileX, DestTileY);
+		const FPaperTileInfo RequiredInk = Layer->GetCell(DestTileX, DestTileY);
 
-		UPaperTileMapRenderComponent* TileMap = CastChecked<UPaperTileMapRenderComponent>(Layer->GetOuter());
+		UPaperTileMap* TileMap = Layer->GetTileMap();
 
 		//@TODO: Unoptimized first-pass approach
 		const int32 NumTiles = TileMap->MapWidth * TileMap->MapHeight;
@@ -648,16 +657,17 @@ bool FEdModeTileMap::FloodFillTiles(const FViewportCursorLocation& Ray)
 				{
 					const int32 InsideBrushX = (DX + BrushPatternOffsetX) % BrushWidth;
 
-					int32 NewInk = 0;
-					if ( GetActiveLayerPaintingMode() == ETileMapLayerPaintingMode::CollisionLayers )
+					FPaperTileInfo NewInk;
+					if (GetActiveLayerPaintingMode() == ETileMapLayerPaintingMode::CollisionLayers)
 					{
-						NewInk = 1;
+						NewInk.PackedTileIndex = 1;
 					}
 					else
 					{
 						const int32 TileSetX = PaintSourceTopLeft.X + InsideBrushX;
 						const int32 TileSetY = PaintSourceTopLeft.Y + InsideBrushY;
-						NewInk = TileSetX + (TileSetY * InkSource->GetTileCountX());
+						NewInk.TileSet = Layer->TileSet;
+						NewInk.PackedTileIndex = TileSetX + (TileSetY * InkSource->GetTileCountX());
 					}
 
 					if (Layer->GetCell(DX, DY) != NewInk)
@@ -711,11 +721,11 @@ void FEdModeTileMap::UpdatePreviewCursor(const FViewportCursorLocation& Ray)
 			const int32 LocalTileX1 = LocalTileX0 + CursorWidth;
 			const int32 LocalTileY1 = LocalTileY0 + CursorHeight;
 
-			UPaperTileMapRenderComponent* TileMap = CastChecked<UPaperTileMapRenderComponent>(TileLayer->GetOuter());
-			const FVector WorldPosition = TileMap->ConvertTilePositionToWorldSpace(LocalTileX0, LocalTileY0);
-			const FVector WorldPositionBR = TileMap->ConvertTilePositionToWorldSpace(LocalTileX1, LocalTileY1);
+			UPaperTileMap* TileMap = TileLayer->GetTileMap();
+			const FVector WorldPosition = ComponentToWorld.TransformPosition(TileMap->GetTilePositionInLocalSpace(LocalTileX0, LocalTileY0));
+			const FVector WorldPositionBR = ComponentToWorld.TransformPosition(TileMap->GetTilePositionInLocalSpace(LocalTileX1, LocalTileY1));
 
-			DrawPreviewSpace = TileMap->ComponentToWorld;
+			DrawPreviewSpace = ComponentToWorld;
 			DrawPreviewLocation = (WorldPosition + WorldPositionBR) * 0.5f;
 
 			DrawPreviewDimensionsLS = 0.5f*FVector(CursorWidth * TileMap->TileWidth, 0.0f, -CursorHeight * TileMap->TileHeight);
diff --git a/Engine/Plugins/Experimental/Paper2D/Source/Paper2DEditor/Private/TileMapEditing/EdModeTileMap.h b/Engine/Plugins/Experimental/Paper2D/Source/Paper2DEditor/Private/TileMapEditing/EdModeTileMap.h
--- a/Engine/Plugins/Experimental/Paper2D/Source/Paper2DEditor/Private/TileMapEditing/EdModeTileMap.h
+++ b/Engine/Plugins/Experimental/Paper2D/Source/Paper2DEditor/Private/TileMapEditing/EdModeTileMap.h
@@ -119,5 +119,6 @@ protected:
 
 	ETileMapEditorTool::Type ActiveTool;
 	ETileMapLayerPaintingMode::Type LayerPaintingMode;
+	mutable FTransform ComponentToWorld;
 };
 
diff --git a/Engine/Plugins/Experimental/Paper2D/Source/Paper2DEditor/Private/TileMapEditing/PaperTileMapDetailsCustomization.cpp b/Engine/Plugins/Experimental/Paper2D/Source/Paper2DEditor/Private/TileMapEditing/PaperTileMapDetailsCustomization.cpp
--- a/Engine/Plugins/Experimental/Paper2D/Source/Paper2DEditor/Private/TileMapEditing/PaperTileMapDetailsCustomization.cpp
+++ b/Engine/Plugins/Experimental/Paper2D/Source/Paper2DEditor/Private/TileMapEditing/PaperTileMapDetailsCustomization.cpp
@@ -24,14 +24,14 @@ void FPaperTileMapDetailsCustomization::CustomizeDetails(IDetailLayoutBuilder& D
 {
 	const TArray< TWeakObjectPtr<UObject> >& SelectedObjects = DetailLayout.GetDetailsView().GetSelectedObjects();
 	
-	UPaperTileMapRenderComponent* TileMap = NULL;
+	UPaperTileMap* TileMap = NULL;
 	for (int32 ObjectIndex = 0; ObjectIndex < SelectedObjects.Num(); ++ObjectIndex)
 	{
 		if (AActor* CurrentActor = Cast<AActor>(SelectedObjects[ObjectIndex].Get()))
 		{
 			if (UPaperTileMapRenderComponent* CurrentTileMap = CurrentActor->FindComponentByClass<UPaperTileMapRenderComponent>())
 			{
-				TileMap = CurrentTileMap;
+				TileMap = CurrentTileMap->TileMap;
 				break;
 			}
 		}
@@ -138,7 +138,7 @@ UPaperTileLayer* FPaperTileMapDetailsCustomization::AddLayer(bool bCollisionLaye
 {
 	UPaperTileLayer* NewLayer = NULL;
 
-	if (UPaperTileMapRenderComponent* TileMap = TileMapPtr.Get())
+	if (UPaperTileMap* TileMap = TileMapPtr.Get())
 	{
 		const FScopedTransaction Transaction( LOCTEXT("TileMapAddLayer", "Add New Layer") );
 		TileMap->SetFlags(RF_Transactional);
diff --git a/Engine/Plugins/Experimental/Paper2D/Source/Paper2DEditor/Private/TileMapEditing/PaperTileMapDetailsCustomization.h b/Engine/Plugins/Experimental/Paper2D/Source/Paper2DEditor/Private/TileMapEditing/PaperTileMapDetailsCustomization.h
--- a/Engine/Plugins/Experimental/Paper2D/Source/Paper2DEditor/Private/TileMapEditing/PaperTileMapDetailsCustomization.h
+++ b/Engine/Plugins/Experimental/Paper2D/Source/Paper2DEditor/Private/TileMapEditing/PaperTileMapDetailsCustomization.h
@@ -18,7 +18,7 @@ public:
 	// End of IDetailCustomization interface
 
 private:
-	TWeakObjectPtr<class UPaperTileMapRenderComponent> TileMapPtr;
+	TWeakObjectPtr<class UPaperTileMap> TileMapPtr;
 private:
 	FReply AddLayerClicked();
 	FReply AddCollisionLayerClicked();
diff --git a/Engine/Plugins/Experimental/Paper2D/Source/Paper2DEditor/Private/TileMapEditing/STileLayerList.cpp b/Engine/Plugins/Experimental/Paper2D/Source/Paper2DEditor/Private/TileMapEditing/STileLayerList.cpp
--- a/Engine/Plugins/Experimental/Paper2D/Source/Paper2DEditor/Private/TileMapEditing/STileLayerList.cpp
+++ b/Engine/Plugins/Experimental/Paper2D/Source/Paper2DEditor/Private/TileMapEditing/STileLayerList.cpp
@@ -7,7 +7,7 @@
 //////////////////////////////////////////////////////////////////////////
 // STileLayerList
 
-void STileLayerList::Construct(const FArguments& InArgs, UPaperTileMapRenderComponent* TileMap)
+void STileLayerList::Construct(const FArguments& InArgs, UPaperTileMap* TileMap)
 {
 	TSharedRef<SHeaderRow> HeaderRowWidget = SNew(SHeaderRow)
 		+SHeaderRow::Column("Name").DefaultLabel(NSLOCTEXT("TileLayerList", "TileLayerNameHeader", "Layer Name")).FillWidth(0.7)
diff --git a/Engine/Plugins/Experimental/Paper2D/Source/Paper2DEditor/Private/TileMapEditing/STileLayerList.h b/Engine/Plugins/Experimental/Paper2D/Source/Paper2DEditor/Private/TileMapEditing/STileLayerList.h
--- a/Engine/Plugins/Experimental/Paper2D/Source/Paper2DEditor/Private/TileMapEditing/STileLayerList.h
+++ b/Engine/Plugins/Experimental/Paper2D/Source/Paper2DEditor/Private/TileMapEditing/STileLayerList.h
@@ -11,7 +11,7 @@ public:
 	SLATE_BEGIN_ARGS(STileLayerList) {}
 	SLATE_END_ARGS()
 
-	void Construct(const FArguments& InArgs, UPaperTileMapRenderComponent* TileMap);
+	void Construct(const FArguments& InArgs, UPaperTileMap* TileMap);
 
 protected:
 	typedef SListView<class UPaperTileLayer*> SPaperLayerListView;
diff --git a/Engine/Plugins/Experimental/Paper2D/Source/Paper2DEditor/Private/TileMapEditing/TileMapEdModeToolkit.cpp b/Engine/Plugins/Experimental/Paper2D/Source/Paper2DEditor/Private/TileMapEditing/TileMapEdModeToolkit.cpp
--- a/Engine/Plugins/Experimental/Paper2D/Source/Paper2DEditor/Private/TileMapEditing/TileMapEdModeToolkit.cpp
+++ b/Engine/Plugins/Experimental/Paper2D/Source/Paper2DEditor/Private/TileMapEditing/TileMapEdModeToolkit.cpp
@@ -65,10 +65,13 @@ void FTileMapEdModeToolkit::Init(const TSharedPtr<IToolkitHost>& InitToolkitHost
 	for (FSelectionIterator Iter(*SelectedActors); Iter; ++Iter)
 	{
 		AActor* Actor = CastChecked<AActor>(*Iter);
-		if (UPaperTileMapRenderComponent* TileMap = Actor->FindComponentByClass<UPaperTileMapRenderComponent>())
+		if (UPaperTileMapRenderComponent* TileMapComponent = Actor->FindComponentByClass<UPaperTileMapRenderComponent>())
 		{
-			CurrentTileSetPtr = TileMap->DefaultLayerTileSet;
-			break;
+			if (TileMapComponent->TileMap != nullptr)
+			{
+				CurrentTileSetPtr = TileMapComponent->TileMap->DefaultLayerTileSet;
+				break;
+			}
 		}
 	}
 
