CIS fix - Non-unity build fix

diff --git a/Engine/Plugins/Experimental/Paper2D/Source/Paper2D/Private/PaperTileMap.cpp b/Engine/Plugins/Experimental/Paper2D/Source/Paper2D/Private/PaperTileMap.cpp
--- a/Engine/Plugins/Experimental/Paper2D/Source/Paper2D/Private/PaperTileMap.cpp
+++ b/Engine/Plugins/Experimental/Paper2D/Source/Paper2D/Private/PaperTileMap.cpp
@@ -1,6 +1,8 @@
 // Copyright 1998-2014 Epic Games, Inc. All Rights Reserved.
 
 #include "Paper2DPrivatePCH.h"
+#include "PhysicsEngine/BodySetup.h"
+#include "PhysicsEngine/BodySetup2D.h"
 
 //////////////////////////////////////////////////////////////////////////
 // UPaperTileMap
