Paper2D: Mark UPaperSprite UBodySetup to use simple as complex, allowing complex raycasts (such as touch/cusor clicks) to succeed on 3D sprite collision

diff --git a/Engine/Plugins/Experimental/Paper2D/Source/Paper2D/Private/PaperCustomVersion.h b/Engine/Plugins/Experimental/Paper2D/Source/Paper2D/Private/PaperCustomVersion.h
--- a/Engine/Plugins/Experimental/Paper2D/Source/Paper2D/Private/PaperCustomVersion.h
+++ b/Engine/Plugins/Experimental/Paper2D/Source/Paper2D/Private/PaperCustomVersion.h
@@ -19,6 +19,9 @@ struct FPaperCustomVersion
 		// Moved tile map data out of the component into UPaperTileMap
 		MovedTileMapDataToSeparateClass = 3,
 
+		// Marked sprite UBodySetup to use simple as complex, allowing 3D raycasts to succeed
+		MarkSpriteBodySetupToUseSimpleAsComplex = 4,
+
 		// -----<new versions can be added above this line>-------------------------------------------------
 		VersionPlusOne,
 		LatestVersion = VersionPlusOne - 1
diff --git a/Engine/Plugins/Experimental/Paper2D/Source/Paper2D/Private/PaperSprite.cpp b/Engine/Plugins/Experimental/Paper2D/Source/Paper2D/Private/PaperSprite.cpp
--- a/Engine/Plugins/Experimental/Paper2D/Source/Paper2D/Private/PaperSprite.cpp
+++ b/Engine/Plugins/Experimental/Paper2D/Source/Paper2D/Private/PaperSprite.cpp
@@ -355,6 +355,7 @@ void UPaperSprite::RebuildCollisionData()
 
 	if (SpriteCollisionDomain != ESpriteCollisionMode::None)
 	{
+		BodySetup->CollisionTraceFlag = CTF_UseSimpleAsComplex;
 		switch (CollisionGeometry.GeometryType)
 		{
 		case ESpritePolygonMode::SourceBoundingBox:
@@ -1064,7 +1065,7 @@ void UPaperSprite::PostLoad()
 	Super::PostLoad();
 
 #if WITH_EDITORONLY_DATA
-	if (GetLinkerCustomVersion(FPaperCustomVersion::GUID) < FPaperCustomVersion::RemovedBodyInstance2D)
+	if (GetLinkerCustomVersion(FPaperCustomVersion::GUID) < FPaperCustomVersion::MarkSpriteBodySetupToUseSimpleAsComplex)
 	{
 		RebuildCollisionData();
 	}
