Paper2D: Fix tile map actors rendering upside down (integrated from 4.2 CL# 2099779)

diff --git a/Engine/Plugins/Experimental/Paper2D/Source/Paper2D/Private/PaperTileMapRenderSceneProxy.cpp b/Engine/Plugins/Experimental/Paper2D/Source/Paper2D/Private/PaperTileMapRenderSceneProxy.cpp
--- a/Engine/Plugins/Experimental/Paper2D/Source/Paper2D/Private/PaperTileMapRenderSceneProxy.cpp
+++ b/Engine/Plugins/Experimental/Paper2D/Source/Paper2D/Private/PaperTileMapRenderSceneProxy.cpp
@@ -238,7 +238,7 @@ void FPaperTileMapRenderSceneProxy::DrawDynamicElements(FPrimitiveDrawInterface*
 					SourceUV.X *= InverseTextureSize.X;
 					SourceUV.Y *= InverseTextureSize.Y;
 
-					const FVector DrawPos(TileSeparationX * (X - 0.5f) * PaperAxisX + TileSeparationY * -(Y - 0.5f) * PaperAxisY + ((LayerIndex * LayerSeparation) * PaperAxisZ));
+					const FVector DrawPos(TileSeparationX * (X - 0.5f) * PaperAxisX + TileSeparationY * -(Y + 0.5f) * PaperAxisY + ((LayerIndex * LayerSeparation) * PaperAxisZ));
 
 					FSpriteDrawCallRecord& NewTile = *(new (BatchedSprites) FSpriteDrawCallRecord());
 					NewTile.Destination = DrawPos;
@@ -248,10 +248,10 @@ void FPaperTileMapRenderSceneProxy::DrawDynamicElements(FPrimitiveDrawInterface*
 					const float W = TileSeparationX;
 					const float H = TileSeparationY;
 
-					const FVector4 BottomLeft(0.0f, 0.0f, SourceUV.X, SourceUV.Y);
-					const FVector4 BottomRight(W, 0.0f, SourceUV.X + SourceDimensionsUV.X, SourceUV.Y);
-					const FVector4 TopRight(W, H, SourceUV.X + SourceDimensionsUV.X, SourceUV.Y + SourceDimensionsUV.Y);
-					const FVector4 TopLeft(0.0f, H, SourceUV.X, SourceUV.Y + SourceDimensionsUV.Y);
+					const FVector4 BottomLeft(0.0f, 0.0f, SourceUV.X, SourceUV.Y + SourceDimensionsUV.Y);
+					const FVector4 BottomRight(W, 0.0f, SourceUV.X + SourceDimensionsUV.X, SourceUV.Y + SourceDimensionsUV.Y);
+					const FVector4 TopRight(W, H, SourceUV.X + SourceDimensionsUV.X, SourceUV.Y);
+					const FVector4 TopLeft(0.0f, H, SourceUV.X, SourceUV.Y);
 
 					new (NewTile.RenderVerts) FVector4(BottomLeft);
 					new (NewTile.RenderVerts) FVector4(TopRight);
