Paper2D: Added "Extract sprites" in texture context menu

diff --git a/Engine/Plugins/Experimental/Paper2D/Source/Paper2D/Classes/PaperSprite.h b/Engine/Plugins/Experimental/Paper2D/Source/Paper2D/Classes/PaperSprite.h
--- a/Engine/Plugins/Experimental/Paper2D/Source/Paper2D/Classes/PaperSprite.h
+++ b/Engine/Plugins/Experimental/Paper2D/Source/Paper2D/Classes/PaperSprite.h
@@ -152,6 +152,7 @@ public:
 
 	void FindTextureBoundingBox(float AlphaThreshold, /*out*/ FVector2D& OutBoxPosition, /*out*/ FVector2D& OutBoxSize);
 	static void FindContours(const FIntPoint& ScanPos, const FIntPoint& ScanSize, float AlphaThreshold, UTexture2D* Texture, TArray< TArray<FIntPoint> >& OutPoints);
+	static void ExtractRectsFromTexture(UTexture2D* Texture, TArray<FIntRect>& OutRects);
 	void BuildGeometryFromContours(FSpritePolygonCollection& GeomOwner);
 
 	void BuildBoundingBoxCollisionData(bool bUseTightBounds);
diff --git a/Engine/Plugins/Experimental/Paper2D/Source/Paper2D/Private/PaperSprite.cpp b/Engine/Plugins/Experimental/Paper2D/Source/Paper2D/Private/PaperSprite.cpp
--- a/Engine/Plugins/Experimental/Paper2D/Source/Paper2D/Private/PaperSprite.cpp
+++ b/Engine/Plugins/Experimental/Paper2D/Source/Paper2D/Private/PaperSprite.cpp
@@ -865,6 +865,12 @@ void UPaperSprite::Triangulate(const FSpritePolygonCollection& Source, TArray<FV
 	}
 }
 
+void UPaperSprite::ExtractRectsFromTexture(UTexture2D* Texture, TArray<FIntRect>& OutRects)
+{
+	FBitmap SpriteTextureBitmap(Texture, 0, 0);
+	SpriteTextureBitmap.ExtractRects(/*out*/ OutRects);
+}
+
 void UPaperSprite::InitializeSprite(UTexture2D* Texture)
 {
 	if (Texture != NULL)
diff --git a/Engine/Plugins/Experimental/Paper2D/Source/Paper2DEditor/Classes/PaperSpriteFactory.h b/Engine/Plugins/Experimental/Paper2D/Source/Paper2DEditor/Classes/PaperSpriteFactory.h
--- a/Engine/Plugins/Experimental/Paper2D/Source/Paper2DEditor/Classes/PaperSpriteFactory.h
+++ b/Engine/Plugins/Experimental/Paper2D/Source/Paper2DEditor/Classes/PaperSpriteFactory.h
@@ -14,6 +14,12 @@ class UPaperSpriteFactory : public UFactory
 	// Initial texture to create the sprite from (Can be NULL)
 	class UTexture2D* InitialTexture;
 
+	// Set bUseSourceRegion to get it to use/set initial sourceUV and dimensions
+	bool bUseSourceRegion;
+	FVector2D InitialSourceUV;
+	FVector2D InitialSourceDimension;
+
+
 	// UFactory interface
 	virtual bool ConfigureProperties() OVERRIDE;
 	virtual UObject* FactoryCreateNew(UClass* Class, UObject* InParent, FName Name, EObjectFlags Flags, UObject* Context, FFeedbackContext* Warn) OVERRIDE;
diff --git a/Engine/Plugins/Experimental/Paper2D/Source/Paper2DEditor/Private/ContentBrowserExtensions/ContentBrowserExtensions.cpp b/Engine/Plugins/Experimental/Paper2D/Source/Paper2DEditor/Private/ContentBrowserExtensions/ContentBrowserExtensions.cpp
--- a/Engine/Plugins/Experimental/Paper2D/Source/Paper2DEditor/Private/ContentBrowserExtensions/ContentBrowserExtensions.cpp
+++ b/Engine/Plugins/Experimental/Paper2D/Source/Paper2DEditor/Private/ContentBrowserExtensions/ContentBrowserExtensions.cpp
@@ -32,8 +32,16 @@ public:
 #include "AssetToolsModule.h"
 #include "AssetRegistryModule.h"
 
+
 struct FCreateSpriteFromTextureExtension : public FContentBrowserSelectedAssetExtensionBase
 {
+	bool bExtractSprites;
+
+	FCreateSpriteFromTextureExtension()
+		: bExtractSprites(false)
+	{
+	}
+
 	void CreateSpritesFromTextures(TArray<UTexture2D*>& Textures)
 	{
 		const FString DefaultSuffix = TEXT("_Sprite");
@@ -44,32 +52,70 @@ struct FCreateSpriteFromTextureExtension : public FContentBrowserSelectedAssetEx
 		const bool bOneTextureSelected = Textures.Num() == 1;
 		TArray<UObject*> ObjectsToSync;
 
+		FString Name;
+		FString PackageName;
+
 		for (auto TextureIt = Textures.CreateConstIterator(); TextureIt; ++TextureIt)
 		{
 			UTexture2D* Texture = *TextureIt;
 
-			// Get a unique name for the sprite
-			FString Name;
-			FString PackageName;
-			AssetToolsModule.Get().CreateUniqueAssetName(Texture->GetOutermost()->GetName(), DefaultSuffix, /*out*/ PackageName, /*out*/ Name);
-			const FString PackagePath = FPackageName::GetLongPackagePath(PackageName);
-
-
 			// Create the factory used to generate the sprite
 			UPaperSpriteFactory* SpriteFactory = ConstructObject<UPaperSpriteFactory>(UPaperSpriteFactory::StaticClass());
 			SpriteFactory->InitialTexture = Texture;
 
 			// Create the sprite
-			if (bOneTextureSelected)
+			if (!bExtractSprites)
 			{
-				ContentBrowserModule.Get().CreateNewAsset(Name, PackagePath, UPaperSprite::StaticClass(), SpriteFactory);
+				if (bOneTextureSelected)
+				{
+					// Get a unique name for the sprite
+					AssetToolsModule.Get().CreateUniqueAssetName(Texture->GetOutermost()->GetName(), DefaultSuffix, /*out*/ PackageName, /*out*/ Name);
+					const FString PackagePath = FPackageName::GetLongPackagePath(PackageName);
+					ContentBrowserModule.Get().CreateNewAsset(Name, PackagePath, UPaperSprite::StaticClass(), SpriteFactory);
+				}
+				else
+				{
+					// Get a unique name for the sprite
+					AssetToolsModule.Get().CreateUniqueAssetName(Texture->GetOutermost()->GetName(), DefaultSuffix, /*out*/ PackageName, /*out*/ Name);
+					const FString PackagePath = FPackageName::GetLongPackagePath(PackageName);
+
+					if (UObject* NewAsset = AssetToolsModule.Get().CreateAsset(Name, PackagePath, UPaperSprite::StaticClass(), SpriteFactory))
+					{
+						ObjectsToSync.Add(NewAsset);
+					}
+				}
 			}
 			else
 			{
-				if (UObject* NewAsset = AssetToolsModule.Get().CreateAsset(Name, PackagePath, UPaperSprite::StaticClass(), SpriteFactory))
+				GWarn->BeginSlowTask(NSLOCTEXT("Paper2D", "Paper2D_ExtractSpritesFromTexture", "Extracting Sprites From Texture"), true, true);
+
+				TArray<FIntRect> ExtractedRects;
+				UPaperSprite::ExtractRectsFromTexture(Texture, /*out*/ ExtractedRects);
+				for (int ExtractedRectIndex = 0; ExtractedRectIndex < ExtractedRects.Num(); ++ExtractedRectIndex)
 				{
-					ObjectsToSync.Add(NewAsset);
+					GWarn->StatusUpdate(ExtractedRectIndex, ExtractedRects.Num(), NSLOCTEXT("Paper2D", "Paper2D_ExtractSpritesFromTexture", "Extracting Sprites From Texture"));
+
+					FIntRect& ExtractedRect = ExtractedRects[ExtractedRectIndex];
+					SpriteFactory->bUseSourceRegion = true;
+					SpriteFactory->InitialSourceUV = FVector2D(ExtractedRect.Min.X, ExtractedRect.Min.Y);
+					SpriteFactory->InitialSourceDimension = FVector2D(ExtractedRect.Width(), ExtractedRect.Height());
+
+					// Get a unique name for the sprite
+					AssetToolsModule.Get().CreateUniqueAssetName(Texture->GetOutermost()->GetName(), DefaultSuffix, /*out*/ PackageName, /*out*/ Name);
+					const FString PackagePath = FPackageName::GetLongPackagePath(PackageName);
+
+					if (UObject* NewAsset = AssetToolsModule.Get().CreateAsset(Name, PackagePath, UPaperSprite::StaticClass(), SpriteFactory))
+					{
+						ObjectsToSync.Add(NewAsset);
+					}
+
+					if (GWarn->ReceivedUserCancel()) 
+					{
+						break;
+					}
 				}
+
+				GWarn->EndSlowTask();
 			}
 		}
 
@@ -132,7 +178,7 @@ public:
 		FUIAction Action_CreateSpritesFromTextures(
 			FExecuteAction::CreateStatic(&FPaperContentBrowserExtensions_Impl::ExecuteSelectedContentFunctor, StaticCastSharedPtr<FContentBrowserSelectedAssetExtensionBase>(SpriteCreatorFunctor)));
 		Builder.AddMenuEntry(
-			LOCTEXT("CB_Extension_Texture_CreateSprite", "Create Sprite"), 
+			LOCTEXT("CB_Extension_Texture_CreateSprite", "Create Sprite"),
 			LOCTEXT("CB_Extension_Texture_CreateSprite_Tooltip", "Create sprites from selected textures"),
 			FSlateIcon(),
 			Action_CreateSpritesFromTextures,
@@ -140,6 +186,19 @@ public:
 			EUserInterfaceActionType::Button);
 	}
 
+	static void ExtractSpritesFromTextureMenuOptions(FMenuBuilder& Builder, TSharedPtr<FCreateSpriteFromTextureExtension> SpriteExtractorFunctor)
+	{
+		FUIAction Action_ExtractSpritesFromTextures(
+			FExecuteAction::CreateStatic(&FPaperContentBrowserExtensions_Impl::ExecuteSelectedContentFunctor, StaticCastSharedPtr<FContentBrowserSelectedAssetExtensionBase>(SpriteExtractorFunctor)));
+		Builder.AddMenuEntry(
+			LOCTEXT("CB_Extension_Texture_ExtractSprites", "Extract Sprites"),
+			LOCTEXT("CB_Extension_Texture_ExtractSprite_Tooltip", "Detects and extracts sprites from the selected textures using transparency"),
+			FSlateIcon(),
+			Action_ExtractSpritesFromTextures,
+			NAME_None,
+			EUserInterfaceActionType::Button);
+	}
+
 	static void CreateTextureSetupMenuOption(FMenuBuilder& Builder, TSharedPtr<FConfigureTexturesForSpriteUsageExtension> ConfigureFunctor)
 	{
 		FUIAction Action_ConfigureTexturesForSprites(
@@ -179,6 +238,16 @@ public:
 				FMenuExtensionDelegate::CreateStatic(&FPaperContentBrowserExtensions_Impl::CreateSpriteTextureMenuOptions, SpriteCreatorFunctor));
 
 			//@TODO: Already getting nasty, need to refactor the data to be independent of the functor
+			TSharedPtr<FCreateSpriteFromTextureExtension> SpriteExtractorFunctor = MakeShareable(new FCreateSpriteFromTextureExtension());
+			SpriteExtractorFunctor->SelectedAssets = SelectedAssets;
+			SpriteExtractorFunctor->bExtractSprites = true;
+
+			Extender->AddMenuExtension(
+				"GetAssetActions",
+				EExtensionHook::After,
+				NULL,
+				FMenuExtensionDelegate::CreateStatic(&FPaperContentBrowserExtensions_Impl::ExtractSpritesFromTextureMenuOptions, SpriteExtractorFunctor));
+
 			TSharedPtr<FConfigureTexturesForSpriteUsageExtension> TextureConfigFunctor = MakeShareable(new FConfigureTexturesForSpriteUsageExtension());
 			TextureConfigFunctor->SelectedAssets = SelectedAssets;
 
diff --git a/Engine/Plugins/Experimental/Paper2D/Source/Paper2DEditor/Private/PaperSpriteFactory.cpp b/Engine/Plugins/Experimental/Paper2D/Source/Paper2DEditor/Private/PaperSpriteFactory.cpp
--- a/Engine/Plugins/Experimental/Paper2D/Source/Paper2DEditor/Private/PaperSpriteFactory.cpp
+++ b/Engine/Plugins/Experimental/Paper2D/Source/Paper2DEditor/Private/PaperSpriteFactory.cpp
@@ -9,6 +9,9 @@
 
 UPaperSpriteFactory::UPaperSpriteFactory(const class FPostConstructInitializeProperties& PCIP)
 	: Super(PCIP)
+	, bUseSourceRegion(false)
+	, InitialSourceUV(0, 0)
+	, InitialSourceDimension(0, 0)
 {
 	bCreateNew = true;
 	bEditAfterNew = true;
@@ -25,7 +28,14 @@ UObject* UPaperSpriteFactory::FactoryCreateNew(UClass* Class, UObject* InParent,
 {
 	UPaperSprite* NewSprite = ConstructObject<UPaperSprite>(Class, InParent, Name, Flags);
 
-	NewSprite->InitializeSprite(InitialTexture);
+	if (bUseSourceRegion)
+	{
+		NewSprite->InitializeSprite(InitialTexture, InitialSourceUV, InitialSourceDimension);
+	}
+	else
+	{
+		NewSprite->InitializeSprite(InitialTexture);
+	}
 
 	return NewSprite;
 }
