Paper2D: Super-experimental 2D terrain support

diff --git a/Engine/Plugins/Experimental/Paper2D/Source/Paper2D/Classes/PaperTerrainActor.h b/Engine/Plugins/Experimental/Paper2D/Source/Paper2D/Classes/PaperTerrainActor.h
new file mode 100644
--- /dev/null
+++ b/Engine/Plugins/Experimental/Paper2D/Source/Paper2D/Classes/PaperTerrainActor.h
@@ -0,0 +1,24 @@
+// Copyright 1998-2014 Epic Games, Inc. All Rights Reserved.
+
+#pragma once
+
+#include "PaperTerrainActor.generated.h"
+
+/**
+ * An instance of a piece of 2D terrain in the level
+ */
+
+UCLASS(DependsOn=UEngineTypes, BlueprintType, meta=(DisplayThumbnail = "true"))
+class PAPER2D_API APaperTerrainActor : public AActor
+{
+	GENERATED_UCLASS_BODY()
+
+	UPROPERTY()
+	TSubobjectPtr<class USceneComponent> DummyRoot;
+
+	UPROPERTY()
+	TSubobjectPtr<class UPaperTerrainSplineComponent> SplineComponent;
+
+	UPROPERTY(Category = Sprite, VisibleAnywhere)
+	TSubobjectPtr<class UPaperTerrainComponent> RenderComponent;
+};
diff --git a/Engine/Plugins/Experimental/Paper2D/Source/Paper2D/Classes/PaperTerrainComponent.h b/Engine/Plugins/Experimental/Paper2D/Source/Paper2D/Classes/PaperTerrainComponent.h
new file mode 100644
--- /dev/null
+++ b/Engine/Plugins/Experimental/Paper2D/Source/Paper2D/Classes/PaperTerrainComponent.h
@@ -0,0 +1,57 @@
+// Copyright 1998-2014 Epic Games, Inc. All Rights Reserved.
+
+#pragma once
+
+#include "PaperTerrainComponent.generated.h"
+
+/**
+ * The terrain visualization component for an associated spline component.
+ * This takes a 2D terrain material and instances sprite geometry along the spline path.
+ */
+
+UCLASS(BlueprintType)
+class PAPER2D_API UPaperTerrainComponent : public UPrimitiveComponent
+{
+	GENERATED_UCLASS_BODY()
+
+public:
+	UPROPERTY(Category=Sprite, EditAnywhere, BlueprintReadOnly)
+	class UPaperTerrainMaterial* TerrainMaterial;
+
+	UPROPERTY(Category = Sprite, EditAnywhere, BlueprintReadOnly)
+	bool bClosedSpline;
+
+	UPROPERTY(Category = Sprite, EditAnywhere, BlueprintReadOnly)
+	float TestScaleFactor;
+
+	UPROPERTY()
+	class UPaperTerrainSplineComponent* AssociatedSpline;
+
+	/** Random seed used for choosing which spline meshes to use. */
+	UPROPERTY(Category=Sprite, EditAnywhere)
+	int32 RandomSeed;
+
+protected:
+
+	UPROPERTY(Transient)
+	TArray<class UPaperRenderComponent*> SpawnedComponents;
+
+public:
+	// UObject interface
+#if WITH_EDITORONLY_DATA
+	virtual void PostEditChangeProperty(FPropertyChangedEvent& PropertyChangedEvent) override;
+#endif
+	// End of UObject interface
+
+	// UActorComponent interface
+	virtual void OnRegister() override;
+	virtual void OnUnregister() override;
+	// End of UActorComponent interface
+
+protected:
+	void SpawnSegment(class UPaperSprite* NewSprite, float Direction, float& RemainingSegStart, float& RemainingSegEnd);
+
+	void OnSplineEdited();
+	void DestroyExistingStuff();
+};
+
diff --git a/Engine/Plugins/Experimental/Paper2D/Source/Paper2D/Classes/PaperTerrainMaterial.h b/Engine/Plugins/Experimental/Paper2D/Source/Paper2D/Classes/PaperTerrainMaterial.h
new file mode 100644
--- /dev/null
+++ b/Engine/Plugins/Experimental/Paper2D/Source/Paper2D/Classes/PaperTerrainMaterial.h
@@ -0,0 +1,26 @@
+// Copyright 1998-2014 Epic Games, Inc. All Rights Reserved.
+
+#pragma once
+
+#include "PaperTerrainMaterial.generated.h"
+
+/**
+ * Paper Terrain Material
+ *
+ * 'Material' setup for a 2D terrain spline (stores references to sprites that will be instanced along the spline path, not actually related to UMaterialInterface).
+ */
+
+UCLASS(BlueprintType)
+class PAPER2D_API UPaperTerrainMaterial : public UDataAsset
+{
+	GENERATED_UCLASS_BODY()
+
+	UPROPERTY(Category = Sprite, EditAnywhere)
+	UPaperSprite* LeftCap;
+
+	UPROPERTY(Category=Sprite, EditAnywhere)
+	TArray<UPaperSprite*> BodySegments;
+
+	UPROPERTY(Category = Sprite, EditAnywhere)
+	UPaperSprite* RightCap;
+};
diff --git a/Engine/Plugins/Experimental/Paper2D/Source/Paper2D/Classes/PaperTerrainSplineComponent.h b/Engine/Plugins/Experimental/Paper2D/Source/Paper2D/Classes/PaperTerrainSplineComponent.h
new file mode 100644
--- /dev/null
+++ b/Engine/Plugins/Experimental/Paper2D/Source/Paper2D/Classes/PaperTerrainSplineComponent.h
@@ -0,0 +1,25 @@
+// Copyright 1998-2014 Epic Games, Inc. All Rights Reserved.
+
+#pragma once
+
+#include "Components/SplineComponent.h"
+#include "PaperTerrainSplineComponent.generated.h"
+
+//@TODO: Document
+UCLASS(BlueprintType)
+class PAPER2D_API UPaperTerrainSplineComponent : public USplineComponent
+{
+	GENERATED_UCLASS_BODY()
+
+public:
+	// UObject interface
+#if WITH_EDITORONLY_DATA
+	virtual void PostEditChangeProperty(FPropertyChangedEvent& PropertyChangedEvent) override;
+#endif
+	// End of UObject interface
+
+public:
+	// Triggered when the spline is edited
+	FSimpleDelegate OnSplineEdited;
+};
+
diff --git a/Engine/Plugins/Experimental/Paper2D/Source/Paper2D/Private/Terrain/PaperTerrainActor.cpp b/Engine/Plugins/Experimental/Paper2D/Source/Paper2D/Private/Terrain/PaperTerrainActor.cpp
new file mode 100644
--- /dev/null
+++ b/Engine/Plugins/Experimental/Paper2D/Source/Paper2D/Private/Terrain/PaperTerrainActor.cpp
@@ -0,0 +1,19 @@
+// Copyright 1998-2014 Epic Games, Inc. All Rights Reserved.
+
+#include "Paper2DPrivatePCH.h"
+
+//////////////////////////////////////////////////////////////////////////
+// APaperTerrainActor
+
+APaperTerrainActor::APaperTerrainActor(const FPostConstructInitializeProperties& PCIP)
+	: Super(PCIP)
+{
+	DummyRoot = PCIP.CreateDefaultSubobject<USceneComponent>(this, TEXT("RootComponent"));
+	SplineComponent = PCIP.CreateDefaultSubobject<UPaperTerrainSplineComponent>(this, TEXT("SplineComponent"));
+ 	RenderComponent = PCIP.CreateDefaultSubobject<UPaperTerrainComponent>(this, TEXT("RenderComponent"));
+ 
+	SplineComponent->AttachParent = DummyRoot;
+	RenderComponent->AttachParent = DummyRoot;
+	RenderComponent->AssociatedSpline = SplineComponent;
+	RootComponent = DummyRoot;
+}
diff --git a/Engine/Plugins/Experimental/Paper2D/Source/Paper2D/Private/Terrain/PaperTerrainComponent.cpp b/Engine/Plugins/Experimental/Paper2D/Source/Paper2D/Private/Terrain/PaperTerrainComponent.cpp
new file mode 100644
--- /dev/null
+++ b/Engine/Plugins/Experimental/Paper2D/Source/Paper2D/Private/Terrain/PaperTerrainComponent.cpp
@@ -0,0 +1,143 @@
+// Copyright 1998-2014 Epic Games, Inc. All Rights Reserved.
+
+#include "Paper2DPrivatePCH.h"
+#include "Components/SplineComponent.h"
+
+//////////////////////////////////////////////////////////////////////////
+// UPaperTerrainComponent
+
+UPaperTerrainComponent::UPaperTerrainComponent(const FPostConstructInitializeProperties& PCIP)
+	: Super(PCIP)
+{
+	TestScaleFactor = 1.0f;
+}
+
+#if WITH_EDITORONLY_DATA
+void UPaperTerrainComponent::PostEditChangeProperty(FPropertyChangedEvent& PropertyChangedEvent)
+{
+	Super::PostEditChangeProperty(PropertyChangedEvent);
+}
+#endif
+
+void UPaperTerrainComponent::OnSplineEdited()
+{
+	DestroyExistingStuff();
+
+	if ((AssociatedSpline != nullptr) && (TerrainMaterial != nullptr))
+	{
+		FRandomStream RandomStream(RandomSeed);
+
+		const FInterpCurveVector& SplineInfo = AssociatedSpline->SplineInfo;
+
+		float SplineLength = AssociatedSpline->GetSplineLength();
+
+		float RemainingSegStart = 0.0f;
+		float RemainingSegEnd = SplineLength;
+
+		SpawnSegment(TerrainMaterial->LeftCap, /*dir=*/ 1.0f, /*inout*/ RemainingSegStart, /*inout*/ RemainingSegEnd);
+		SpawnSegment(TerrainMaterial->RightCap, /*dir=*/ -1.0f, /*inout*/ RemainingSegStart, /*inout*/ RemainingSegEnd);
+
+		if (TerrainMaterial->BodySegments.Num() > 0)
+		{
+			while (RemainingSegEnd - RemainingSegStart > 0.0f)
+			{
+				const int32 BodySegmentIndex = RandomStream.GetUnsignedInt() % TerrainMaterial->BodySegments.Num();
+				UPaperSprite* BodySegment = TerrainMaterial->BodySegments[BodySegmentIndex];
+				SpawnSegment(BodySegment, /*dir=*/ 1.0f, /*inout*/ RemainingSegStart, /*inout*/ RemainingSegEnd);
+			}
+		}
+	}
+}
+
+void UPaperTerrainComponent::OnRegister()
+{
+	Super::OnRegister();
+
+	if (AssociatedSpline != nullptr)
+	{
+		AssociatedSpline->OnSplineEdited = FSimpleDelegate::CreateUObject(this, &UPaperTerrainComponent::OnSplineEdited);
+	}
+
+	OnSplineEdited();
+}
+
+void UPaperTerrainComponent::OnUnregister()
+{
+	Super::OnUnregister();
+
+	if (AssociatedSpline != nullptr)
+	{
+		AssociatedSpline->OnSplineEdited.Unbind();
+	}
+	DestroyExistingStuff();
+}
+
+void UPaperTerrainComponent::DestroyExistingStuff()
+{
+	for (UPaperRenderComponent* Component : SpawnedComponents)
+	{
+		Component->DetachFromParent();
+		Component->DestroyComponent();
+	}
+	SpawnedComponents.Empty();
+}
+
+void UPaperTerrainComponent::SpawnSegment(class UPaperSprite* NewSprite, float Direction, float& RemainingSegStart, float& RemainingSegEnd)
+{
+	if (NewSprite)
+	{
+		UPaperRenderComponent* NewRenderComp = ConstructObject<UPaperRenderComponent>(UPaperRenderComponent::StaticClass(), GetOuter(), NAME_None, RF_Transactional);
+		NewRenderComp->AttachTo(this);
+		NewRenderComp->bCreatedByConstructionScript = bCreatedByConstructionScript;
+
+		NewRenderComp->SetSprite(NewSprite);
+
+
+		//@TODO: Hacky!
+		const float SpriteWidthUnsafe = NewSprite->GetRenderBounds().BoxExtent.ProjectOnTo(PaperAxisX).Size() * 2.0f;
+		const float SpriteWidth = (SpriteWidthUnsafe < KINDA_SMALL_NUMBER) ? 1.0f : SpriteWidthUnsafe; // ensure we terminate in the outer loop even with malformed render data
+
+		float Position = 0.0f;
+		if (Direction < 0.0f)
+		{
+			Position = RemainingSegEnd - SpriteWidth * 0.5f;
+			RemainingSegEnd -= SpriteWidth;
+		}
+		else
+		{
+			Position = RemainingSegStart + SpriteWidth * 0.5f;
+			RemainingSegStart += SpriteWidth;
+		}
+
+		// Position and orient the sprite
+		const float Param = AssociatedSpline->SplineReparamTable.Eval(Position, 0.f);
+		const FVector Position3D = AssociatedSpline->SplineInfo.Eval(Param, FVector::ZeroVector);
+		
+		const FVector Tangent = AssociatedSpline->SplineInfo.EvalDerivative(Param, FVector(1.0f, 0.0f, 0.0f)).SafeNormal();
+		const FVector NormalEst = AssociatedSpline->SplineInfo.EvalSecondDerivative(Param, FVector(0.0f, 1.0f, 0.0f)).SafeNormal();
+		const FVector Bitangent = FVector::CrossProduct(Tangent, NormalEst);
+		const FVector Normal = FVector::CrossProduct(Bitangent, Tangent);
+		const FVector Floop = FVector::CrossProduct(PaperAxisZ, Tangent);
+		//const FRotator Rotation3D(Tangent.Rotation());
+		
+		FTransform LocalTransform(Tangent, PaperAxisZ, Floop, Position3D);
+
+		const FRotator Rotation3D(LocalTransform.GetRotation());
+		const FVector WorldPos = ComponentToWorld.TransformPosition(Position3D);
+		
+		DrawDebugLine(GetWorld(), WorldPos, WorldPos + ComponentToWorld.TransformVector(Tangent) * 10.0f, FColor::Red, true, 1.0f);
+// 		DrawDebugLine(GetWorld(), WorldPos, WorldPos + ComponentToWorld.TransformVector(NormalEst) * 10.0f, FColor::Green, true, 1.0f);
+// 		DrawDebugLine(GetWorld(), WorldPos, WorldPos + ComponentToWorld.TransformVector(Bitangent) * 10.0f, FColor::Blue, true, 1.0f);
+		DrawDebugLine(GetWorld(), WorldPos, WorldPos + ComponentToWorld.TransformVector(Floop) * 10.0f, FColor::Yellow, true, 1.0f);
+// 		DrawDebugPoint(GetWorld(), WorldPos, 4.0f, FColor::White, true, 1.0f);
+
+		NewRenderComp->RelativeLocation = Position3D;
+		NewRenderComp->RelativeRotation = Rotation3D;
+		NewRenderComp->RelativeScale3D = FVector(TestScaleFactor, TestScaleFactor, TestScaleFactor);
+
+		NewRenderComp->RegisterComponentWithWorld(GetWorld());
+		SpawnedComponents.Add(NewRenderComp);
+
+		
+	}
+}
diff --git a/Engine/Plugins/Experimental/Paper2D/Source/Paper2D/Private/Terrain/PaperTerrainMaterial.cpp b/Engine/Plugins/Experimental/Paper2D/Source/Paper2D/Private/Terrain/PaperTerrainMaterial.cpp
new file mode 100644
--- /dev/null
+++ b/Engine/Plugins/Experimental/Paper2D/Source/Paper2D/Private/Terrain/PaperTerrainMaterial.cpp
@@ -0,0 +1,11 @@
+// Copyright 1998-2014 Epic Games, Inc. All Rights Reserved.
+
+#include "Paper2DPrivatePCH.h"
+
+//////////////////////////////////////////////////////////////////////////
+// UPaperTerrainMaterial
+
+UPaperTerrainMaterial::UPaperTerrainMaterial(const FPostConstructInitializeProperties& PCIP)
+	: Super(PCIP)
+{
+}
diff --git a/Engine/Plugins/Experimental/Paper2D/Source/Paper2D/Private/Terrain/PaperTerrainSplineComponent.cpp b/Engine/Plugins/Experimental/Paper2D/Source/Paper2D/Private/Terrain/PaperTerrainSplineComponent.cpp
new file mode 100644
--- /dev/null
+++ b/Engine/Plugins/Experimental/Paper2D/Source/Paper2D/Private/Terrain/PaperTerrainSplineComponent.cpp
@@ -0,0 +1,19 @@
+// Copyright 1998-2014 Epic Games, Inc. All Rights Reserved.
+
+#include "Paper2DPrivatePCH.h"
+
+//////////////////////////////////////////////////////////////////////////
+// UPaperTerrainSplineComponent
+
+UPaperTerrainSplineComponent::UPaperTerrainSplineComponent(const FPostConstructInitializeProperties& PCIP)
+	: Super(PCIP)
+{
+}
+
+#if WITH_EDITORONLY_DATA
+void UPaperTerrainSplineComponent::PostEditChangeProperty(FPropertyChangedEvent& PropertyChangedEvent)
+{
+	OnSplineEdited.ExecuteIfBound();
+	Super::PostEditChangeProperty(PropertyChangedEvent);
+}
+#endif
