Paper2D: Fix bug in UPaperRenderComponent::GetSocketTransform that would cause components attached to a sprite to behave as if they weren't attached

diff --git a/Engine/Plugins/Experimental/Paper2D/Source/Paper2D/Private/PaperRenderComponent.cpp b/Engine/Plugins/Experimental/Paper2D/Source/Paper2D/Private/PaperRenderComponent.cpp
--- a/Engine/Plugins/Experimental/Paper2D/Source/Paper2D/Private/PaperRenderComponent.cpp
+++ b/Engine/Plugins/Experimental/Paper2D/Source/Paper2D/Private/PaperRenderComponent.cpp
@@ -37,7 +37,7 @@ FPrimitiveSceneProxy* UPaperRenderComponent::CreateSceneProxy()
 	return NewProxy;
 }
 
-FBoxSphereBounds UPaperRenderComponent::CalcBounds(const FTransform & LocalToWorld) const
+FBoxSphereBounds UPaperRenderComponent::CalcBounds(const FTransform& LocalToWorld) const
 {
 	if (SourceSprite != NULL)
 	{
@@ -122,7 +122,7 @@ FTransform UPaperRenderComponent::GetSocketTransform(FName InSocketName, ERelati
 		}
 	}
 
-	return FTransform::Identity;
+	return Super::GetSocketTransform(InSocketName, TransformSpace);
 }
 
 void UPaperRenderComponent::QuerySupportedSockets(TArray<FComponentSocketDescription>& OutSockets) const
