Paper2D: Add thumbnail metadata to paper terrain materials

diff --git a/Engine/Plugins/Experimental/Paper2D/Source/Paper2D/Classes/PaperTerrainMaterial.h b/Engine/Plugins/Experimental/Paper2D/Source/Paper2D/Classes/PaperTerrainMaterial.h
--- a/Engine/Plugins/Experimental/Paper2D/Source/Paper2D/Classes/PaperTerrainMaterial.h
+++ b/Engine/Plugins/Experimental/Paper2D/Source/Paper2D/Classes/PaperTerrainMaterial.h
@@ -15,12 +15,12 @@ class PAPER2D_API UPaperTerrainMaterial : public UDataAsset
 {
 	GENERATED_UCLASS_BODY()
 
-	UPROPERTY(Category = Sprite, EditAnywhere)
+	UPROPERTY(Category=Sprite, EditAnywhere, meta=(DisplayThumbnail = "true"))
 	UPaperSprite* LeftCap;
 
-	UPROPERTY(Category=Sprite, EditAnywhere)
+	UPROPERTY(Category=Sprite, EditAnywhere, meta=(DisplayThumbnail = "true"))
 	TArray<UPaperSprite*> BodySegments;
 
-	UPROPERTY(Category = Sprite, EditAnywhere)
+	UPROPERTY(Category=Sprite, EditAnywhere, meta=(DisplayThumbnail = "true"))
 	UPaperSprite* RightCap;
 };
