Paper2D: Implement APaperTerrainActor::GetReferencedContentObjects to allow Ctrl+E editing of the terrain material

diff --git a/Engine/Plugins/Experimental/Paper2D/Source/Paper2D/Classes/PaperTerrainActor.h b/Engine/Plugins/Experimental/Paper2D/Source/Paper2D/Classes/PaperTerrainActor.h
--- a/Engine/Plugins/Experimental/Paper2D/Source/Paper2D/Classes/PaperTerrainActor.h
+++ b/Engine/Plugins/Experimental/Paper2D/Source/Paper2D/Classes/PaperTerrainActor.h
@@ -8,7 +8,7 @@
  * An instance of a piece of 2D terrain in the level
  */
 
-UCLASS(DependsOn=UEngineTypes, BlueprintType, meta=(DisplayThumbnail = "true"))
+UCLASS(DependsOn=UEngineTypes, BlueprintType)
 class PAPER2D_API APaperTerrainActor : public AActor
 {
 	GENERATED_UCLASS_BODY()
@@ -21,4 +21,10 @@ class PAPER2D_API APaperTerrainActor : public AActor
 
 	UPROPERTY(Category = Sprite, VisibleAnywhere)
 	TSubobjectPtr<class UPaperTerrainComponent> RenderComponent;
+
+	// AActor interface
+#if WITH_EDITOR
+	virtual bool GetReferencedContentObjects(TArray<UObject*>& Objects) const override;
+#endif
+	// End of AActor interface
 };
diff --git a/Engine/Plugins/Experimental/Paper2D/Source/Paper2D/Private/Terrain/PaperTerrainActor.cpp b/Engine/Plugins/Experimental/Paper2D/Source/Paper2D/Private/Terrain/PaperTerrainActor.cpp
--- a/Engine/Plugins/Experimental/Paper2D/Source/Paper2D/Private/Terrain/PaperTerrainActor.cpp
+++ b/Engine/Plugins/Experimental/Paper2D/Source/Paper2D/Private/Terrain/PaperTerrainActor.cpp
@@ -17,3 +17,14 @@ APaperTerrainActor::APaperTerrainActor(const FPostConstructInitializeProperties&
 	RenderComponent->AssociatedSpline = SplineComponent;
 	RootComponent = DummyRoot;
 }
+
+#if WITH_EDITOR
+bool APaperTerrainActor::GetReferencedContentObjects(TArray<UObject*>& Objects) const
+{
+	if (RenderComponent->TerrainMaterial != nullptr)
+	{
+		Objects.Add(RenderComponent->TerrainMaterial);
+	}
+	return true;
+}
+#endif
