Paper2D: Added vertex deformation for sprite segments on 2D terrain splines (super-WIP)

diff --git a/Engine/Plugins/Experimental/Paper2D/Source/Paper2D/Classes/PaperTerrainComponent.h b/Engine/Plugins/Experimental/Paper2D/Source/Paper2D/Classes/PaperTerrainComponent.h
--- a/Engine/Plugins/Experimental/Paper2D/Source/Paper2D/Classes/PaperTerrainComponent.h
+++ b/Engine/Plugins/Experimental/Paper2D/Source/Paper2D/Classes/PaperTerrainComponent.h
@@ -2,6 +2,7 @@
 
 #pragma once
 
+#include "SpriteDrawCall.h"
 #include "PaperTerrainComponent.generated.h"
 
 /**
@@ -32,12 +33,13 @@ public:
 	int32 RandomSeed;
 
 protected:
-
-	UPROPERTY(Transient)
-	TArray<class UPaperRenderComponent*> SpawnedComponents;
+	/** Number of steps per spline segment to place in the reparameterization table */
+	UPROPERTY(EditAnywhere, Category=Sprite, meta=(ClampMin=4, UIMin=4))
+	int32 ReparamStepsPerSegment;
 
 public:
 	// UObject interface
+	virtual const UObject* AdditionalStatObject() const override;
 #if WITH_EDITORONLY_DATA
 	virtual void PostEditChangeProperty(FPropertyChangedEvent& PropertyChangedEvent) override;
 #endif
@@ -48,10 +50,18 @@ public:
 	virtual void OnUnregister() override;
 	// End of UActorComponent interface
 
+	// UPrimitiveComponent interface
+	virtual FPrimitiveSceneProxy* CreateSceneProxy() override;
+	virtual FBoxSphereBounds CalcBounds(const FTransform& LocalToWorld) const override;
+	// End of UPrimitiveComponent interface
+
 protected:
 	void SpawnSegment(class UPaperSprite* NewSprite, float Direction, float& RemainingSegStart, float& RemainingSegEnd);
 
 	void OnSplineEdited();
-	void DestroyExistingStuff();
+
+	TArray<FSpriteDrawCallRecord> GeneratedSpriteGeometry;
+
+	FTransform GetTransformAtDistance(float InDistance) const;
 };
 
diff --git a/Engine/Plugins/Experimental/Paper2D/Source/Paper2D/Private/Terrain/PaperTerrainComponent.cpp b/Engine/Plugins/Experimental/Paper2D/Source/Paper2D/Private/Terrain/PaperTerrainComponent.cpp
--- a/Engine/Plugins/Experimental/Paper2D/Source/Paper2D/Private/Terrain/PaperTerrainComponent.cpp
+++ b/Engine/Plugins/Experimental/Paper2D/Source/Paper2D/Private/Terrain/PaperTerrainComponent.cpp
@@ -3,28 +3,202 @@
 #include "Paper2DPrivatePCH.h"
 #include "Components/SplineComponent.h"
 
+#include "PaperRenderSceneProxy.h"
+
+//////////////////////////////////////////////////////////////////////////
+// UPaperTerrainComponent
+
+class FPaperTerrainSceneProxy : public FPaperRenderSceneProxy
+{
+public:
+	FPaperTerrainSceneProxy(const UPaperTerrainComponent* InComponent, const TArray<FSpriteDrawCallRecord>& InGeometry, UMaterialInterface* InMaterial);
+
+	// FPrimitiveSceneProxy interface.
+	//virtual void DrawDynamicElements(FPrimitiveDrawInterface* PDI, const FSceneView* View) override;
+	// End of FPrimitiveSceneProxy interface.
+};
+
+FPaperTerrainSceneProxy::FPaperTerrainSceneProxy(const UPaperTerrainComponent* InComponent, const TArray<FSpriteDrawCallRecord>& InGeometry, UMaterialInterface* InMaterial)
+	: FPaperRenderSceneProxy(InComponent)
+{
+	BatchedSprites = InGeometry;
+
+	//@TODO: Should make the input geom + material be a pair, allowing multiple materials for solid/opaque
+	Material = InMaterial;
+	if (Material)
+	{
+		MaterialRelevance = Material->GetRelevance();
+	}
+}
+
 //////////////////////////////////////////////////////////////////////////
 // UPaperTerrainComponent
 
 UPaperTerrainComponent::UPaperTerrainComponent(const FPostConstructInitializeProperties& PCIP)
 	: Super(PCIP)
+	, ReparamStepsPerSegment(4)
 {
 	TestScaleFactor = 1.0f;
 }
 
+const UObject* UPaperTerrainComponent::AdditionalStatObject() const
+{
+	return TerrainMaterial;
+}
+
 #if WITH_EDITORONLY_DATA
 void UPaperTerrainComponent::PostEditChangeProperty(FPropertyChangedEvent& PropertyChangedEvent)
 {
+	if (PropertyChangedEvent.Property != nullptr)
+	{
+		const FName PropertyName(PropertyChangedEvent.Property->GetFName());
+		if (PropertyName == GET_MEMBER_NAME_CHECKED(USplineComponent, ReparamStepsPerSegment))
+		{
+			if (AssociatedSpline != nullptr)
+			{
+				AssociatedSpline->ReparamStepsPerSegment = ReparamStepsPerSegment;
+				AssociatedSpline->UpdateSplineReparamTable();
+			}
+		}
+	}
+
 	Super::PostEditChangeProperty(PropertyChangedEvent);
 }
 #endif
 
+void UPaperTerrainComponent::OnRegister()
+{
+	Super::OnRegister();
+
+	if (AssociatedSpline != nullptr)
+	{
+		AssociatedSpline->OnSplineEdited = FSimpleDelegate::CreateUObject(this, &UPaperTerrainComponent::OnSplineEdited);
+	}
+
+	OnSplineEdited();
+}
+
+void UPaperTerrainComponent::OnUnregister()
+{
+	Super::OnUnregister();
+
+	if (AssociatedSpline != nullptr)
+	{
+		AssociatedSpline->OnSplineEdited.Unbind();
+	}
+}
+
+FPrimitiveSceneProxy* UPaperTerrainComponent::CreateSceneProxy()
+{
+	//@TODO: Super-hacky
+	UMaterialInterface* Material = nullptr;
+	if (TerrainMaterial != nullptr)
+	{
+		if (TerrainMaterial->LeftCap)
+		{
+			Material = TerrainMaterial->LeftCap->GetDefaultMaterial();
+		}
+	}
+
+	FPaperTerrainSceneProxy* NewProxy = new FPaperTerrainSceneProxy(this, GeneratedSpriteGeometry, Material);
+	return NewProxy;
+}
+
+FBoxSphereBounds UPaperTerrainComponent::CalcBounds(const FTransform& LocalToWorld) const
+{
+	// Determine the rendering bounds
+	FBoxSphereBounds LocalRenderBounds;
+	{
+		FBox BoundingBox(ForceInit);
+
+		for (const FSpriteDrawCallRecord& DrawCall : GeneratedSpriteGeometry)
+		{
+			for (const FVector4& VertXYUV : DrawCall.RenderVerts)
+			{
+				const FVector Vert((PaperAxisX * VertXYUV.X) + (PaperAxisY * VertXYUV.Y));
+				BoundingBox += Vert;
+			}
+		}
+
+		// Make the whole thing a single unit 'deep'
+		const FVector HalfThicknessVector = 0.5f * PaperAxisZ;
+		BoundingBox.Min -= HalfThicknessVector;
+		BoundingBox.Max += HalfThicknessVector;
+
+		LocalRenderBounds = FBoxSphereBounds(BoundingBox);
+	}
+
+	// Graphics bounds.
+	FBoxSphereBounds NewBounds = LocalRenderBounds.TransformBy(LocalToWorld);
+
+#if 0
+	// Add bounds of collision geometry (if present).
+	if (UBodySetup* BodySetup = GetBodySetup())
+	{
+		const FBox AggGeomBox = BodySetup->AggGeom.CalcAABB(LocalToWorld);
+		if (AggGeomBox.IsValid)
+		{
+			NewBounds = Union(NewBounds, FBoxSphereBounds(AggGeomBox));
+		}
+	}
+#endif
+
+	// Apply bounds scale
+	NewBounds.BoxExtent *= BoundsScale;
+	NewBounds.SphereRadius *= BoundsScale;
+
+	return NewBounds;
+}
+
+FTransform UPaperTerrainComponent::GetTransformAtDistance(float InDistance) const
+{
+	const float SplineLength = AssociatedSpline->GetSplineLength();
+	InDistance = FMath::Clamp<float>(InDistance, 0.0f, SplineLength);
+
+	const float Param = AssociatedSpline->SplineReparamTable.Eval(InDistance, 0.0f);
+	const FVector Position3D = AssociatedSpline->SplineInfo.Eval(Param, FVector::ZeroVector);
+
+	const FVector Tangent = AssociatedSpline->SplineInfo.EvalDerivative(Param, FVector(1.0f, 0.0f, 0.0f)).SafeNormal();
+	const FVector NormalEst = AssociatedSpline->SplineInfo.EvalSecondDerivative(Param, FVector(0.0f, 1.0f, 0.0f)).SafeNormal();
+	const FVector Bitangent = FVector::CrossProduct(Tangent, NormalEst);
+	const FVector Normal = FVector::CrossProduct(Bitangent, Tangent);
+	const FVector Floop = FVector::CrossProduct(PaperAxisZ, Tangent);
+
+	FTransform LocalTransform(Tangent, PaperAxisZ, Floop, Position3D);
+
+	LocalTransform = FTransform(FRotator(0.0f, 180.0f, 0.0f), FVector::ZeroVector) * LocalTransform;
+
+#if PAPER_TERRAIN_DRAW_DEBUG
+	FTransform WorldTransform = LocalTransform * ComponentToWorld;
+
+	const float Time = 2.5f;
+
+	DrawDebugCoordinateSystem(GetWorld(), WorldTransform.GetLocation(), FRotator(WorldTransform.GetRotation()), 15.0f, true, Time, SDPG_Foreground);
+// 	FVector WorldPos = ComponentToWorld.TransformPosition(Position3D);
+// 	WorldPos.Y -= 0.01;
+// 
+// 	//DrawDebugLine(GetWorld(), WorldPos, WorldPos + ComponentToWorld.TransformVector(Tangent) * 10.0f, FColor::Red, true, Time);
+// 	// 		DrawDebugLine(GetWorld(), WorldPos, WorldPos + ComponentToWorld.TransformVector(NormalEst) * 10.0f, FColor::Green, true, Time);
+// 	// 		DrawDebugLine(GetWorld(), WorldPos, WorldPos + ComponentToWorld.TransformVector(Bitangent) * 10.0f, FColor::Blue, true, Time);
+// 	//DrawDebugLine(GetWorld(), WorldPos, WorldPos + ComponentToWorld.TransformVector(Floop) * 10.0f, FColor::Yellow, true, Time);
+// 	// 		DrawDebugPoint(GetWorld(), WorldPos, 4.0f, FColor::White, true, 1.0f);
+#endif
+
+	return LocalTransform;
+}
+
 void UPaperTerrainComponent::OnSplineEdited()
 {
-	DestroyExistingStuff();
+	GeneratedSpriteGeometry.Empty();
 
 	if ((AssociatedSpline != nullptr) && (TerrainMaterial != nullptr))
 	{
+		if (AssociatedSpline->ReparamStepsPerSegment != ReparamStepsPerSegment)
+		{
+			AssociatedSpline->ReparamStepsPerSegment = ReparamStepsPerSegment;
+			AssociatedSpline->UpdateSplineReparamTable();
+		}
+
 		FRandomStream RandomStream(RandomSeed);
 
 		const FInterpCurveVector& SplineInfo = AssociatedSpline->SplineInfo;
@@ -46,55 +220,45 @@ void UPaperTerrainComponent::OnSplineEdited()
 				SpawnSegment(BodySegment, /*dir=*/ 1.0f, /*inout*/ RemainingSegStart, /*inout*/ RemainingSegEnd);
 			}
 		}
-	}
-}
 
-void UPaperTerrainComponent::OnRegister()
-{
-	Super::OnRegister();
 
-	if (AssociatedSpline != nullptr)
-	{
-		AssociatedSpline->OnSplineEdited = FSimpleDelegate::CreateUObject(this, &UPaperTerrainComponent::OnSplineEdited);
+		// Draw debug frames at the start and end of the spline
+#if PAPER_TERRAIN_DRAW_DEBUG
+		{
+			const float Time = 5.0f;
+			{
+				FTransform WorldTransform = GetTransformAtDistance(0.0f) * ComponentToWorld;
+				DrawDebugCoordinateSystem(GetWorld(), WorldTransform.GetLocation(), FRotator(WorldTransform.GetRotation()), 30.0f, true, Time, SDPG_Foreground);
+			}
+			{
+				FTransform WorldTransform = GetTransformAtDistance(SplineLength) * ComponentToWorld;
+				DrawDebugCoordinateSystem(GetWorld(), WorldTransform.GetLocation(), FRotator(WorldTransform.GetRotation()), 30.0f, true, Time, SDPG_Foreground);
+			}
+		}
+#endif
 	}
 
-	OnSplineEdited();
+	RecreateRenderState_Concurrent();
 }
 
-void UPaperTerrainComponent::OnUnregister()
+FBox2D GetSpriteRenderDataBounds2D(const TArray<FVector4>& Data)
 {
-	Super::OnUnregister();
-
-	if (AssociatedSpline != nullptr)
+	FBox2D Bounds(ForceInit);
+	for (const FVector4& XYUV : Data)
 	{
-		AssociatedSpline->OnSplineEdited.Unbind();
+		Bounds += FVector2D(XYUV.X, XYUV.Y);
 	}
-	DestroyExistingStuff();
-}
 
-void UPaperTerrainComponent::DestroyExistingStuff()
-{
-	for (UPaperRenderComponent* Component : SpawnedComponents)
-	{
-		Component->DetachFromParent();
-		Component->DestroyComponent();
-	}
-	SpawnedComponents.Empty();
+	return Bounds;
 }
 
 void UPaperTerrainComponent::SpawnSegment(class UPaperSprite* NewSprite, float Direction, float& RemainingSegStart, float& RemainingSegEnd)
 {
 	if (NewSprite)
 	{
-		UPaperRenderComponent* NewRenderComp = ConstructObject<UPaperRenderComponent>(UPaperRenderComponent::StaticClass(), GetOuter(), NAME_None, RF_Transactional);
-		NewRenderComp->AttachTo(this);
-		NewRenderComp->bCreatedByConstructionScript = bCreatedByConstructionScript;
-
-		NewRenderComp->SetSprite(NewSprite);
+		const FBox2D Bounds2D = GetSpriteRenderDataBounds2D(NewSprite->BakedRenderData);
 
-
-		//@TODO: Hacky!
-		const float SpriteWidthUnsafe = NewSprite->GetRenderBounds().BoxExtent.ProjectOnTo(PaperAxisX).Size() * 2.0f;
+		const float SpriteWidthUnsafe = Bounds2D.GetSize().X; //NewSprite->GetRenderBounds().BoxExtent.ProjectOnTo(PaperAxisX).Size() * 2.0f;
 		const float SpriteWidth = (SpriteWidthUnsafe < KINDA_SMALL_NUMBER) ? 1.0f : SpriteWidthUnsafe; // ensure we terminate in the outer loop even with malformed render data
 
 		float Position = 0.0f;
@@ -109,35 +273,24 @@ void UPaperTerrainComponent::SpawnSegment(class UPaperSprite* NewSprite, float D
 			RemainingSegStart += SpriteWidth;
 		}
 
-		// Position and orient the sprite
-		const float Param = AssociatedSpline->SplineReparamTable.Eval(Position, 0.f);
-		const FVector Position3D = AssociatedSpline->SplineInfo.Eval(Param, FVector::ZeroVector);
-		
-		const FVector Tangent = AssociatedSpline->SplineInfo.EvalDerivative(Param, FVector(1.0f, 0.0f, 0.0f)).SafeNormal();
-		const FVector NormalEst = AssociatedSpline->SplineInfo.EvalSecondDerivative(Param, FVector(0.0f, 1.0f, 0.0f)).SafeNormal();
-		const FVector Bitangent = FVector::CrossProduct(Tangent, NormalEst);
-		const FVector Normal = FVector::CrossProduct(Bitangent, Tangent);
-		const FVector Floop = FVector::CrossProduct(PaperAxisZ, Tangent);
-		//const FRotator Rotation3D(Tangent.Rotation());
-		
-		FTransform LocalTransform(Tangent, PaperAxisZ, Floop, Position3D);
-
-		const FRotator Rotation3D(LocalTransform.GetRotation());
-		const FVector WorldPos = ComponentToWorld.TransformPosition(Position3D);
-		
-		DrawDebugLine(GetWorld(), WorldPos, WorldPos + ComponentToWorld.TransformVector(Tangent) * 10.0f, FColor::Red, true, 1.0f);
-// 		DrawDebugLine(GetWorld(), WorldPos, WorldPos + ComponentToWorld.TransformVector(NormalEst) * 10.0f, FColor::Green, true, 1.0f);
-// 		DrawDebugLine(GetWorld(), WorldPos, WorldPos + ComponentToWorld.TransformVector(Bitangent) * 10.0f, FColor::Blue, true, 1.0f);
-		DrawDebugLine(GetWorld(), WorldPos, WorldPos + ComponentToWorld.TransformVector(Floop) * 10.0f, FColor::Yellow, true, 1.0f);
-// 		DrawDebugPoint(GetWorld(), WorldPos, 4.0f, FColor::White, true, 1.0f);
-
-		NewRenderComp->RelativeLocation = Position3D;
-		NewRenderComp->RelativeRotation = Rotation3D;
-		NewRenderComp->RelativeScale3D = FVector(TestScaleFactor, TestScaleFactor, TestScaleFactor);
-
-		NewRenderComp->RegisterComponentWithWorld(GetWorld());
-		SpawnedComponents.Add(NewRenderComp);
-
-		
+
+		const FTransform LocalTransform(GetTransformAtDistance(Position));
+
+		FSpriteDrawCallRecord& Record = *new (GeneratedSpriteGeometry) FSpriteDrawCallRecord();
+		Record.BuildFromSprite(NewSprite);
+
+		for (FVector4& XYUV : Record.RenderVerts)
+		{
+			const FTransform LocalTransformAtX(GetTransformAtDistance(Position + XYUV.X));
+			const FVector SourceVector = (XYUV.Y * PaperAxisY);
+			const FVector NewVector = LocalTransformAtX.TransformPosition(SourceVector);
+
+			const float NewX = FVector::DotProduct(NewVector, PaperAxisX);
+			const float NewY = FVector::DotProduct(NewVector, PaperAxisY);
+
+			XYUV.X = NewX;
+			XYUV.Y = NewY;
+		}
 	}
 }
+
