Paper2D: Add tint color to UPaperTerrainComponent

diff --git a/Engine/Plugins/Experimental/Paper2D/Source/Paper2D/Classes/PaperTerrainComponent.h b/Engine/Plugins/Experimental/Paper2D/Source/Paper2D/Classes/PaperTerrainComponent.h
--- a/Engine/Plugins/Experimental/Paper2D/Source/Paper2D/Classes/PaperTerrainComponent.h
+++ b/Engine/Plugins/Experimental/Paper2D/Source/Paper2D/Classes/PaperTerrainComponent.h
@@ -33,6 +33,10 @@ public:
 	int32 RandomSeed;
 
 protected:
+	// The color of the terrain (passed to the sprite material as a vertex color)
+	UPROPERTY(Category=Sprite, BlueprintReadOnly, Interp)
+	FLinearColor TerrainColor;
+
 	/** Number of steps per spline segment to place in the reparameterization table */
 	UPROPERTY(EditAnywhere, Category=Sprite, meta=(ClampMin=4, UIMin=4))
 	int32 ReparamStepsPerSegment;
@@ -55,6 +59,10 @@ public:
 	virtual FBoxSphereBounds CalcBounds(const FTransform& LocalToWorld) const override;
 	// End of UPrimitiveComponent interface
 
+	// Set color of the terrain
+	UFUNCTION(BlueprintCallable, Category = "Sprite")
+	void SetTerrainColor(FLinearColor NewColor);
+
 protected:
 	void SpawnSegment(class UPaperSprite* NewSprite, float Direction, float& RemainingSegStart, float& RemainingSegEnd);
 
diff --git a/Engine/Plugins/Experimental/Paper2D/Source/Paper2D/Private/Terrain/PaperTerrainComponent.cpp b/Engine/Plugins/Experimental/Paper2D/Source/Paper2D/Private/Terrain/PaperTerrainComponent.cpp
--- a/Engine/Plugins/Experimental/Paper2D/Source/Paper2D/Private/Terrain/PaperTerrainComponent.cpp
+++ b/Engine/Plugins/Experimental/Paper2D/Source/Paper2D/Private/Terrain/PaperTerrainComponent.cpp
@@ -36,6 +36,7 @@ FPaperTerrainSceneProxy::FPaperTerrainSceneProxy(const UPaperTerrainComponent* I
 
 UPaperTerrainComponent::UPaperTerrainComponent(const FPostConstructInitializeProperties& PCIP)
 	: Super(PCIP)
+	, TerrainColor(FLinearColor::White)
 	, ReparamStepsPerSegment(4)
 {
 	TestScaleFactor = 1.0f;
@@ -278,6 +279,7 @@ void UPaperTerrainComponent::SpawnSegment(class UPaperSprite* NewSprite, float D
 
 		FSpriteDrawCallRecord& Record = *new (GeneratedSpriteGeometry) FSpriteDrawCallRecord();
 		Record.BuildFromSprite(NewSprite);
+		Record.Color = TerrainColor;
 
 		for (FVector4& XYUV : Record.RenderVerts)
 		{
@@ -294,3 +296,14 @@ void UPaperTerrainComponent::SpawnSegment(class UPaperSprite* NewSprite, float D
 	}
 }
 
+void UPaperTerrainComponent::SetTerrainColor(FLinearColor NewColor)
+{
+	// Can't set color on a static component
+	if (!(IsRegistered() && (Mobility == EComponentMobility::Static)) && (TerrainColor != NewColor))
+	{
+		TerrainColor = NewColor;
+
+		//@TODO: Should we send immediately?
+		MarkRenderDynamicDataDirty();
+	}
+}
