Paper2D: Add support to UPaperTerrainComponent for each segment sprite possibly using a different material

diff --git a/Engine/Plugins/Experimental/Paper2D/Source/Paper2D/Classes/PaperTerrainComponent.h b/Engine/Plugins/Experimental/Paper2D/Source/Paper2D/Classes/PaperTerrainComponent.h
--- a/Engine/Plugins/Experimental/Paper2D/Source/Paper2D/Classes/PaperTerrainComponent.h
+++ b/Engine/Plugins/Experimental/Paper2D/Source/Paper2D/Classes/PaperTerrainComponent.h
@@ -5,6 +5,12 @@
 #include "SpriteDrawCall.h"
 #include "PaperTerrainComponent.generated.h"
 
+struct FPaperTerrainMaterialPair
+{
+	TArray<FSpriteDrawCallRecord> Records;
+	UMaterialInterface* Material;
+};
+
 /**
  * The terrain visualization component for an associated spline component.
  * This takes a 2D terrain material and instances sprite geometry along the spline path.
@@ -68,8 +74,8 @@ protected:
 
 	void OnSplineEdited();
 
-	TArray<FSpriteDrawCallRecord> GeneratedSpriteGeometry;
-
+	TArray<FPaperTerrainMaterialPair> GeneratedSpriteGeometry;
+	
 	FTransform GetTransformAtDistance(float InDistance) const;
 };
 
diff --git a/Engine/Plugins/Experimental/Paper2D/Source/Paper2D/Private/PaperRenderSceneProxy.cpp b/Engine/Plugins/Experimental/Paper2D/Source/Paper2D/Private/PaperRenderSceneProxy.cpp
--- a/Engine/Plugins/Experimental/Paper2D/Source/Paper2D/Private/PaperRenderSceneProxy.cpp
+++ b/Engine/Plugins/Experimental/Paper2D/Source/Paper2D/Private/PaperRenderSceneProxy.cpp
@@ -297,25 +297,32 @@ void FPaperRenderSceneProxy::DrawDynamicElements(FPrimitiveDrawInterface* PDI, c
 #endif
 }
 
+FVertexFactory* FPaperRenderSceneProxy::GetPaperSpriteVertexFactory() const
+{
+	static TGlobalResource<FPaperSpriteVertexFactory> GPaperSpriteVertexFactory;
+	return &GPaperSpriteVertexFactory;
+}
+
 void FPaperRenderSceneProxy::DrawDynamicElements_RichMesh(FPrimitiveDrawInterface* PDI, const FSceneView* View, bool bUseOverrideColor, const FLinearColor& OverrideColor)
 {
-	if (Material == nullptr)
+	if (Material != nullptr)
 	{
-		return;
+		DrawBatch(PDI, View, bUseOverrideColor, OverrideColor, Material, BatchedSprites);
 	}
+}
 
+void FPaperRenderSceneProxy::DrawBatch(FPrimitiveDrawInterface* PDI, const FSceneView* View, bool bUseOverrideColor, const FLinearColor& OverrideColor, class UMaterialInterface* BatchMaterial, const TArray<FSpriteDrawCallRecord>& Batch)
+{
 	const uint8 DPG = GetDepthPriorityGroup(View);
 	
-	static TGlobalResource<FPaperSpriteVertexFactory> GPaperSpriteVertexFactory;
+	FVertexFactory* VertexFactory = GetPaperSpriteVertexFactory();
 
-	for (int32 BatchIndex = 0; BatchIndex < BatchedSprites.Num(); ++BatchIndex)
+	for (const FSpriteDrawCallRecord& Record : Batch)
 	{
-		const FSpriteDrawCallRecord& Record = BatchedSprites[BatchIndex];
-
-		FTexture* TextureResource = (Record.Texture != NULL) ? Record.Texture->Resource : NULL;
-		if ((TextureResource != NULL) && (Record.RenderVerts.Num() > 0))
+		FTexture* TextureResource = (Record.Texture != nullptr) ? Record.Texture->Resource : nullptr;
+		if ((TextureResource != nullptr) && (Record.RenderVerts.Num() > 0))
 		{
-			FLinearColor SpriteColor = bUseOverrideColor ? OverrideColor : Record.Color;
+			const FLinearColor SpriteColor = bUseOverrideColor ? OverrideColor : Record.Color;
 
 			const FVector EffectiveOrigin = Record.Destination;
 
@@ -338,16 +345,9 @@ void FPaperRenderSceneProxy::DrawDynamicElements_RichMesh(FPrimitiveDrawInterfac
 			Mesh.DynamicVertexData = Vertices.GetTypedData();
 			Mesh.DynamicVertexStride = sizeof(FPaperSpriteVertex);
 
-			Mesh.VertexFactory = &GPaperSpriteVertexFactory;
-
-// 			const FColoredMaterialRenderProxy VertexColorVisualizationMaterialInstance(
-// 				VertexColorVisualizationMaterial->GetRenderProxy( MeshElement.MaterialRenderProxy->IsSelected(),MeshElement.MaterialRenderProxy->IsHovered() ),
-// 				GetSelectionColor( FLinearColor::White, bIsSelected, IsHovered() )
-// 				);
-
-
+			Mesh.VertexFactory = VertexFactory;
 
-			FMaterialRenderProxy* ParentMaterialProxy = Material->GetRenderProxy((View->Family->EngineShowFlags.Selection) && IsSelected(), IsHovered());
+			FMaterialRenderProxy* ParentMaterialProxy = BatchMaterial->GetRenderProxy((View->Family->EngineShowFlags.Selection) && IsSelected(), IsHovered());
 
 			// Create a texture override material proxy and register it as a dynamic resource so that it won't be deleted until the rendering thread has finished with it
 			FTextureOverrideRenderProxy* TextureOverrideMaterialProxy = new FTextureOverrideRenderProxy(ParentMaterialProxy, Record.Texture, TEXT("SpriteTexture"));
diff --git a/Engine/Plugins/Experimental/Paper2D/Source/Paper2D/Private/PaperRenderSceneProxy.h b/Engine/Plugins/Experimental/Paper2D/Source/Paper2D/Private/PaperRenderSceneProxy.h
--- a/Engine/Plugins/Experimental/Paper2D/Source/Paper2D/Private/PaperRenderSceneProxy.h
+++ b/Engine/Plugins/Experimental/Paper2D/Source/Paper2D/Private/PaperRenderSceneProxy.h
@@ -30,12 +30,16 @@ public:
 	}
 
 protected:
-	void DrawDynamicElements_RichMesh(FPrimitiveDrawInterface* PDI, const FSceneView* View, bool bUseOverrideColor, const FLinearColor& OverrideColor);
+	virtual void DrawDynamicElements_RichMesh(FPrimitiveDrawInterface* PDI, const FSceneView* View, bool bUseOverrideColor, const FLinearColor& OverrideColor);
+
+	void DrawBatch(FPrimitiveDrawInterface* PDI, const FSceneView* View, bool bUseOverrideColor, const FLinearColor& OverrideColor, class UMaterialInterface* BatchMaterial, const TArray<FSpriteDrawCallRecord>& Batch);
+
 
 	bool IsCollisionView(const FSceneView* View, bool& bDrawSimpleCollision, bool& bDrawComplexCollision) const;
 
 	friend class FPaperBatchSceneProxy;
 
+	FVertexFactory* GetPaperSpriteVertexFactory() const;
 
 protected:
 	TArray<FSpriteDrawCallRecord> BatchedSprites;
diff --git a/Engine/Plugins/Experimental/Paper2D/Source/Paper2D/Private/Terrain/PaperTerrainComponent.cpp b/Engine/Plugins/Experimental/Paper2D/Source/Paper2D/Private/Terrain/PaperTerrainComponent.cpp
--- a/Engine/Plugins/Experimental/Paper2D/Source/Paper2D/Private/Terrain/PaperTerrainComponent.cpp
+++ b/Engine/Plugins/Experimental/Paper2D/Source/Paper2D/Private/Terrain/PaperTerrainComponent.cpp
@@ -6,28 +6,46 @@
 #include "PaperRenderSceneProxy.h"
 
 //////////////////////////////////////////////////////////////////////////
-// UPaperTerrainComponent
+// FPaperTerrainSceneProxy
 
 class FPaperTerrainSceneProxy : public FPaperRenderSceneProxy
 {
 public:
-	FPaperTerrainSceneProxy(const UPaperTerrainComponent* InComponent, const TArray<FSpriteDrawCallRecord>& InGeometry, UMaterialInterface* InMaterial);
+	FPaperTerrainSceneProxy(const UPaperTerrainComponent* InComponent, const TArray<FPaperTerrainMaterialPair>& InDrawingData);
 
 	// FPrimitiveSceneProxy interface.
 	//virtual void DrawDynamicElements(FPrimitiveDrawInterface* PDI, const FSceneView* View) override;
 	// End of FPrimitiveSceneProxy interface.
+
+protected:
+	TArray<FPaperTerrainMaterialPair> DrawingData;
+protected:
+	// FPaperRenderSceneProxy interface
+	virtual void DrawDynamicElements_RichMesh(FPrimitiveDrawInterface* PDI, const FSceneView* View, bool bUseOverrideColor, const FLinearColor& OverrideColor) override;
+	// End of FPaperRenderSceneProxy interface
 };
 
-FPaperTerrainSceneProxy::FPaperTerrainSceneProxy(const UPaperTerrainComponent* InComponent, const TArray<FSpriteDrawCallRecord>& InGeometry, UMaterialInterface* InMaterial)
+FPaperTerrainSceneProxy::FPaperTerrainSceneProxy(const UPaperTerrainComponent* InComponent, const TArray<FPaperTerrainMaterialPair>& InDrawingData)
 	: FPaperRenderSceneProxy(InComponent)
 {
-	BatchedSprites = InGeometry;
+	DrawingData = InDrawingData;
+
+	// Combine the material relevance for all materials
+	for (const FPaperTerrainMaterialPair& Batch : DrawingData)
+	{
+		const UMaterialInterface* MaterialInterface = (Batch.Material != nullptr) ? Batch.Material : UMaterial::GetDefaultMaterial(MD_Surface);
+		MaterialRelevance |= MaterialInterface->GetRelevance_Concurrent();
+	}
+}
 
-	//@TODO: Should make the input geom + material be a pair, allowing multiple materials for solid/opaque
-	Material = InMaterial;
-	if (Material)
+void FPaperTerrainSceneProxy::DrawDynamicElements_RichMesh(FPrimitiveDrawInterface* PDI, const FSceneView* View, bool bUseOverrideColor, const FLinearColor& OverrideColor)
+{
+	for (const FPaperTerrainMaterialPair& Batch : DrawingData)
 	{
-		MaterialRelevance = Material->GetRelevance();
+		if (Batch.Material != nullptr)
+		{
+			DrawBatch(PDI, View, bUseOverrideColor, OverrideColor, Batch.Material, Batch.Records);
+		}
 	}
 }
 
@@ -101,7 +119,7 @@ FPrimitiveSceneProxy* UPaperTerrainComponent::CreateSceneProxy()
 		}
 	}
 
-	FPaperTerrainSceneProxy* NewProxy = new FPaperTerrainSceneProxy(this, GeneratedSpriteGeometry, Material);
+	FPaperTerrainSceneProxy* NewProxy = new FPaperTerrainSceneProxy(this, GeneratedSpriteGeometry);
 	return NewProxy;
 }
 
@@ -112,12 +130,15 @@ FBoxSphereBounds UPaperTerrainComponent::CalcBounds(const FTransform& LocalToWor
 	{
 		FBox BoundingBox(ForceInit);
 
-		for (const FSpriteDrawCallRecord& DrawCall : GeneratedSpriteGeometry)
+		for (const FPaperTerrainMaterialPair& DrawCall : GeneratedSpriteGeometry)
 		{
-			for (const FVector4& VertXYUV : DrawCall.RenderVerts)
+			for (const FSpriteDrawCallRecord& Record : DrawCall.Records)
 			{
-				const FVector Vert((PaperAxisX * VertXYUV.X) + (PaperAxisY * VertXYUV.Y));
-				BoundingBox += Vert;
+				for (const FVector4& VertXYUV : Record.RenderVerts)
+				{
+					const FVector Vert((PaperAxisX * VertXYUV.X) + (PaperAxisY * VertXYUV.Y));
+					BoundingBox += Vert;
+				}
 			}
 		}
 
@@ -277,7 +298,10 @@ void UPaperTerrainComponent::SpawnSegment(class UPaperSprite* NewSprite, float D
 
 		const FTransform LocalTransform(GetTransformAtDistance(Position));
 
-		FSpriteDrawCallRecord& Record = *new (GeneratedSpriteGeometry) FSpriteDrawCallRecord();
+		FPaperTerrainMaterialPair& MaterialBatch = *new (GeneratedSpriteGeometry) FPaperTerrainMaterialPair(); //@TODO: Look up the existing one instead
+		MaterialBatch.Material = NewSprite->GetDefaultMaterial();
+
+		FSpriteDrawCallRecord& Record = *new (MaterialBatch.Records) FSpriteDrawCallRecord();
 		Record.BuildFromSprite(NewSprite);
 		Record.Color = TerrainColor;
 
