Paper2D: WIP support for multiple 2D terrain coverage 'rules' based on slope angle

diff --git a/Engine/Plugins/Experimental/Paper2D/Source/Paper2D/Classes/PaperTerrainComponent.h b/Engine/Plugins/Experimental/Paper2D/Source/Paper2D/Classes/PaperTerrainComponent.h
--- a/Engine/Plugins/Experimental/Paper2D/Source/Paper2D/Classes/PaperTerrainComponent.h
+++ b/Engine/Plugins/Experimental/Paper2D/Source/Paper2D/Classes/PaperTerrainComponent.h
@@ -71,6 +71,7 @@ public:
 
 protected:
 	void SpawnSegment(class UPaperSprite* NewSprite, float Direction, float& RemainingSegStart, float& RemainingSegEnd);
+	void SpawnSegment(class UPaperSprite* NewSprite, float Position, float HorizontalScale);
 
 	void OnSplineEdited();
 
diff --git a/Engine/Plugins/Experimental/Paper2D/Source/Paper2D/Classes/PaperTerrainMaterial.h b/Engine/Plugins/Experimental/Paper2D/Source/Paper2D/Classes/PaperTerrainMaterial.h
--- a/Engine/Plugins/Experimental/Paper2D/Source/Paper2D/Classes/PaperTerrainMaterial.h
+++ b/Engine/Plugins/Experimental/Paper2D/Source/Paper2D/Classes/PaperTerrainMaterial.h
@@ -4,6 +4,46 @@
 
 #include "PaperTerrainMaterial.generated.h"
 
+// Rule for a single section of a terrain material
+USTRUCT(BlueprintType)
+struct FPaperTerrainMaterialRule
+{
+	GENERATED_USTRUCT_BODY()
+
+	// The sprite to use at the 'left' (closest to spline start) edge of the terrain segment
+	UPROPERTY(Category=Sprite, EditAnywhere, meta=(DisplayThumbnail = "true"))
+	UPaperSprite* StartCap;
+
+	// A set of sprites to randomly choose to fill up the interior space between the caps in a terrain segment
+	UPROPERTY(Category = Sprite, EditAnywhere, meta = (DisplayThumbnail = "true"))
+	TArray<UPaperSprite*> Body;
+
+	// The sprite to use at the 'right' (closest to spline end) edge of the terrain segment
+	UPROPERTY(Category = Sprite, EditAnywhere, meta = (DisplayThumbnail = "true"))
+	UPaperSprite* EndCap;
+
+	// Minimum slope angle (in degrees) to apply this rule
+	UPROPERTY(Category = Sprite, EditAnywhere)
+	float MinimumAngle;
+
+	// Maximum slope angle (in degrees) to apply this rule
+	UPROPERTY(Category = Sprite, EditAnywhere)
+	float MaximumAngle;
+
+#if WITH_EDITORONLY_DATA
+	UPROPERTY(Category = Sprite, EditAnywhere)
+	FText Description;
+#endif
+
+	FPaperTerrainMaterialRule()
+		: StartCap(nullptr)
+		, EndCap(nullptr)
+		, MinimumAngle(0.0f)
+		, MaximumAngle(360.0f)
+	{
+	}
+};
+
 /**
  * Paper Terrain Material
  *
@@ -16,11 +56,15 @@ class PAPER2D_API UPaperTerrainMaterial : public UDataAsset
 	GENERATED_UCLASS_BODY()
 
 	UPROPERTY(Category=Sprite, EditAnywhere, meta=(DisplayThumbnail = "true"))
-	UPaperSprite* LeftCap;
+	TArray<FPaperTerrainMaterialRule> Rules;
 
-	UPROPERTY(Category=Sprite, EditAnywhere, meta=(DisplayThumbnail = "true"))
-	TArray<UPaperSprite*> BodySegments;
+	// The sprite to use for an interior region fill
+	UPROPERTY(Category = Sprite, EditAnywhere, meta = (DisplayThumbnail = "true"))
+	UPaperSprite* InteriorFill;
 
-	UPROPERTY(Category=Sprite, EditAnywhere, meta=(DisplayThumbnail = "true"))
-	UPaperSprite* RightCap;
+#if WITH_EDITOR
+	// UObject interface
+	virtual void PostEditChangeProperty(FPropertyChangedEvent& PropertyChangedEvent) override;
+	// End of UObject interface
+#endif
 };
diff --git a/Engine/Plugins/Experimental/Paper2D/Source/Paper2D/Private/Terrain/PaperTerrainComponent.cpp b/Engine/Plugins/Experimental/Paper2D/Source/Paper2D/Private/Terrain/PaperTerrainComponent.cpp
--- a/Engine/Plugins/Experimental/Paper2D/Source/Paper2D/Private/Terrain/PaperTerrainComponent.cpp
+++ b/Engine/Plugins/Experimental/Paper2D/Source/Paper2D/Private/Terrain/PaperTerrainComponent.cpp
@@ -5,6 +5,21 @@
 
 #include "PaperRenderSceneProxy.h"
 
+#define PAPER_USE_MATERIAL_SLOPES 1
+
+//////////////////////////////////////////////////////////////////////////
+
+FBox2D GetSpriteRenderDataBounds2D(const TArray<FVector4>& Data)
+{
+	FBox2D Bounds(ForceInit);
+	for (const FVector4& XYUV : Data)
+	{
+		Bounds += FVector2D(XYUV.X, XYUV.Y);
+	}
+
+	return Bounds;
+}
+
 //////////////////////////////////////////////////////////////////////////
 // FPaperTerrainSceneProxy
 
@@ -109,16 +124,6 @@ void UPaperTerrainComponent::OnUnregister()
 
 FPrimitiveSceneProxy* UPaperTerrainComponent::CreateSceneProxy()
 {
-	//@TODO: Super-hacky
-	UMaterialInterface* Material = nullptr;
-	if (TerrainMaterial != nullptr)
-	{
-		if (TerrainMaterial->LeftCap)
-		{
-			Material = TerrainMaterial->LeftCap->GetDefaultMaterial();
-		}
-	}
-
 	FPaperTerrainSceneProxy* NewProxy = new FPaperTerrainSceneProxy(this, GeneratedSpriteGeometry);
 	return NewProxy;
 }
@@ -227,22 +232,108 @@ void UPaperTerrainComponent::OnSplineEdited()
 
 		float SplineLength = AssociatedSpline->GetSplineLength();
 
-		float RemainingSegStart = 0.0f;
-		float RemainingSegEnd = SplineLength;
 
-		SpawnSegment(TerrainMaterial->LeftCap, /*dir=*/ 1.0f, /*inout*/ RemainingSegStart, /*inout*/ RemainingSegEnd);
-		SpawnSegment(TerrainMaterial->RightCap, /*dir=*/ -1.0f, /*inout*/ RemainingSegStart, /*inout*/ RemainingSegEnd);
+		struct FSpriteStamp
+		{
+			UPaperSprite* Sprite;
+			float NominalWidth;
+			float Time;
+			float Scale;
+			bool bCanStretch;
+
+			FSpriteStamp(UPaperSprite* InSprite, float InTime, bool bIsEndCap)
+				: Sprite(InSprite)
+				, Time(InTime)
+				, bCanStretch(!bIsEndCap)
+			{
+				const FBox2D Bounds2D = GetSpriteRenderDataBounds2D(InSprite->BakedRenderData);
+				NominalWidth = FMath::Max<float>(Bounds2D.GetSize().X, 1.0f);
+			}
+		};
+
+		struct FTerrainSegment
+		{
+			float StartTime;
+			float EndTime;
+			const FPaperTerrainMaterialRule* Rule;
+			TArray<FSpriteStamp> Stamps;
+
+			FTerrainSegment()
+				: Rule(nullptr)
+			{
+			}
+
+			void RepositionStampsToFillSpace()
+			{
+			}
+		};
+
+
+		// Split the spline into segments based on the slope rules in the material
+		TArray<FTerrainSegment> Segments;
 
-		if (TerrainMaterial->BodySegments.Num() > 0)
+		FTerrainSegment* ActiveSegment = new (Segments) FTerrainSegment();
+		ActiveSegment->StartTime = 0.0f;
+		ActiveSegment->EndTime = SplineLength;
+
+		float CurrentTime = 0.0f;
+		while (CurrentTime < SplineLength)
 		{
-			while (RemainingSegEnd - RemainingSegStart > 0.0f)
+			const FTransform Frame(GetTransformAtDistance(CurrentTime));
+			const FVector UnitTangent = Frame.GetUnitAxis(EAxis::X);
+			const float RawSlopeAngleRadians = FMath::Atan2(FVector::DotProduct(UnitTangent, PaperAxisY), FVector::DotProduct(UnitTangent, PaperAxisX));
+			const float RawSlopeAngle = FMath::RadiansToDegrees(RawSlopeAngleRadians);
+			const float SlopeAngle = FMath::Fmod(FMath::UnwindDegrees(RawSlopeAngle) + 360.0f, 360.0f);
+
+			const FPaperTerrainMaterialRule* DesiredRule = (TerrainMaterial->Rules.Num() > 0) ? &(TerrainMaterial->Rules[0]) : nullptr;
+			for (const FPaperTerrainMaterialRule& TestRule : TerrainMaterial->Rules)
 			{
-				const int32 BodySegmentIndex = RandomStream.GetUnsignedInt() % TerrainMaterial->BodySegments.Num();
-				UPaperSprite* BodySegment = TerrainMaterial->BodySegments[BodySegmentIndex];
-				SpawnSegment(BodySegment, /*dir=*/ 1.0f, /*inout*/ RemainingSegStart, /*inout*/ RemainingSegEnd);
+				if ((SlopeAngle >= TestRule.MinimumAngle) && (SlopeAngle < TestRule.MaximumAngle))
+				{
+					DesiredRule = &TestRule;
+				}
 			}
+
+			if (ActiveSegment->Rule != DesiredRule)
+			{
+				if (ActiveSegment->Rule == nullptr)
+				{
+					ActiveSegment->Rule = DesiredRule;
+				}
+				else
+				{
+					ActiveSegment->EndTime = CurrentTime;
+
+					ActiveSegment = new (Segments) FTerrainSegment();
+					ActiveSegment->StartTime = CurrentTime;
+					ActiveSegment->EndTime = SplineLength;
+					ActiveSegment->Rule = DesiredRule;
+				}
+			}
+
+			CurrentTime += 10.0f;
 		}
 
+		// Convert those segments to actual geometry
+		for (FTerrainSegment& Segment : Segments)
+		{
+			check(Segment.Rule);
+			float RemainingSegStart = Segment.StartTime;
+			float RemainingSegEnd = Segment.EndTime;
+			SpawnSegment(Segment.Rule->StartCap, /*dir=*/ 1.0f, /*inout*/ RemainingSegStart, /*inout*/ RemainingSegEnd);
+			SpawnSegment(Segment.Rule->EndCap, /*dir=*/ -1.0f, /*inout*/ RemainingSegStart, /*inout*/ RemainingSegEnd);
+
+			if (Segment.Rule->Body.Num() > 0)
+			{
+				while (RemainingSegEnd - RemainingSegStart > 0.0f)
+				{
+					const int32 BodySegmentIndex = RandomStream.GetUnsignedInt() % Segment.Rule->Body.Num();
+					UPaperSprite* BodySegment = Segment.Rule->Body[BodySegmentIndex];
+					SpawnSegment(BodySegment, /*dir=*/ 1.0f, /*inout*/ RemainingSegStart, /*inout*/ RemainingSegEnd);
+				}
+			}
+
+		}
 
 		// Draw debug frames at the start and end of the spline
 #if PAPER_TERRAIN_DRAW_DEBUG
@@ -263,15 +354,30 @@ void UPaperTerrainComponent::OnSplineEdited()
 	RecreateRenderState_Concurrent();
 }
 
-FBox2D GetSpriteRenderDataBounds2D(const TArray<FVector4>& Data)
+void UPaperTerrainComponent::SpawnSegment(class UPaperSprite* NewSprite, float Position, float HorizontalScale)
 {
-	FBox2D Bounds(ForceInit);
-	for (const FVector4& XYUV : Data)
+	if (NewSprite)
 	{
-		Bounds += FVector2D(XYUV.X, XYUV.Y);
-	}
+		FPaperTerrainMaterialPair& MaterialBatch = *new (GeneratedSpriteGeometry) FPaperTerrainMaterialPair(); //@TODO: Look up the existing one instead
+		MaterialBatch.Material = NewSprite->GetDefaultMaterial();
 
-	return Bounds;
+		FSpriteDrawCallRecord& Record = *new (MaterialBatch.Records) FSpriteDrawCallRecord();
+		Record.BuildFromSprite(NewSprite);
+		Record.Color = TerrainColor;
+
+		for (FVector4& XYUV : Record.RenderVerts)
+		{
+			const FTransform LocalTransformAtX(GetTransformAtDistance(Position + XYUV.X));
+			const FVector SourceVector = (XYUV.Y * PaperAxisY);
+			const FVector NewVector = LocalTransformAtX.TransformPosition(SourceVector);
+
+			const float NewX = FVector::DotProduct(NewVector, PaperAxisX);
+			const float NewY = FVector::DotProduct(NewVector, PaperAxisY);
+
+			XYUV.X = NewX;
+			XYUV.Y = NewY;
+		}
+	}
 }
 
 void UPaperTerrainComponent::SpawnSegment(class UPaperSprite* NewSprite, float Direction, float& RemainingSegStart, float& RemainingSegEnd)
@@ -295,28 +401,7 @@ void UPaperTerrainComponent::SpawnSegment(class UPaperSprite* NewSprite, float D
 			RemainingSegStart += SpriteWidth;
 		}
 
-
-		const FTransform LocalTransform(GetTransformAtDistance(Position));
-
-		FPaperTerrainMaterialPair& MaterialBatch = *new (GeneratedSpriteGeometry) FPaperTerrainMaterialPair(); //@TODO: Look up the existing one instead
-		MaterialBatch.Material = NewSprite->GetDefaultMaterial();
-
-		FSpriteDrawCallRecord& Record = *new (MaterialBatch.Records) FSpriteDrawCallRecord();
-		Record.BuildFromSprite(NewSprite);
-		Record.Color = TerrainColor;
-
-		for (FVector4& XYUV : Record.RenderVerts)
-		{
-			const FTransform LocalTransformAtX(GetTransformAtDistance(Position + XYUV.X));
-			const FVector SourceVector = (XYUV.Y * PaperAxisY);
-			const FVector NewVector = LocalTransformAtX.TransformPosition(SourceVector);
-
-			const float NewX = FVector::DotProduct(NewVector, PaperAxisX);
-			const float NewY = FVector::DotProduct(NewVector, PaperAxisY);
-
-			XYUV.X = NewX;
-			XYUV.Y = NewY;
-		}
+		SpawnSegment(NewSprite, Position, 1.0f);
 	}
 }
 
diff --git a/Engine/Plugins/Experimental/Paper2D/Source/Paper2D/Private/Terrain/PaperTerrainMaterial.cpp b/Engine/Plugins/Experimental/Paper2D/Source/Paper2D/Private/Terrain/PaperTerrainMaterial.cpp
--- a/Engine/Plugins/Experimental/Paper2D/Source/Paper2D/Private/Terrain/PaperTerrainMaterial.cpp
+++ b/Engine/Plugins/Experimental/Paper2D/Source/Paper2D/Private/Terrain/PaperTerrainMaterial.cpp
@@ -2,10 +2,31 @@
 
 #include "Paper2DPrivatePCH.h"
 
+#define LOCTEXT_NAMESPACE "Paper2D"
+
 //////////////////////////////////////////////////////////////////////////
 // UPaperTerrainMaterial
 
 UPaperTerrainMaterial::UPaperTerrainMaterial(const FPostConstructInitializeProperties& PCIP)
 	: Super(PCIP)
 {
+	// Add a default rule
+	FPaperTerrainMaterialRule& DefaultRule = *new (Rules) FPaperTerrainMaterialRule();
+	DefaultRule.MinimumAngle = 0.0f;
+	DefaultRule.MaximumAngle = 360.0f;
+#if WITH_EDITORONLY_DATA
+	DefaultRule.Description = LOCTEXT("TerrainDefaultRuleDesc", "Top");
+#endif
+}
+
+#if WITH_EDITOR
+void UPaperTerrainMaterial::PostEditChangeProperty(FPropertyChangedEvent& PropertyChangedEvent)
+{
+	//@TODO: PAPER2D: Ensure that the angle ranges of rules don't overlap
+	Super::PostEditChangeProperty(PropertyChangedEvent);
 }
+#endif
+
+//////////////////////////////////////////////////////////////////////////
+
+#undef LOCTEXT_NAMESPACE
\ No newline at end of file
