Paper2D: Add repositioning and scaling of body sprites in 2D terrain splines to accomodate gaps

diff --git a/Engine/Plugins/Experimental/Paper2D/Source/Paper2D/Classes/PaperTerrainComponent.h b/Engine/Plugins/Experimental/Paper2D/Source/Paper2D/Classes/PaperTerrainComponent.h
--- a/Engine/Plugins/Experimental/Paper2D/Source/Paper2D/Classes/PaperTerrainComponent.h
+++ b/Engine/Plugins/Experimental/Paper2D/Source/Paper2D/Classes/PaperTerrainComponent.h
@@ -70,8 +70,7 @@ public:
 	void SetTerrainColor(FLinearColor NewColor);
 
 protected:
-	void SpawnSegment(class UPaperSprite* NewSprite, float Direction, float& RemainingSegStart, float& RemainingSegEnd);
-	void SpawnSegment(class UPaperSprite* NewSprite, float Position, float HorizontalScale);
+	void SpawnSegment(const class UPaperSprite* NewSprite, float Position, float HorizontalScale);
 
 	void OnSplineEdited();
 
diff --git a/Engine/Plugins/Experimental/Paper2D/Source/Paper2D/Classes/SpriteDrawCall.h b/Engine/Plugins/Experimental/Paper2D/Source/Paper2D/Classes/SpriteDrawCall.h
--- a/Engine/Plugins/Experimental/Paper2D/Source/Paper2D/Classes/SpriteDrawCall.h
+++ b/Engine/Plugins/Experimental/Paper2D/Source/Paper2D/Classes/SpriteDrawCall.h
@@ -24,7 +24,7 @@ public:
 	// Render triangle list (stored as loose vertices)
 	TArray< FVector4, TInlineAllocator<6> > RenderVerts;
 
-	void BuildFromSprite(class UPaperSprite* Sprite);
+	void BuildFromSprite(const class UPaperSprite* Sprite);
 
 	FSpriteDrawCallRecord()
 		: Texture(NULL)
diff --git a/Engine/Plugins/Experimental/Paper2D/Source/Paper2D/Private/PaperSprite.cpp b/Engine/Plugins/Experimental/Paper2D/Source/Paper2D/Private/PaperSprite.cpp
--- a/Engine/Plugins/Experimental/Paper2D/Source/Paper2D/Private/PaperSprite.cpp
+++ b/Engine/Plugins/Experimental/Paper2D/Source/Paper2D/Private/PaperSprite.cpp
@@ -234,9 +234,9 @@ void UPaperSprite::ExtractSourceRegionFromTexturePoint(const FVector2D& SourcePo
 //////////////////////////////////////////////////////////////////////////
 // FSpriteDrawCallRecord
 
-void FSpriteDrawCallRecord::BuildFromSprite(class UPaperSprite* Sprite)
+void FSpriteDrawCallRecord::BuildFromSprite(const UPaperSprite* Sprite)
 {
-	if (Sprite != NULL)
+	if (Sprite != nullptr)
 	{
 		Destination = FVector::ZeroVector;
 		Texture = Sprite->GetBakedTexture();
diff --git a/Engine/Plugins/Experimental/Paper2D/Source/Paper2D/Private/Terrain/PaperTerrainComponent.cpp b/Engine/Plugins/Experimental/Paper2D/Source/Paper2D/Private/Terrain/PaperTerrainComponent.cpp
--- a/Engine/Plugins/Experimental/Paper2D/Source/Paper2D/Private/Terrain/PaperTerrainComponent.cpp
+++ b/Engine/Plugins/Experimental/Paper2D/Source/Paper2D/Private/Terrain/PaperTerrainComponent.cpp
@@ -235,15 +235,16 @@ void UPaperTerrainComponent::OnSplineEdited()
 
 		struct FSpriteStamp
 		{
-			UPaperSprite* Sprite;
+			const UPaperSprite* Sprite;
 			float NominalWidth;
 			float Time;
 			float Scale;
 			bool bCanStretch;
 
-			FSpriteStamp(UPaperSprite* InSprite, float InTime, bool bIsEndCap)
+			FSpriteStamp(const UPaperSprite* InSprite, float InTime, bool bIsEndCap)
 				: Sprite(InSprite)
 				, Time(InTime)
+				, Scale(1.0f)
 				, bCanStretch(!bIsEndCap)
 			{
 				const FBox2D Bounds2D = GetSpriteRenderDataBounds2D(InSprite->BakedRenderData);
@@ -266,8 +267,60 @@ void UPaperTerrainComponent::OnSplineEdited()
 			void RepositionStampsToFillSpace()
 			{
 			}
+
 		};
 
+		struct FTerrainRuleHelper
+		{
+		public:
+			FTerrainRuleHelper(const FPaperTerrainMaterialRule* Rule)
+				: StartWidth(0.0f)
+				, EndWidth(0.0f)
+			{
+				for (const UPaperSprite* Sprite : Rule->Body)
+				{
+					if (Sprite != nullptr)
+					{
+						const float Width = GetSpriteRenderDataBounds2D(Sprite->BakedRenderData).GetSize().X;
+						if (Width > 0.0f)
+						{
+							ValidBodies.Add(Sprite);
+							ValidBodyWidths.Add(Width);
+						}
+					}
+				}
+
+				if (Rule->StartCap != nullptr)
+				{
+					const float Width = GetSpriteRenderDataBounds2D(Rule->StartCap->BakedRenderData).GetSize().X;
+					if (Width > 0.0f)
+					{
+						StartWidth = Width;
+					}
+				}
+
+				if (Rule->EndCap != nullptr)
+				{
+					const float Width = GetSpriteRenderDataBounds2D(Rule->EndCap->BakedRenderData).GetSize().X;
+					if (Width > 0.0f)
+					{
+						EndWidth = Width;
+					}
+				}
+			}
+
+			float StartWidth;
+			float EndWidth;
+
+			TArray<const UPaperSprite*> ValidBodies;
+			TArray<float> ValidBodyWidths;
+
+			int32 GenerateBodyIndex(FRandomStream& RandomStream) const
+			{
+				check(ValidBodies.Num() > 0);
+				return RandomStream.GetUnsignedInt() % ValidBodies.Num();
+			}
+		};
 
 		// Split the spline into segments based on the slope rules in the material
 		TArray<FTerrainSegment> Segments;
@@ -318,21 +371,78 @@ void UPaperTerrainComponent::OnSplineEdited()
 		for (FTerrainSegment& Segment : Segments)
 		{
 			check(Segment.Rule);
-			float RemainingSegStart = Segment.StartTime;
-			float RemainingSegEnd = Segment.EndTime;
-			SpawnSegment(Segment.Rule->StartCap, /*dir=*/ 1.0f, /*inout*/ RemainingSegStart, /*inout*/ RemainingSegEnd);
-			SpawnSegment(Segment.Rule->EndCap, /*dir=*/ -1.0f, /*inout*/ RemainingSegStart, /*inout*/ RemainingSegEnd);
+			FTerrainRuleHelper RuleHelper(Segment.Rule);
+
+			float RemainingSegStart = Segment.StartTime + RuleHelper.StartWidth;
+			float RemainingSegEnd = Segment.EndTime - RuleHelper.EndWidth;
+			const float BodyDistance = RemainingSegEnd - RemainingSegStart;
+			float DistanceBudget = BodyDistance;
+
+			bool bUseBodySegments = (DistanceBudget > 0.0f) && (RuleHelper.ValidBodies.Num() > 0);
+
+			// Add the start cap
+			if (RuleHelper.StartWidth > 0.0f)
+			{
+				new (Segment.Stamps) FSpriteStamp(Segment.Rule->StartCap, Segment.StartTime + RuleHelper.StartWidth * 0.5f, /*bIsEndCap=*/ bUseBodySegments);
+			}
+
+			// Add the end cap
+			if (RuleHelper.EndWidth > 0.0f)
+			{
+				new (Segment.Stamps) FSpriteStamp(Segment.Rule->EndCap, Segment.EndTime - RuleHelper.EndWidth * 0.5f, /*bIsEndCap=*/ bUseBodySegments);
+			}
+			
 
-			if (Segment.Rule->Body.Num() > 0)
+			// Add body segments
+			if (bUseBodySegments)
 			{
-				while (RemainingSegEnd - RemainingSegStart > 0.0f)
+				int32 NumSegments = 0;
+				float Position = RemainingSegStart;
+				
+				while (DistanceBudget > 0.0f)
 				{
-					const int32 BodySegmentIndex = RandomStream.GetUnsignedInt() % Segment.Rule->Body.Num();
-					UPaperSprite* BodySegment = Segment.Rule->Body[BodySegmentIndex];
-					SpawnSegment(BodySegment, /*dir=*/ 1.0f, /*inout*/ RemainingSegStart, /*inout*/ RemainingSegEnd);
+					const int32 BodyIndex = RuleHelper.GenerateBodyIndex(RandomStream);
+					const UPaperSprite* Sprite = RuleHelper.ValidBodies[BodyIndex];
+					const float Width = RuleHelper.ValidBodyWidths[BodyIndex];
+
+ 					if ((NumSegments > 0) && ((Width * 0.5f) > DistanceBudget))
+ 					{
+ 						break;
+ 					}
+					new (Segment.Stamps) FSpriteStamp(Sprite, Position + (Width * 0.5f), /*bIsEndCap=*/ false);
+
+					DistanceBudget -= Width;
+					Position += Width;
+					++NumSegments;
 				}
+
+				const float UsedSpace = (BodyDistance - DistanceBudget);
+				const float OverallScaleFactor = BodyDistance / UsedSpace;
+				
+ 				// Stretch body segments
+				float PositionCorrectionSum = 0.0f;
+ 				for (int32 Index = 0; Index < NumSegments; ++Index)
+ 				{
+					FSpriteStamp& Stamp = Segment.Stamps[Index + (Segment.Stamps.Num() - NumSegments)];
+					
+					const float WidthChange = (OverallScaleFactor - 1.0f) * Stamp.NominalWidth;
+					const float FirstGapIsSmallerFactor = (Index == 0) ? 0.5f : 1.0f;
+					PositionCorrectionSum += WidthChange * FirstGapIsSmallerFactor;
+
+					Stamp.Scale = OverallScaleFactor;
+					Stamp.Time += PositionCorrectionSum;
+ 				}
+			}
+			else
+			{
+				// Stretch endcaps
 			}
 
+			// Convert stamps into geometry
+			for (const FSpriteStamp& Stamp : Segment.Stamps)
+			{
+				SpawnSegment(Stamp.Sprite, Stamp.Time, Stamp.Scale);
+			}
 		}
 
 		// Draw debug frames at the start and end of the spline
@@ -354,7 +464,7 @@ void UPaperTerrainComponent::OnSplineEdited()
 	RecreateRenderState_Concurrent();
 }
 
-void UPaperTerrainComponent::SpawnSegment(class UPaperSprite* NewSprite, float Position, float HorizontalScale)
+void UPaperTerrainComponent::SpawnSegment(const UPaperSprite* NewSprite, float Position, float HorizontalScale)
 {
 	if (NewSprite)
 	{
@@ -367,7 +477,7 @@ void UPaperTerrainComponent::SpawnSegment(class UPaperSprite* NewSprite, float P
 
 		for (FVector4& XYUV : Record.RenderVerts)
 		{
-			const FTransform LocalTransformAtX(GetTransformAtDistance(Position + XYUV.X));
+			const FTransform LocalTransformAtX(GetTransformAtDistance(Position + XYUV.X * HorizontalScale));
 			const FVector SourceVector = (XYUV.Y * PaperAxisY);
 			const FVector NewVector = LocalTransformAtX.TransformPosition(SourceVector);
 
@@ -380,31 +490,6 @@ void UPaperTerrainComponent::SpawnSegment(class UPaperSprite* NewSprite, float P
 	}
 }
 
-void UPaperTerrainComponent::SpawnSegment(class UPaperSprite* NewSprite, float Direction, float& RemainingSegStart, float& RemainingSegEnd)
-{
-	if (NewSprite)
-	{
-		const FBox2D Bounds2D = GetSpriteRenderDataBounds2D(NewSprite->BakedRenderData);
-
-		const float SpriteWidthUnsafe = Bounds2D.GetSize().X; //NewSprite->GetRenderBounds().BoxExtent.ProjectOnTo(PaperAxisX).Size() * 2.0f;
-		const float SpriteWidth = (SpriteWidthUnsafe < KINDA_SMALL_NUMBER) ? 1.0f : SpriteWidthUnsafe; // ensure we terminate in the outer loop even with malformed render data
-
-		float Position = 0.0f;
-		if (Direction < 0.0f)
-		{
-			Position = RemainingSegEnd - SpriteWidth * 0.5f;
-			RemainingSegEnd -= SpriteWidth;
-		}
-		else
-		{
-			Position = RemainingSegStart + SpriteWidth * 0.5f;
-			RemainingSegStart += SpriteWidth;
-		}
-
-		SpawnSegment(NewSprite, Position, 1.0f);
-	}
-}
-
 void UPaperTerrainComponent::SetTerrainColor(FLinearColor NewColor)
 {
 	// Can't set color on a static component
