Paper2D: Marked render component in terrain actor with ExposeFunctionCategories

diff --git a/Engine/Plugins/Experimental/Paper2D/Source/Paper2D/Classes/PaperTerrainActor.h b/Engine/Plugins/Experimental/Paper2D/Source/Paper2D/Classes/PaperTerrainActor.h
--- a/Engine/Plugins/Experimental/Paper2D/Source/Paper2D/Classes/PaperTerrainActor.h
+++ b/Engine/Plugins/Experimental/Paper2D/Source/Paper2D/Classes/PaperTerrainActor.h
@@ -19,7 +19,7 @@ class PAPER2D_API APaperTerrainActor : public AActor
 	UPROPERTY()
 	TSubobjectPtr<class UPaperTerrainSplineComponent> SplineComponent;
 
-	UPROPERTY(Category = Sprite, VisibleAnywhere)
+	UPROPERTY(Category=Sprite, VisibleAnywhere, BlueprintReadOnly, meta = (ExposeFunctionCategories = "Sprite,Rendering,Physics,Components|Sprite"))
 	TSubobjectPtr<class UPaperTerrainComponent> RenderComponent;
 
 	// AActor interface
