Paper2D: Fix UPaperTerrainComponent::SetTerrainColor not taking effect if the component was registered

diff --git a/Engine/Plugins/Experimental/Paper2D/Source/Paper2D/Private/Terrain/PaperTerrainComponent.cpp b/Engine/Plugins/Experimental/Paper2D/Source/Paper2D/Private/Terrain/PaperTerrainComponent.cpp
--- a/Engine/Plugins/Experimental/Paper2D/Source/Paper2D/Private/Terrain/PaperTerrainComponent.cpp
+++ b/Engine/Plugins/Experimental/Paper2D/Source/Paper2D/Private/Terrain/PaperTerrainComponent.cpp
@@ -497,7 +497,16 @@ void UPaperTerrainComponent::SetTerrainColor(FLinearColor NewColor)
 	{
 		TerrainColor = NewColor;
 
-		//@TODO: Should we send immediately?
-		MarkRenderDynamicDataDirty();
+		// Update the color in the game-thread copy of the render geometry
+		for (FPaperTerrainMaterialPair& Batch : GeneratedSpriteGeometry)
+		{
+			for (FSpriteDrawCallRecord& DrawCall : Batch.Records)
+			{
+				DrawCall.Color = TerrainColor;
+			}
+		}
+
+		// Update the render thread copy
+		RecreateRenderState_Concurrent();
 	}
 }
