Paper2D: Add limited interior fill support to terrain splines (editor-only, and interior sprite must use the entire texture, render geometry is ignored)

diff --git a/Engine/Plugins/Experimental/Paper2D/Source/Paper2D/Classes/PaperTerrainComponent.h b/Engine/Plugins/Experimental/Paper2D/Source/Paper2D/Classes/PaperTerrainComponent.h
--- a/Engine/Plugins/Experimental/Paper2D/Source/Paper2D/Classes/PaperTerrainComponent.h
+++ b/Engine/Plugins/Experimental/Paper2D/Source/Paper2D/Classes/PaperTerrainComponent.h
@@ -71,6 +71,7 @@ public:
 
 protected:
 	void SpawnSegment(const class UPaperSprite* NewSprite, float Position, float HorizontalScale);
+	void SpawnFromPoly(const class UPaperSprite* NewSprite, FSpriteDrawCallRecord& FillDrawCall, const FVector2D& TextureSize, FPoly& Poly);
 
 	void OnSplineEdited();
 
diff --git a/Engine/Plugins/Experimental/Paper2D/Source/Paper2D/Private/Terrain/PaperTerrainComponent.cpp b/Engine/Plugins/Experimental/Paper2D/Source/Paper2D/Private/Terrain/PaperTerrainComponent.cpp
--- a/Engine/Plugins/Experimental/Paper2D/Source/Paper2D/Private/Terrain/PaperTerrainComponent.cpp
+++ b/Engine/Plugins/Experimental/Paper2D/Source/Paper2D/Private/Terrain/PaperTerrainComponent.cpp
@@ -216,6 +216,9 @@ FTransform UPaperTerrainComponent::GetTransformAtDistance(float InDistance) cons
 
 void UPaperTerrainComponent::OnSplineEdited()
 {
+	const float SlopeAnalysisTimeRate = 10.0f;
+	const float FillRasterizationTimeRate = 100.0f;
+
 	GeneratedSpriteGeometry.Empty();
 
 	if ((AssociatedSpline != nullptr) && (TerrainMaterial != nullptr))
@@ -329,42 +332,44 @@ void UPaperTerrainComponent::OnSplineEdited()
 		ActiveSegment->StartTime = 0.0f;
 		ActiveSegment->EndTime = SplineLength;
 
-		float CurrentTime = 0.0f;
-		while (CurrentTime < SplineLength)
 		{
-			const FTransform Frame(GetTransformAtDistance(CurrentTime));
-			const FVector UnitTangent = Frame.GetUnitAxis(EAxis::X);
-			const float RawSlopeAngleRadians = FMath::Atan2(FVector::DotProduct(UnitTangent, PaperAxisY), FVector::DotProduct(UnitTangent, PaperAxisX));
-			const float RawSlopeAngle = FMath::RadiansToDegrees(RawSlopeAngleRadians);
-			const float SlopeAngle = FMath::Fmod(FMath::UnwindDegrees(RawSlopeAngle) + 360.0f, 360.0f);
-
-			const FPaperTerrainMaterialRule* DesiredRule = (TerrainMaterial->Rules.Num() > 0) ? &(TerrainMaterial->Rules[0]) : nullptr;
-			for (const FPaperTerrainMaterialRule& TestRule : TerrainMaterial->Rules)
+			float CurrentTime = 0.0f;
+			while (CurrentTime < SplineLength)
 			{
-				if ((SlopeAngle >= TestRule.MinimumAngle) && (SlopeAngle < TestRule.MaximumAngle))
+				const FTransform Frame(GetTransformAtDistance(CurrentTime));
+				const FVector UnitTangent = Frame.GetUnitAxis(EAxis::X);
+				const float RawSlopeAngleRadians = FMath::Atan2(FVector::DotProduct(UnitTangent, PaperAxisY), FVector::DotProduct(UnitTangent, PaperAxisX));
+				const float RawSlopeAngle = FMath::RadiansToDegrees(RawSlopeAngleRadians);
+				const float SlopeAngle = FMath::Fmod(FMath::UnwindDegrees(RawSlopeAngle) + 360.0f, 360.0f);
+
+				const FPaperTerrainMaterialRule* DesiredRule = (TerrainMaterial->Rules.Num() > 0) ? &(TerrainMaterial->Rules[0]) : nullptr;
+				for (const FPaperTerrainMaterialRule& TestRule : TerrainMaterial->Rules)
 				{
-					DesiredRule = &TestRule;
+					if ((SlopeAngle >= TestRule.MinimumAngle) && (SlopeAngle < TestRule.MaximumAngle))
+					{
+						DesiredRule = &TestRule;
+					}
 				}
-			}
 
-			if (ActiveSegment->Rule != DesiredRule)
-			{
-				if (ActiveSegment->Rule == nullptr)
-				{
-					ActiveSegment->Rule = DesiredRule;
-				}
-				else
+				if (ActiveSegment->Rule != DesiredRule)
 				{
-					ActiveSegment->EndTime = CurrentTime;
+					if (ActiveSegment->Rule == nullptr)
+					{
+						ActiveSegment->Rule = DesiredRule;
+					}
+					else
+					{
+						ActiveSegment->EndTime = CurrentTime;
 
-					ActiveSegment = new (Segments) FTerrainSegment();
-					ActiveSegment->StartTime = CurrentTime;
-					ActiveSegment->EndTime = SplineLength;
-					ActiveSegment->Rule = DesiredRule;
+						ActiveSegment = new (Segments)FTerrainSegment();
+						ActiveSegment->StartTime = CurrentTime;
+						ActiveSegment->EndTime = SplineLength;
+						ActiveSegment->Rule = DesiredRule;
+					}
 				}
-			}
 
-			CurrentTime += 10.0f;
+				CurrentTime += SlopeAnalysisTimeRate;
+			}
 		}
 
 		// Convert those segments to actual geometry
@@ -445,6 +450,72 @@ void UPaperTerrainComponent::OnSplineEdited()
 			}
 		}
 
+		// Generate the background if the spline is closed
+		if (bClosedSpline && (TerrainMaterial->InteriorFill != nullptr))
+		{
+			// Create a polygon from the spline
+			FBox2D SplineBounds(ForceInit);
+			FPoly SplineAsPolygon;
+			{
+				float CurrentTime = 0.0f;
+				while (CurrentTime < SplineLength)
+				{
+					const float Param = AssociatedSpline->SplineReparamTable.Eval(CurrentTime, 0.0f);
+					const FVector Position3D = AssociatedSpline->SplineInfo.Eval(Param, FVector::ZeroVector);
+					const FVector2D Position2D = FVector2D(FVector::DotProduct(Position3D, PaperAxisX), FVector::DotProduct(Position3D, PaperAxisY));
+					const FVector CleanPosition3D = (PaperAxisX * Position2D.X) + (PaperAxisY * Position2D.Y);
+
+					SplineBounds += Position2D;
+					SplineAsPolygon.Vertices.Add(CleanPosition3D);
+
+					CurrentTime += FillRasterizationTimeRate;
+				}
+			}
+
+			const UPaperSprite* FillSprite = TerrainMaterial->InteriorFill;
+			FPaperTerrainMaterialPair& MaterialBatch = *new (GeneratedSpriteGeometry) FPaperTerrainMaterialPair(); //@TODO: Look up the existing one instead
+			MaterialBatch.Material = FillSprite->GetDefaultMaterial();
+
+			FSpriteDrawCallRecord& FillDrawCall = *new (MaterialBatch.Records) FSpriteDrawCallRecord();
+			FillDrawCall.BuildFromSprite(FillSprite);
+			FillDrawCall.RenderVerts.Empty();
+			FillDrawCall.Color = TerrainColor;
+			FillDrawCall.Destination = PaperAxisZ * 0.1f;
+
+			const FVector2D TextureSize = GetSpriteRenderDataBounds2D(FillSprite->BakedRenderData).GetSize();
+			const FVector2D SplineSize = SplineBounds.GetSize();
+			
+			SpawnFromPoly(FillSprite, FillDrawCall, TextureSize, SplineAsPolygon);
+
+			//@TODO: Add support for the fill sprite being smaller than the entire texture
+#if NOT_WORKING
+			const float StartingDivisionPointX = FMath::CeilToFloat(SplineBounds.Min.X / TextureSize.X);
+			const float StartingDivisionPointY = FMath::CeilToFloat(SplineBounds.Min.Y / TextureSize.Y);
+
+			FPoly VerticalRemainder = SplineAsPolygon;
+			for (float Y = StartingDivisionPointY; VerticalRemainder.Vertices.Num() > 0; Y += TextureSize.Y)
+			{
+				FPoly Top;
+				FPoly Bottom;
+				const FVector SplitBaseOuter = (Y * PaperAxisY);
+				VerticalRemainder.SplitWithPlane(SplitBaseOuter, -PaperAxisY, &Top, &Bottom, 1);
+				VerticalRemainder = Bottom;
+
+				FPoly HorizontalRemainder = Top;
+				for (float X = StartingDivisionPointX; HorizontalRemainder.Vertices.Num() > 0; X += TextureSize.X)
+				{
+					FPoly Left;
+					FPoly Right;
+					const FVector SplitBaseInner = (X * PaperAxisX) + (Y * PaperAxisY);
+					HorizontalRemainder.SplitWithPlane(SplitBaseInner, -PaperAxisX, &Left, &Right, 1);
+					HorizontalRemainder = Right;
+
+					SpawnFromPoly(FillSprite, FillDrawCall, TextureSize, Left);
+				}
+			}
+#endif
+		}
+
 		// Draw debug frames at the start and end of the spline
 #if PAPER_TERRAIN_DRAW_DEBUG
 		{
@@ -477,7 +548,7 @@ void UPaperTerrainComponent::SpawnSegment(const UPaperSprite* NewSprite, float P
 
 		for (FVector4& XYUV : Record.RenderVerts)
 		{
-			const FTransform LocalTransformAtX(GetTransformAtDistance(Position + XYUV.X * HorizontalScale));
+			const FTransform LocalTransformAtX(GetTransformAtDistance(Position + (XYUV.X * HorizontalScale)));
 			const FVector SourceVector = (XYUV.Y * PaperAxisY);
 			const FVector NewVector = LocalTransformAtX.TransformPosition(SourceVector);
 
@@ -490,6 +561,56 @@ void UPaperTerrainComponent::SpawnSegment(const UPaperSprite* NewSprite, float P
 	}
 }
 
+#if WITH_EDITOR
+#include "GeomTools.h"
+#endif
+
+void UPaperTerrainComponent::SpawnFromPoly(const class UPaperSprite* NewSprite, FSpriteDrawCallRecord& Batch, const FVector2D& TextureSize, FPoly& Poly)
+{
+	//@TODO: Need to split geom tools out of UnrealEd to support this outside of the editor
+#if WITH_EDITOR
+	if (Poly.Vertices.Num() >= 3)
+	{
+		FClipSMPolygon Polygon(0);
+		for (int32 v = 0; v < Poly.Vertices.Num(); ++v)
+		{
+			FClipSMVertex vtx;
+			vtx.Pos = Poly.Vertices[v];
+			Polygon.Vertices.Add(vtx);
+		}
+
+		Polygon.FaceNormal = -PaperAxisZ;
+
+		TArray<FClipSMTriangle> GeneratedTriangles;
+		TriangulatePoly(GeneratedTriangles, Polygon, false);
+
+		struct Local
+		{
+			static FVector4 GenerateVert(const FVector& Source3D, const FVector2D& Size2D)
+			{
+				const float X = FVector::DotProduct(Source3D, PaperAxisX);
+				const float Y = FVector::DotProduct(Source3D, PaperAxisY);
+				const float U = (X / (Size2D.X));
+				const float V = (-Y / (Size2D.Y));
+
+				return FVector4(X, Y, U, V);
+
+			}
+		};
+
+		// Convert the triangles back to our 2D data structure
+		for (int32 TriangleIndex = 0; TriangleIndex < GeneratedTriangles.Num(); ++TriangleIndex)
+		{
+			const FClipSMTriangle& Triangle = GeneratedTriangles[TriangleIndex];
+
+			new (Batch.RenderVerts) FVector4(Local::GenerateVert(Triangle.Vertices[0].Pos, TextureSize));
+			new (Batch.RenderVerts) FVector4(Local::GenerateVert(Triangle.Vertices[1].Pos, TextureSize));
+			new (Batch.RenderVerts) FVector4(Local::GenerateVert(Triangle.Vertices[2].Pos, TextureSize));
+		}
+	}
+#endif
+}
+
 void UPaperTerrainComponent::SetTerrainColor(FLinearColor NewColor)
 {
 	// Can't set color on a static component
