Paper2D: Add PixelsPerUnrealUnit to move away from a 1:1 relationship between pixels and uu

diff --git a/Engine/Plugins/Experimental/Paper2D/Source/Paper2D/Classes/PaperRuntimeSettings.h b/Engine/Plugins/Experimental/Paper2D/Source/Paper2D/Classes/PaperRuntimeSettings.h
--- a/Engine/Plugins/Experimental/Paper2D/Source/Paper2D/Classes/PaperRuntimeSettings.h
+++ b/Engine/Plugins/Experimental/Paper2D/Source/Paper2D/Classes/PaperRuntimeSettings.h
@@ -12,7 +12,7 @@ class PAPER2D_API UPaperRuntimeSettings : public UObject
 {
 	GENERATED_UCLASS_BODY()
 
-	// Test property
+	// The default scaling factor between pixels and Unreal units (cm) to use for newly created sprite assets (e.g., 0.64 would make a 64 pixel wide sprite take up 100 cm)
 	UPROPERTY(GlobalConfig, EditAnywhere, Category=Settings)
-	int32 TestProperty;
+	float DefaultPixelsPerUnrealUnit;
 };
diff --git a/Engine/Plugins/Experimental/Paper2D/Source/Paper2D/Classes/PaperSprite.h b/Engine/Plugins/Experimental/Paper2D/Source/Paper2D/Classes/PaperSprite.h
--- a/Engine/Plugins/Experimental/Paper2D/Source/Paper2D/Classes/PaperSprite.h
+++ b/Engine/Plugins/Experimental/Paper2D/Source/Paper2D/Classes/PaperSprite.h
@@ -104,6 +104,10 @@ protected:
 	UPROPERTY(Category=Collision, EditAnywhere)
 	float CollisionThickness;
 
+	// The scaling factor between pixels and Unreal units (cm) (e.g., 0.64 would make a 64 pixel wide sprite take up 100 cm)
+	UPROPERTY(Category=Sprite, EditAnywhere)
+	float PixelsPerUnrealUnit;
+
 	// Custom render geometry polygons (in texture space)
 	UPROPERTY(Category=Rendering, EditAnywhere)
 	FSpritePolygonCollection RenderGeometry;
@@ -169,6 +173,7 @@ public:
 	FVector2D GetSourceUV() const { return SourceUV; }
 	FVector2D GetSourceSize() const { return SourceDimension; }
 	UTexture2D* GetSourceTexture() const { return SourceTexture; }
+	float GetPixelsPerUnrealUnit() const { return PixelsPerUnrealUnit; }
 #endif
 
 	// Returns the texture this should be rendered with
diff --git a/Engine/Plugins/Experimental/Paper2D/Source/Paper2D/Private/PaperCustomVersion.h b/Engine/Plugins/Experimental/Paper2D/Source/Paper2D/Private/PaperCustomVersion.h
--- a/Engine/Plugins/Experimental/Paper2D/Source/Paper2D/Private/PaperCustomVersion.h
+++ b/Engine/Plugins/Experimental/Paper2D/Source/Paper2D/Private/PaperCustomVersion.h
@@ -22,6 +22,9 @@ struct FPaperCustomVersion
 		// Marked sprite UBodySetup to use simple as complex, allowing 3D raycasts to succeed
 		MarkSpriteBodySetupToUseSimpleAsComplex = 4,
 
+		// Added PixelsPerUnrealUnit to move away from a 1:1 relationship between pixels and uu
+		AddPixelsPerUnrealUnit = 5,
+
 		// -----<new versions can be added above this line>-------------------------------------------------
 		VersionPlusOne,
 		LatestVersion = VersionPlusOne - 1
diff --git a/Engine/Plugins/Experimental/Paper2D/Source/Paper2D/Private/PaperRuntimeSettings.cpp b/Engine/Plugins/Experimental/Paper2D/Source/Paper2D/Private/PaperRuntimeSettings.cpp
--- a/Engine/Plugins/Experimental/Paper2D/Source/Paper2D/Private/PaperRuntimeSettings.cpp
+++ b/Engine/Plugins/Experimental/Paper2D/Source/Paper2D/Private/PaperRuntimeSettings.cpp
@@ -8,5 +8,6 @@
 
 UPaperRuntimeSettings::UPaperRuntimeSettings(const FPostConstructInitializeProperties& PCIP)
 	: Super(PCIP)
+	, DefaultPixelsPerUnrealUnit(2.56f)
 {
 }
diff --git a/Engine/Plugins/Experimental/Paper2D/Source/Paper2D/Private/PaperSprite.cpp b/Engine/Plugins/Experimental/Paper2D/Source/Paper2D/Private/PaperSprite.cpp
--- a/Engine/Plugins/Experimental/Paper2D/Source/Paper2D/Private/PaperSprite.cpp
+++ b/Engine/Plugins/Experimental/Paper2D/Source/Paper2D/Private/PaperSprite.cpp
@@ -260,6 +260,8 @@ UPaperSprite::UPaperSprite(const FPostConstructInitializeProperties& PCIP)
 
 	CollisionGeometry.GeometryType = ESpritePolygonMode::TightBoundingBox;
 	CollisionThickness = 10.0f;
+
+	PixelsPerUnrealUnit = 2.56f;
 #endif
 
 	static ConstructorHelpers::FObjectFinder<UMaterial> DefaultMaterialRef(TEXT("/Paper2D/DefaultSpriteMaterial.DefaultSpriteMaterial"));
@@ -387,11 +389,12 @@ void UPaperSprite::BuildCustomCollisionData()
 	TArray<FVector2D> CollisionData;
 	Triangulate(CollisionGeometry, CollisionData);
 
-	// Adjust the collision data to be relative to the pivot
-	for (int32 PointIndex = 0; PointIndex < CollisionData.Num(); ++PointIndex)
+
+	// Adjust the collision data to be relative to the pivot and scaled from pixels to uu
+	const float UnitsPerPixel = 1.0f / PixelsPerUnrealUnit;
+	for (FVector2D& Point : CollisionData)
 	{
-		FVector2D& Point = CollisionData[PointIndex];
-		Point = ConvertTextureSpaceToPivotSpace(Point);
+		Point = ConvertTextureSpaceToPivotSpace(Point) * UnitsPerPixel;
 	}
 
 	//@TODO: Observe if the custom data is really a hand-edited bounding box, and generate box geom instead of convex geom!
@@ -463,6 +466,8 @@ void UPaperSprite::BuildBoundingBoxCollisionData(bool bUseTightBounds)
 	CreatePolygonFromBoundingBox(CollisionGeometry, bUseTightBounds);
 
 	// Bake it to the runtime structure
+	const float UnitsPerPixel = 1.0f / PixelsPerUnrealUnit;
+
 	switch (SpriteCollisionDomain)
 	{
 	case ESpriteCollisionMode::Use3DPhysics:
@@ -472,17 +477,21 @@ void UPaperSprite::BuildBoundingBoxCollisionData(bool bUseTightBounds)
 			UBodySetup* BodySetup3D = BodySetup;
 			BodySetup3D->AggGeom.EmptyElements();
 
-			// Determine the box center in pivot space
-			const FVector2D& BoxSize2D = CollisionGeometry.Polygons[0].BoxSize;
+			// Determine the box size and center in pivot space
+			const FVector2D& BoxSize2DInPixels = CollisionGeometry.Polygons[0].BoxSize;
 			const FVector2D& BoxPos = CollisionGeometry.Polygons[0].BoxPosition;
-			const FVector2D CenterInTextureSpace = BoxPos + (BoxSize2D * 0.5f);
+			const FVector2D CenterInTextureSpace = BoxPos + (BoxSize2DInPixels * 0.5f);
 			const FVector2D CenterInPivotSpace = ConvertTextureSpaceToPivotSpace(CenterInTextureSpace);
 
+			// Convert from pixels to uu
+			const FVector2D BoxSize2D = BoxSize2DInPixels * UnitsPerPixel;
+			const FVector2D CenterInScaledSpace = CenterInPivotSpace * UnitsPerPixel;
+
 			// Create a new box primitive
 			const FVector BoxSize3D = (PaperAxisX * BoxSize2D.X) + (PaperAxisY * BoxSize2D.Y) + (PaperAxisZ * CollisionThickness);
 
 			FKBoxElem& Box = *new (BodySetup3D->AggGeom.BoxElems) FKBoxElem(FMath::Abs(BoxSize3D.X), FMath::Abs(BoxSize3D.Y), FMath::Abs(BoxSize3D.Z));
-			Box.Center = (PaperAxisX * CenterInPivotSpace.X) + (PaperAxisY * CenterInPivotSpace.Y);
+			Box.Center = (PaperAxisX * CenterInScaledSpace.X) + (PaperAxisY * CenterInScaledSpace.Y);
 
 			BodySetup3D->InvalidatePhysicsData();
 			BodySetup3D->CreatePhysicsMeshes();
@@ -494,17 +503,21 @@ void UPaperSprite::BuildBoundingBoxCollisionData(bool bUseTightBounds)
 			BodySetup2D->AggGeom2D.EmptyElements();
 
 			// Determine the box center in pivot space
-			const FVector2D& BoxSize2D = CollisionGeometry.Polygons[0].BoxSize;
+			const FVector2D& BoxSize2DInPixels = CollisionGeometry.Polygons[0].BoxSize;
 			const FVector2D& BoxPos = CollisionGeometry.Polygons[0].BoxPosition;
-			const FVector2D CenterInTextureSpace = BoxPos + (BoxSize2D * 0.5f);
+			const FVector2D CenterInTextureSpace = BoxPos + (BoxSize2DInPixels * 0.5f);
 			const FVector2D CenterInPivotSpace = ConvertTextureSpaceToPivotSpace(CenterInTextureSpace);
 
+			// Convert from pixels to uu
+			const FVector2D BoxSize2D = BoxSize2DInPixels * UnitsPerPixel;
+			const FVector2D CenterInScaledSpace = CenterInPivotSpace * UnitsPerPixel;
+
 			// Create a new box primitive
 			FBoxElement2D& Box = *new (BodySetup2D->AggGeom2D.BoxElements) FBoxElement2D();
 			Box.Width = FMath::Abs(BoxSize2D.X);
 			Box.Height = FMath::Abs(BoxSize2D.Y);
-			Box.Center.X = CenterInPivotSpace.X;
-			Box.Center.Y = CenterInPivotSpace.Y;
+			Box.Center.X = CenterInScaledSpace.X;
+			Box.Center.Y = CenterInScaledSpace.Y;
 
 			BodySetup2D->InvalidatePhysicsData();
 			BodySetup2D->CreatePhysicsMeshes();
@@ -551,6 +564,7 @@ void UPaperSprite::RebuildRenderData()
 
 	const FVector2D DeltaUV((BakedSourceTexture != nullptr) ? (BakedSourceUV - SourceUV) : FVector2D::ZeroVector);
 
+	const float UnitsPerPixel = 1.0f / PixelsPerUnrealUnit;
 
 	BakedRenderData.Empty(TriangluatedPoints.Num());
 	for (int32 PointIndex = 0; PointIndex < TriangluatedPoints.Num(); ++PointIndex)
@@ -559,7 +573,7 @@ void UPaperSprite::RebuildRenderData()
 		const FVector2D PivotSpacePos = ConvertTextureSpaceToPivotSpace(SourcePos);
 		const FVector2D UV(SourcePos + DeltaUV);
 
-		new (BakedRenderData) FVector4(PivotSpacePos.X, PivotSpacePos.Y, UV.X * InverseWidth, UV.Y * InverseHeight);
+		new (BakedRenderData) FVector4(PivotSpacePos.X * UnitsPerPixel, PivotSpacePos.Y * UnitsPerPixel, UV.X * InverseWidth, UV.Y * InverseHeight);
 	}
 	check((BakedRenderData.Num() % 3) == 0);
 }
@@ -883,6 +897,8 @@ void UPaperSprite::InitializeSprite(UTexture2D* Texture)
 	{
 		InitializeSprite(NULL, FVector2D::ZeroVector, FVector2D(1.0f, 1.0f));
 	}
+
+	PixelsPerUnrealUnit = GetDefault<UPaperRuntimeSettings>()->DefaultPixelsPerUnrealUnit;
 }
 
 void UPaperSprite::InitializeSprite(UTexture2D* Texture, const FVector2D& Offset, const FVector2D& Dimension)
@@ -1072,9 +1088,11 @@ void UPaperSprite::PostLoad()
 	Super::PostLoad();
 
 #if WITH_EDITORONLY_DATA
-	if (GetLinkerCustomVersion(FPaperCustomVersion::GUID) < FPaperCustomVersion::MarkSpriteBodySetupToUseSimpleAsComplex)
+	if (GetLinkerCustomVersion(FPaperCustomVersion::GUID) < FPaperCustomVersion::AddPixelsPerUnrealUnit)
 	{
+		PixelsPerUnrealUnit = 1.0f;
 		RebuildCollisionData();
+		RebuildRenderData();
 	}
 #endif
 }
diff --git a/Engine/Plugins/Experimental/Paper2D/Source/Paper2D/Private/Terrain/PaperTerrainComponent.cpp b/Engine/Plugins/Experimental/Paper2D/Source/Paper2D/Private/Terrain/PaperTerrainComponent.cpp
--- a/Engine/Plugins/Experimental/Paper2D/Source/Paper2D/Private/Terrain/PaperTerrainComponent.cpp
+++ b/Engine/Plugins/Experimental/Paper2D/Source/Paper2D/Private/Terrain/PaperTerrainComponent.cpp
@@ -70,7 +70,7 @@ void FPaperTerrainSceneProxy::DrawDynamicElements_RichMesh(FPrimitiveDrawInterfa
 UPaperTerrainComponent::UPaperTerrainComponent(const FPostConstructInitializeProperties& PCIP)
 	: Super(PCIP)
 	, TerrainColor(FLinearColor::White)
-	, ReparamStepsPerSegment(4)
+	, ReparamStepsPerSegment(8)
 {
 	TestScaleFactor = 1.0f;
 }
@@ -569,6 +569,8 @@ void UPaperTerrainComponent::SpawnFromPoly(const class UPaperSprite* NewSprite,
 {
 	//@TODO: Need to split geom tools out of UnrealEd to support this outside of the editor
 #if WITH_EDITOR
+	const FVector2D TextureSizeInUnits = TextureSize / NewSprite->GetPixelsPerUnrealUnit();
+
 	if (Poly.Vertices.Num() >= 3)
 	{
 		FClipSMPolygon Polygon(0);
@@ -594,7 +596,6 @@ void UPaperTerrainComponent::SpawnFromPoly(const class UPaperSprite* NewSprite,
 				const float V = (-Y / (Size2D.Y));
 
 				return FVector4(X, Y, U, V);
-
 			}
 		};
 
@@ -603,9 +604,9 @@ void UPaperTerrainComponent::SpawnFromPoly(const class UPaperSprite* NewSprite,
 		{
 			const FClipSMTriangle& Triangle = GeneratedTriangles[TriangleIndex];
 
-			new (Batch.RenderVerts) FVector4(Local::GenerateVert(Triangle.Vertices[0].Pos, TextureSize));
-			new (Batch.RenderVerts) FVector4(Local::GenerateVert(Triangle.Vertices[1].Pos, TextureSize));
-			new (Batch.RenderVerts) FVector4(Local::GenerateVert(Triangle.Vertices[2].Pos, TextureSize));
+			new (Batch.RenderVerts) FVector4(Local::GenerateVert(Triangle.Vertices[0].Pos, TextureSizeInUnits));
+			new (Batch.RenderVerts) FVector4(Local::GenerateVert(Triangle.Vertices[1].Pos, TextureSizeInUnits));
+			new (Batch.RenderVerts) FVector4(Local::GenerateVert(Triangle.Vertices[2].Pos, TextureSizeInUnits));
 		}
 	}
 #endif
