Paper2D: Fix typo in 3D convex hull collision generation for sprites that generates invalid shapes and can cause cooking to fail

diff --git a/Engine/Plugins/Experimental/Paper2D/Source/Paper2D/Private/PaperCustomVersion.h b/Engine/Plugins/Experimental/Paper2D/Source/Paper2D/Private/PaperCustomVersion.h
--- a/Engine/Plugins/Experimental/Paper2D/Source/Paper2D/Private/PaperCustomVersion.h
+++ b/Engine/Plugins/Experimental/Paper2D/Source/Paper2D/Private/PaperCustomVersion.h
@@ -25,6 +25,9 @@ struct FPaperCustomVersion
 		// Added PixelsPerUnrealUnit to move away from a 1:1 relationship between pixels and uu
 		AddPixelsPerUnrealUnit = 5,
 
+		// Fixed typo in convex hull generation code for 3D custom collision
+		FixTypoIn3DConvexHullCollisionGeneration = 6,
+
 		// -----<new versions can be added above this line>-------------------------------------------------
 		VersionPlusOne,
 		LatestVersion = VersionPlusOne - 1
diff --git a/Engine/Plugins/Experimental/Paper2D/Source/Paper2D/Private/PaperSprite.cpp b/Engine/Plugins/Experimental/Paper2D/Source/Paper2D/Private/PaperSprite.cpp
--- a/Engine/Plugins/Experimental/Paper2D/Source/Paper2D/Private/PaperSprite.cpp
+++ b/Engine/Plugins/Experimental/Paper2D/Source/Paper2D/Private/PaperSprite.cpp
@@ -421,7 +421,7 @@ void UPaperSprite::BuildCustomCollisionData()
 				{
 					const FVector2D& Pos2D = CollisionData[RunningIndex++];
 					
-					const FVector Pos3D = (PaperAxisX * Pos2D.X) + (PaperAxisY * Pos2D.X);
+					const FVector Pos3D = (PaperAxisX * Pos2D.X) + (PaperAxisY * Pos2D.Y);
 
 					new (ConvexTri.VertexData) FVector(Pos3D - HalfThicknessVector);
 					new (ConvexTri.VertexData) FVector(Pos3D + HalfThicknessVector);
@@ -1088,10 +1088,29 @@ void UPaperSprite::PostLoad()
 	Super::PostLoad();
 
 #if WITH_EDITORONLY_DATA
-	if (GetLinkerCustomVersion(FPaperCustomVersion::GUID) < FPaperCustomVersion::AddPixelsPerUnrealUnit)
+	const int32 PaperVer = GetLinkerCustomVersion(FPaperCustomVersion::GUID);
+
+	bool bRebuildCollision = false;
+	bool bRebuildRenderData = false;
+
+	if (PaperVer < FPaperCustomVersion::AddPixelsPerUnrealUnit)
 	{
 		PixelsPerUnrealUnit = 1.0f;
+		bRebuildCollision = true;
+		bRebuildRenderData = true;
+	}
+	else if (PaperVer < FPaperCustomVersion::FixTypoIn3DConvexHullCollisionGeneration)
+	{
+		bRebuildCollision = true;
+	}
+
+	if (bRebuildCollision)
+	{
 		RebuildCollisionData();
+	}
+
+	if (bRebuildRenderData)
+	{
 		RebuildRenderData();
 	}
 #endif
