Paper2D: Fix possible load order assert in Paper2D

diff --git a/Engine/Plugins/Experimental/Paper2D/Source/SmartSnapping/Private/SmartSnappingModule.cpp b/Engine/Plugins/Experimental/Paper2D/Source/SmartSnapping/Private/SmartSnappingModule.cpp
--- a/Engine/Plugins/Experimental/Paper2D/Source/SmartSnapping/Private/SmartSnappingModule.cpp
+++ b/Engine/Plugins/Experimental/Paper2D/Source/SmartSnapping/Private/SmartSnappingModule.cpp
@@ -30,7 +30,7 @@ public:
 				ViewMenuExtender = FLevelEditorModule::FLevelEditorMenuExtender::CreateRaw(this, &FSmartSnappingModule::OnExtendLevelEditorViewMenu);
 
 
-				FLevelEditorModule& LevelEditor = FModuleManager::GetModuleChecked<FLevelEditorModule>(TEXT("LevelEditor"));
+				FLevelEditorModule& LevelEditor = FModuleManager::LoadModuleChecked<FLevelEditorModule>(TEXT("LevelEditor"));
 				LevelEditor.GetAllLevelEditorToolbarViewMenuExtenders().Add(ViewMenuExtender);
 			}
 		}
