Paper2D: Add placeholder tutorials for sprites and flipbooks

diff --git a/Engine/Plugins/Experimental/Paper2D/Source/Paper2DEditor/Paper2DEditor.Build.cs b/Engine/Plugins/Experimental/Paper2D/Source/Paper2DEditor/Paper2DEditor.Build.cs
--- a/Engine/Plugins/Experimental/Paper2D/Source/Paper2DEditor/Paper2DEditor.Build.cs
+++ b/Engine/Plugins/Experimental/Paper2D/Source/Paper2DEditor/Paper2DEditor.Build.cs
@@ -38,6 +38,7 @@ public class Paper2DEditor : ModuleRules
 		PrivateIncludePathModuleNames.AddRange(
 			new string[] {
 				"Settings",
+				"IntroTutorials"
 			}
 		);
 	}
diff --git a/Engine/Plugins/Experimental/Paper2D/Source/Paper2DEditor/Private/Paper2DEditorModule.cpp b/Engine/Plugins/Experimental/Paper2D/Source/Paper2DEditor/Private/Paper2DEditorModule.cpp
--- a/Engine/Plugins/Experimental/Paper2D/Source/Paper2DEditor/Private/Paper2DEditorModule.cpp
+++ b/Engine/Plugins/Experimental/Paper2D/Source/Paper2DEditor/Private/Paper2DEditorModule.cpp
@@ -29,6 +29,9 @@
 #include "PaperRuntimeSettings.h"
 #include "Settings.h"
 
+// Intro tutorials
+#include "Editor/IntroTutorials/Public/IIntroTutorials.h"
+
 DEFINE_LOG_CATEGORY(LogPaper2DEditor);
 
 #define LOCTEXT_NAMESPACE "Paper2DEditor"
@@ -122,6 +125,7 @@ public:
 		FPaperContentBrowserExtensions::InstallHooks();
 
 		RegisterSettings();
+		RegisterIntroTutorials();
 	}
 
 	virtual void ShutdownModule() override
@@ -134,6 +138,7 @@ public:
 
 		if (UObjectInitialized())
 		{
+			UnregisterIntroTutorials();
 			UnregisterSettings();
 
 			FPaperContentBrowserExtensions::RemoveHooks();
@@ -205,6 +210,34 @@ private:
 			SettingsModule->UnregisterSettings("Project", "Plugins", "Paper2D");
 		}
 	}
+
+	void RegisterIntroTutorials()
+	{
+		if (!IsRunningCommandlet())
+		{
+			//@TODO: PAPER2D: Remove the _Preview suffix on the config keys once the final doc is in place
+			// (this is so that people who dismiss the early warning message still get the final intro doc later on)
+			IIntroTutorials::Get().RegisterTutorialForAssetEditor(
+				UPaperSprite::StaticClass(),
+				TEXT("Shared/Tutorials/InPaperSpriteEditorTutorial"),
+				TEXT("SeenPaperSpriteEditorWelcome_Preview"),
+				FString("19D7EA18-629B-4A86-BD19-ED2B3BE53600"));
+			IIntroTutorials::Get().RegisterTutorialForAssetEditor(
+				UPaperFlipbook::StaticClass(),
+				TEXT("Shared/Tutorials/InPaperFlipbookEditorTutorial"),
+				TEXT("SeenPaperFlipbookEditorWelcome_Preview"),
+				FString("B24214C1-E17A-4F95-BE18-2ED8BFCEC008"));
+		}
+	}
+
+	void UnregisterIntroTutorials()
+	{
+		if (IIntroTutorials::IsAvailable())
+		{
+			IIntroTutorials::Get().UnregisterTutorialForAssetEditor(UPaperSprite::StaticClass());
+			IIntroTutorials::Get().UnregisterTutorialForAssetEditor(UPaperFlipbook::StaticClass());
+		}
+	}
 };
 
 //////////////////////////////////////////////////////////////////////////
