Paper2D: Add preview watermark to sprite and flipbook editors

diff --git a/Engine/Plugins/Experimental/Paper2D/Source/Paper2DEditor/Private/FlipbookEditor/FlipbookEditor.cpp b/Engine/Plugins/Experimental/Paper2D/Source/Paper2DEditor/Private/FlipbookEditor/FlipbookEditor.cpp
--- a/Engine/Plugins/Experimental/Paper2D/Source/Paper2DEditor/Private/FlipbookEditor/FlipbookEditor.cpp
+++ b/Engine/Plugins/Experimental/Paper2D/Source/Paper2DEditor/Private/FlipbookEditor/FlipbookEditor.cpp
@@ -181,7 +181,25 @@ TSharedRef<SDockTab> FFlipbookEditor::SpawnTab_Viewport(const FSpawnTabArgs& Arg
 			
 			+SVerticalBox::Slot()
 			[
-				ViewportPtr.ToSharedRef()
+				SNew(SOverlay)
+
+				// The sprite editor viewport
+				+SOverlay::Slot()
+				[
+					ViewportPtr.ToSharedRef()
+				]
+
+				// Bottom-right corner text indicating the preview nature of the sprite editor
+				+SOverlay::Slot()
+				.Padding(10)
+				.VAlign(VAlign_Bottom)
+				.HAlign(HAlign_Right)
+				[
+					SNew(STextBlock)
+					.Visibility( EVisibility::HitTestInvisible )
+					.TextStyle( FEditorStyle::Get(), "Graph.CornerText" )
+					.Text(LOCTEXT("FlipbookEditorViewportExperimentalWarning", "Early access preview"))
+				]
 			]
 			
 			+SVerticalBox::Slot()
diff --git a/Engine/Plugins/Experimental/Paper2D/Source/Paper2DEditor/Private/SpriteEditor/SpriteEditor.cpp b/Engine/Plugins/Experimental/Paper2D/Source/Paper2DEditor/Private/SpriteEditor/SpriteEditor.cpp
--- a/Engine/Plugins/Experimental/Paper2D/Source/Paper2DEditor/Private/SpriteEditor/SpriteEditor.cpp
+++ b/Engine/Plugins/Experimental/Paper2D/Source/Paper2DEditor/Private/SpriteEditor/SpriteEditor.cpp
@@ -245,7 +245,25 @@ TSharedRef<SDockTab> FSpriteEditor::SpawnTab_Viewport(const FSpawnTabArgs& Args)
 	return SNew(SDockTab)
 		.Label(LOCTEXT("ViewportTab_Title", "Viewport"))
 		[
-			ViewportPtr.ToSharedRef()
+			SNew(SOverlay)
+
+			// The sprite editor viewport
+			+SOverlay::Slot()
+			[
+				ViewportPtr.ToSharedRef()
+			]
+
+			// Bottom-right corner text indicating the preview nature of the sprite editor
+			+SOverlay::Slot()
+			.Padding(10)
+			.VAlign(VAlign_Bottom)
+			.HAlign(HAlign_Right)
+			[
+				SNew(STextBlock)
+				.Visibility( EVisibility::HitTestInvisible )
+				.TextStyle( FEditorStyle::Get(), "Graph.CornerText" )
+				.Text(LOCTEXT("SpriteEditorViewportExperimentalWarning", "Early access preview"))
+			]
 		];
 }
 
