Paper2D: Mark classes as EarlyAccessPreview or Experimental as appropriate

diff --git a/Engine/Plugins/Experimental/Paper2D/Source/Paper2D/Classes/PaperAnimatedRenderComponent.h b/Engine/Plugins/Experimental/Paper2D/Source/Paper2D/Classes/PaperAnimatedRenderComponent.h
--- a/Engine/Plugins/Experimental/Paper2D/Source/Paper2D/Classes/PaperAnimatedRenderComponent.h
+++ b/Engine/Plugins/Experimental/Paper2D/Source/Paper2D/Classes/PaperAnimatedRenderComponent.h
@@ -5,7 +5,7 @@
 #include "PaperFlipbook.h"
 #include "PaperAnimatedRenderComponent.generated.h"
 
-UCLASS(MinimalAPI, ShowCategories=(Mobility), meta=(BlueprintSpawnableComponent))
+UCLASS(MinimalAPI, ShowCategories=(Mobility), EarlyAccessPreview, meta=(BlueprintSpawnableComponent))
 class UPaperAnimatedRenderComponent : public UPrimitiveComponent
 {
 	GENERATED_UCLASS_BODY()
diff --git a/Engine/Plugins/Experimental/Paper2D/Source/Paper2D/Classes/PaperRenderComponent.h b/Engine/Plugins/Experimental/Paper2D/Source/Paper2D/Classes/PaperRenderComponent.h
--- a/Engine/Plugins/Experimental/Paper2D/Source/Paper2D/Classes/PaperRenderComponent.h
+++ b/Engine/Plugins/Experimental/Paper2D/Source/Paper2D/Classes/PaperRenderComponent.h
@@ -4,7 +4,7 @@
 
 #include "PaperRenderComponent.generated.h"
 
-UCLASS(DependsOn=UPaperSprite, MinimalAPI, ShowCategories=(Mobility), meta=(BlueprintSpawnableComponent))
+UCLASS(DependsOn=UPaperSprite, MinimalAPI, ShowCategories=(Mobility), EarlyAccessPreview, meta=(BlueprintSpawnableComponent))
 class UPaperRenderComponent : public UPrimitiveComponent
 {
 	GENERATED_UCLASS_BODY()
diff --git a/Engine/Plugins/Experimental/Paper2D/Source/Paper2D/Classes/PaperTerrainActor.h b/Engine/Plugins/Experimental/Paper2D/Source/Paper2D/Classes/PaperTerrainActor.h
--- a/Engine/Plugins/Experimental/Paper2D/Source/Paper2D/Classes/PaperTerrainActor.h
+++ b/Engine/Plugins/Experimental/Paper2D/Source/Paper2D/Classes/PaperTerrainActor.h
@@ -8,7 +8,7 @@
  * An instance of a piece of 2D terrain in the level
  */
 
-UCLASS(DependsOn=UEngineTypes, BlueprintType)
+UCLASS(DependsOn=UEngineTypes, BlueprintType, Experimental)
 class PAPER2D_API APaperTerrainActor : public AActor
 {
 	GENERATED_UCLASS_BODY()
diff --git a/Engine/Plugins/Experimental/Paper2D/Source/Paper2D/Classes/PaperTerrainComponent.h b/Engine/Plugins/Experimental/Paper2D/Source/Paper2D/Classes/PaperTerrainComponent.h
--- a/Engine/Plugins/Experimental/Paper2D/Source/Paper2D/Classes/PaperTerrainComponent.h
+++ b/Engine/Plugins/Experimental/Paper2D/Source/Paper2D/Classes/PaperTerrainComponent.h
@@ -16,7 +16,7 @@ struct FPaperTerrainMaterialPair
  * This takes a 2D terrain material and instances sprite geometry along the spline path.
  */
 
-UCLASS(BlueprintType)
+UCLASS(BlueprintType, Experimental)
 class PAPER2D_API UPaperTerrainComponent : public UPrimitiveComponent
 {
 	GENERATED_UCLASS_BODY()
diff --git a/Engine/Plugins/Experimental/Paper2D/Source/Paper2D/Classes/PaperTerrainSplineComponent.h b/Engine/Plugins/Experimental/Paper2D/Source/Paper2D/Classes/PaperTerrainSplineComponent.h
--- a/Engine/Plugins/Experimental/Paper2D/Source/Paper2D/Classes/PaperTerrainSplineComponent.h
+++ b/Engine/Plugins/Experimental/Paper2D/Source/Paper2D/Classes/PaperTerrainSplineComponent.h
@@ -6,7 +6,7 @@
 #include "PaperTerrainSplineComponent.generated.h"
 
 //@TODO: Document
-UCLASS(BlueprintType)
+UCLASS(BlueprintType, Experimental)
 class PAPER2D_API UPaperTerrainSplineComponent : public USplineComponent
 {
 	GENERATED_UCLASS_BODY()
diff --git a/Engine/Plugins/Experimental/Paper2D/Source/Paper2D/Classes/PaperTileMapRenderComponent.h b/Engine/Plugins/Experimental/Paper2D/Source/Paper2D/Classes/PaperTileMapRenderComponent.h
--- a/Engine/Plugins/Experimental/Paper2D/Source/Paper2D/Classes/PaperTileMapRenderComponent.h
+++ b/Engine/Plugins/Experimental/Paper2D/Source/Paper2D/Classes/PaperTileMapRenderComponent.h
@@ -4,7 +4,7 @@
 
 #include "PaperTileMapRenderComponent.generated.h"
 
-UCLASS(DependsOn=UPaperSprite, hideCategories=Object)
+UCLASS(DependsOn=UPaperSprite, hideCategories=Object, Experimental)
 class PAPER2D_API UPaperTileMapRenderComponent : public UPrimitiveComponent
 {
 	GENERATED_UCLASS_BODY()
