Paper2D: Relax restrictions on JSON "app" tag to allow importing sprite sheets exported from tools other than Flash

diff --git a/Engine/Plugins/Experimental/Paper2D/Source/PaperJsonImporter/Private/PaperJsonImporterFactory.cpp b/Engine/Plugins/Experimental/Paper2D/Source/PaperJsonImporter/Private/PaperJsonImporterFactory.cpp
--- a/Engine/Plugins/Experimental/Paper2D/Source/PaperJsonImporter/Private/PaperJsonImporterFactory.cpp
+++ b/Engine/Plugins/Experimental/Paper2D/Source/PaperJsonImporter/Private/PaperJsonImporterFactory.cpp
@@ -90,12 +90,28 @@ UObject* UPaperJsonImporterFactory::FactoryCreateText(UClass* InClass, UObject*
 			const FString AppName = ReadString(MetaBlock, TEXT("app"), TEXT(""));
 			const FString Image = ReadString(MetaBlock, TEXT("image"), TEXT(""));
 			
-			const FString FlashPrefix(TEXT("Adobe Flash"));
-			if (bLoadedSuccessfully && !AppName.StartsWith(FlashPrefix))
+ 			if (bLoadedSuccessfully)
 			{
-				UE_LOG(LogPaperJsonImporter, Warning, TEXT("Failed to parse sprite descriptor file '%s'.  Expected 'app' to start with %s"), *NameForErrors, *FlashPrefix);
-				bLoadedSuccessfully = false;
-			}
+				const FString FlashPrefix(TEXT("Adobe Flash"));
+				const FString TexturePackerPrefix(TEXT("http://www.codeandweb.com/texturepacker"));
+
+				if (AppName.StartsWith(FlashPrefix) || AppName.StartsWith(TexturePackerPrefix))
+				{
+					// Cool, we (mostly) know how to handle these sorts of files!
+					UE_LOG(LogPaperJsonImporter, Log, TEXT("Parsing sprite sheet exported from '%s'"), *AppName);
+				}
+				else if (!AppName.IsEmpty())
+				{
+					// It's got an app tag inside a meta block, so we'll take a crack at it
+					UE_LOG(LogPaperJsonImporter, Warning, TEXT("Unexpected 'app' named '%s' while parsing sprite descriptor file '%s'.  Parsing will continue but the format may not be fully supported"), *AppName, *NameForErrors);
+				}
+				else
+				{
+					// Probably not a sprite sheet
+					UE_LOG(LogPaperJsonImporter, Warning, TEXT("Failed to parse sprite descriptor file '%s'.  Expected 'app' key indicating the exporter (might not be a sprite sheet)"), *NameForErrors);
+					bLoadedSuccessfully = false;
+				}
+ 			}
 
 			if (bLoadedSuccessfully)
 			{
