Paper2D: Fix various sprite editor issues when PixelsPerUnrealUnit is not 1.0

diff --git a/Engine/Plugins/Experimental/Paper2D/Source/Paper2D/Private/PaperSprite.cpp b/Engine/Plugins/Experimental/Paper2D/Source/Paper2D/Private/PaperSprite.cpp
--- a/Engine/Plugins/Experimental/Paper2D/Source/Paper2D/Private/PaperSprite.cpp
+++ b/Engine/Plugins/Experimental/Paper2D/Source/Paper2D/Private/PaperSprite.cpp
@@ -952,15 +952,15 @@ FVector UPaperSprite::ConvertPivotSpaceToTextureSpace(FVector Input) const
 
 FVector UPaperSprite::ConvertTextureSpaceToWorldSpace(const FVector2D& SourcePoint) const
 {
-	const FVector2D SourceDims = (SourceTexture != NULL) ? FVector2D(SourceTexture->GetSurfaceWidth(), SourceTexture->GetSurfaceHeight()) : FVector2D::ZeroVector;
-	return (PaperAxisX * SourcePoint.X) + (PaperAxisY * (SourceDims.Y - SourcePoint.Y));
-}
+	const float UnitsPerPixel = 1.0f / PixelsPerUnrealUnit;
 
-// FVector UPaperSprite::ConvertTextureSpaceToWorldSpace(const FVector& SourcePoint) const
-// {
-// 	const FVector2D SourceDims = (SourceTexture != NULL) ? FVector2D(SourceTexture->GetSurfaceWidth(), SourceTexture->GetSurfaceHeight()) : FVector2D::ZeroVector;
-// 	return FVector(SourcePoint.X, SourcePoint.Y, SourceDims.Y - SourcePoint.Z);
-// }
+	const FVector2D SourcePointInUU = SourcePoint * UnitsPerPixel;
+
+	const FVector2D SourceDimsInPixels = (SourceTexture != NULL) ? FVector2D(SourceTexture->GetSurfaceWidth(), SourceTexture->GetSurfaceHeight()) : FVector2D::ZeroVector;
+	const FVector2D SourceDimsInUU = SourceDimsInPixels * UnitsPerPixel;
+
+	return (PaperAxisX * SourcePointInUU.X) + (PaperAxisY * (SourceDimsInUU.Y - SourcePointInUU.Y));
+}
 
 FVector2D UPaperSprite::ConvertWorldSpaceToTextureSpace(const FVector& WorldPoint) const
 {
diff --git a/Engine/Plugins/Experimental/Paper2D/Source/Paper2DEditor/Private/SpriteEditor/SpriteEditorViewportClient.cpp b/Engine/Plugins/Experimental/Paper2D/Source/Paper2DEditor/Private/SpriteEditor/SpriteEditorViewportClient.cpp
--- a/Engine/Plugins/Experimental/Paper2D/Source/Paper2DEditor/Private/SpriteEditor/SpriteEditorViewportClient.cpp
+++ b/Engine/Plugins/Experimental/Paper2D/Source/Paper2DEditor/Private/SpriteEditor/SpriteEditorViewportClient.cpp
@@ -95,6 +95,7 @@ void FSpriteEditorViewportClient::UpdateSourceTextureSpriteFromSprite(UPaperSpri
 	if (SourceSprite != NULL)
 	{
 		TargetSprite->PivotMode = ESpritePivotMode::Bottom_Left;
+		TargetSprite->PixelsPerUnrealUnit = SourceSprite->PixelsPerUnrealUnit;
 		TargetSprite->BlendMode = SourceSprite->BlendMode;
 
 		if (SourceSprite->GetSourceTexture() != TargetSprite->GetSourceTexture())
@@ -537,7 +538,7 @@ void FSpriteEditorViewportClient::Tick(float DeltaSeconds)
 		FIntPoint Size = Viewport->GetSizeXY();
 		if (bDeferZoomToSprite && Size.X > 0 && Size.Y > 0)
 		{
-			UPaperRenderComponent *ComponentToFocusOn = SourceTextureViewComponent->IsVisible() ? SourceTextureViewComponent : RenderSpriteComponent;
+			UPaperRenderComponent* ComponentToFocusOn = SourceTextureViewComponent->IsVisible() ? SourceTextureViewComponent : RenderSpriteComponent;
 			FocusViewportOnBox(ComponentToFocusOn->Bounds.GetBox(), true);
 			bDeferZoomToSprite = false;
 		}		
