Paper2D: Add brief description of flipbook assets

diff --git a/Engine/Plugins/Experimental/Paper2D/Source/Paper2D/Classes/PaperFlipbook.h b/Engine/Plugins/Experimental/Paper2D/Source/Paper2D/Classes/PaperFlipbook.h
--- a/Engine/Plugins/Experimental/Paper2D/Source/Paper2D/Classes/PaperFlipbook.h
+++ b/Engine/Plugins/Experimental/Paper2D/Source/Paper2D/Classes/PaperFlipbook.h
@@ -24,7 +24,10 @@ public:
 	}
 };
 
-UCLASS(MinimalAPI, BlueprintType, meta=(DisplayThumbnail = "true"))
+/**
+ * Contains an animation sequence of sprite frames
+ */
+UCLASS(MinimalAPI, BlueprintType, meta = (DisplayThumbnail = "true"))
 class UPaperFlipbook : public UObject
 {
 	GENERATED_UCLASS_BODY()
