Paper2D: Hide experimental sprite atlas group asset behind an experimental flag

diff --git a/Engine/Plugins/Experimental/Paper2D/Source/Paper2D/Classes/PaperRuntimeSettings.h b/Engine/Plugins/Experimental/Paper2D/Source/Paper2D/Classes/PaperRuntimeSettings.h
--- a/Engine/Plugins/Experimental/Paper2D/Source/Paper2D/Classes/PaperRuntimeSettings.h
+++ b/Engine/Plugins/Experimental/Paper2D/Source/Paper2D/Classes/PaperRuntimeSettings.h
@@ -15,4 +15,8 @@ class PAPER2D_API UPaperRuntimeSettings : public UObject
 	// The default scaling factor between pixels and Unreal units (cm) to use for newly created sprite assets (e.g., 0.64 would make a 64 pixel wide sprite take up 100 cm)
 	UPROPERTY(GlobalConfig, EditAnywhere, Category=Settings)
 	float DefaultPixelsPerUnrealUnit;
+
+	// Enables experimental *incomplete and unsupported* texture atlas groups that sprites can be assigned to
+	UPROPERTY(EditAnywhere, config, Category=Experimental)
+	bool bEnableSpriteAtlasGroups;
 };
diff --git a/Engine/Plugins/Experimental/Paper2D/Source/Paper2D/Classes/PaperSpriteAtlas.h b/Engine/Plugins/Experimental/Paper2D/Source/Paper2D/Classes/PaperSpriteAtlas.h
--- a/Engine/Plugins/Experimental/Paper2D/Source/Paper2D/Classes/PaperSpriteAtlas.h
+++ b/Engine/Plugins/Experimental/Paper2D/Source/Paper2D/Classes/PaperSpriteAtlas.h
@@ -5,7 +5,7 @@
 #include "PaperSpriteAtlas.generated.h"
 
 // Groups together a set of sprites that will try to share the same texture atlas (allowing them to be combined into a single draw call)
-UCLASS(MinimalAPI, BlueprintType, meta=(DisplayThumbnail = "true"))
+UCLASS(MinimalAPI, BlueprintType, Experimental, meta=(DisplayThumbnail = "true"))
 class UPaperSpriteAtlas : public UObject
 {
 	GENERATED_UCLASS_BODY()
diff --git a/Engine/Plugins/Experimental/Paper2D/Source/Paper2D/Private/PaperRuntimeSettings.cpp b/Engine/Plugins/Experimental/Paper2D/Source/Paper2D/Private/PaperRuntimeSettings.cpp
--- a/Engine/Plugins/Experimental/Paper2D/Source/Paper2D/Private/PaperRuntimeSettings.cpp
+++ b/Engine/Plugins/Experimental/Paper2D/Source/Paper2D/Private/PaperRuntimeSettings.cpp
@@ -9,5 +9,6 @@
 UPaperRuntimeSettings::UPaperRuntimeSettings(const FPostConstructInitializeProperties& PCIP)
 	: Super(PCIP)
 	, DefaultPixelsPerUnrealUnit(2.56f)
+	, bEnableSpriteAtlasGroups(false)
 {
 }
diff --git a/Engine/Plugins/Experimental/Paper2D/Source/Paper2DEditor/Private/Atlasing/PaperSpriteAtlasFactory.cpp b/Engine/Plugins/Experimental/Paper2D/Source/Paper2DEditor/Private/Atlasing/PaperSpriteAtlasFactory.cpp
--- a/Engine/Plugins/Experimental/Paper2D/Source/Paper2DEditor/Private/Atlasing/PaperSpriteAtlasFactory.cpp
+++ b/Engine/Plugins/Experimental/Paper2D/Source/Paper2DEditor/Private/Atlasing/PaperSpriteAtlasFactory.cpp
@@ -20,3 +20,8 @@ UObject* UPaperSpriteAtlasFactory::FactoryCreateNew(UClass* InClass, UObject* In
 {
 	return ConstructObject<UObject>(SupportedClass, InParent, InName, Flags);
 }
+
+bool UPaperSpriteAtlasFactory::CanCreateNew() const
+{
+	return GetDefault<UPaperRuntimeSettings>()->bEnableSpriteAtlasGroups;
+}
diff --git a/Engine/Plugins/Experimental/Paper2D/Source/Paper2DEditor/Private/Atlasing/PaperSpriteAtlasFactory.h b/Engine/Plugins/Experimental/Paper2D/Source/Paper2DEditor/Private/Atlasing/PaperSpriteAtlasFactory.h
--- a/Engine/Plugins/Experimental/Paper2D/Source/Paper2DEditor/Private/Atlasing/PaperSpriteAtlasFactory.h
+++ b/Engine/Plugins/Experimental/Paper2D/Source/Paper2DEditor/Private/Atlasing/PaperSpriteAtlasFactory.h
@@ -11,6 +11,7 @@ class UPaperSpriteAtlasFactory : public UFactory
 
 	// UFactory interface
 	virtual UObject* FactoryCreateNew(UClass* Class, UObject* InParent, FName Name, EObjectFlags Flags, UObject* Context, FFeedbackContext* Warn) override;
+	virtual bool CanCreateNew() const override;
 	// End of UFactory interface
 };
 
