Move PaperCharacter into Paper2D

diff --git a/Engine/Plugins/Experimental/Paper2D/Source/Paper2D/Classes/PaperCharacter.h b/Engine/Plugins/Experimental/Paper2D/Source/Paper2D/Classes/PaperCharacter.h
new file mode 100644
--- /dev/null
+++ b/Engine/Plugins/Experimental/Paper2D/Source/Paper2D/Classes/PaperCharacter.h
@@ -0,0 +1,26 @@
+// Copyright 1998-2014 Epic Games, Inc. All Rights Reserved.
+
+#pragma once
+
+#include "PaperCharacter.generated.h"
+
+class UPaperAnimatedRenderComponent;
+
+//=============================================================================
+// Paper Characters Exactly like Characters but are sprite based.
+//=============================================================================
+UCLASS(MinimalAPI)
+class APaperCharacter : public ACharacter
+{
+	GENERATED_UCLASS_BODY()
+	
+	static FName SpriteComponentName;
+	/** The main skeletal mesh associated with this Character (optional sub-object). */
+	UPROPERTY(Category = Character, VisibleAnywhere, BlueprintReadOnly)
+	TSubobjectPtr<class UPaperAnimatedRenderComponent> Sprite;
+
+protected:
+
+public:	
+
+};
diff --git a/Engine/Plugins/Experimental/Paper2D/Source/Paper2D/Private/PaperAnimatedRenderComponent.cpp b/Engine/Plugins/Experimental/Paper2D/Source/Paper2D/Private/PaperAnimatedRenderComponent.cpp
--- a/Engine/Plugins/Experimental/Paper2D/Source/Paper2D/Private/PaperAnimatedRenderComponent.cpp
+++ b/Engine/Plugins/Experimental/Paper2D/Source/Paper2D/Private/PaperAnimatedRenderComponent.cpp
@@ -134,6 +134,10 @@ bool UPaperAnimatedRenderComponent::SetFlipbook(class UPaperFlipbook* NewFlipboo
 		{
 			SourceFlipbook = NewFlipbook;
 
+			// We need to also reset the frame and time also
+			AccumulatedTime = 0.0f;
+			CalculateCurrentFrame();
+
 			// Need to send this to render thread at some point
 			MarkRenderStateDirty();
 
diff --git a/Engine/Plugins/Experimental/Paper2D/Source/Paper2D/Private/PaperCharacter.cpp b/Engine/Plugins/Experimental/Paper2D/Source/Paper2D/Private/PaperCharacter.cpp
new file mode 100644
--- /dev/null
+++ b/Engine/Plugins/Experimental/Paper2D/Source/Paper2D/Private/PaperCharacter.cpp
@@ -0,0 +1,41 @@
+// Copyright 1998-2014 Epic Games, Inc. All Rights Reserved.
+
+/*=============================================================================
+	PaperCharacter.cpp: APaperCharacter implementation
+	A character that uses a paper sprite.
+=============================================================================*/
+
+#include "Paper2DPrivatePCH.h"
+
+
+FName APaperCharacter::SpriteComponentName(TEXT("Sprite0"));
+
+APaperCharacter::APaperCharacter(const class FPostConstructInitializeProperties& PCIP)
+	: Super(PCIP)
+{
+	// Prevent creation of mesh
+	PCIP.DoNotCreateDefaultSubobject(ACharacter::MeshComponentName);
+	
+	Sprite = PCIP.CreateOptionalDefaultSubobject<UPaperAnimatedRenderComponent>(this, APaperCharacter::SpriteComponentName);
+	if (Sprite)
+	{
+		Sprite->AlwaysLoadOnClient = true;
+		Sprite->AlwaysLoadOnServer = true;
+		Sprite->bOwnerNoSee = false;
+		//Mesh->MeshComponentUpdateFlag = EMeshComponentUpdateFlag::AlwaysTickPose;
+		//Mesh->bCastDynamicShadow = true;
+		Sprite->bAffectDynamicIndirectLighting = true;
+		Sprite->PrimaryComponentTick.TickGroup = TG_PrePhysics;
+		// force tick after movement component updates
+		if (CharacterMovement)
+		{
+			Sprite->PrimaryComponentTick.AddPrerequisite(this, CharacterMovement->PrimaryComponentTick);
+		}
+		Sprite->AttachParent = CapsuleComponent;
+		static FName CollisionProfileName(TEXT("CharacterMesh"));
+		Sprite->SetCollisionProfileName(CollisionProfileName);
+		Sprite->bGenerateOverlapEvents = false;
+	}
+}
+
+
