Paper2D: Remove unused BlendMode property from UPaperSprite

diff --git a/Engine/Plugins/Experimental/Paper2D/Source/Paper2D/Classes/PaperSprite.h b/Engine/Plugins/Experimental/Paper2D/Source/Paper2D/Classes/PaperSprite.h
--- a/Engine/Plugins/Experimental/Paper2D/Source/Paper2D/Classes/PaperSprite.h
+++ b/Engine/Plugins/Experimental/Paper2D/Source/Paper2D/Classes/PaperSprite.h
@@ -74,10 +74,6 @@ protected:
 	UPROPERTY(Category=Sprite, EditAnywhere)
 	TArray<FPaperSpriteSocket> Sockets;
 
-public:
-	UPROPERTY(Category=Sprite, EditAnywhere)
-	TEnumAsByte<EBlendMode> BlendMode;
-
 	// Collision domain (no collision, 2D, or 3D)
 	UPROPERTY(Category=Collision, EditAnywhere)
 	TEnumAsByte<ESpriteCollisionMode::Type> SpriteCollisionDomain;
diff --git a/Engine/Plugins/Experimental/Paper2D/Source/Paper2DEditor/Private/SpriteEditor/SpriteEditorViewportClient.cpp b/Engine/Plugins/Experimental/Paper2D/Source/Paper2DEditor/Private/SpriteEditor/SpriteEditorViewportClient.cpp
--- a/Engine/Plugins/Experimental/Paper2D/Source/Paper2DEditor/Private/SpriteEditor/SpriteEditorViewportClient.cpp
+++ b/Engine/Plugins/Experimental/Paper2D/Source/Paper2DEditor/Private/SpriteEditor/SpriteEditorViewportClient.cpp
@@ -96,7 +96,6 @@ void FSpriteEditorViewportClient::UpdateSourceTextureSpriteFromSprite(UPaperSpri
 	{
 		TargetSprite->PivotMode = ESpritePivotMode::Bottom_Left;
 		TargetSprite->PixelsPerUnrealUnit = SourceSprite->PixelsPerUnrealUnit;
-		TargetSprite->BlendMode = SourceSprite->BlendMode;
 
 		if (SourceSprite->GetSourceTexture() != TargetSprite->GetSourceTexture())
 		{
