Paper2D: Hid UPaperSprite::AtlasGroup behind experimental flag and hid custom pivot value unless PivotMode=Custom

diff --git a/Engine/Plugins/Experimental/Paper2D/Source/Paper2DEditor/Private/SpriteEditor/SpriteDetailsCustomization.cpp b/Engine/Plugins/Experimental/Paper2D/Source/Paper2DEditor/Private/SpriteEditor/SpriteDetailsCustomization.cpp
--- a/Engine/Plugins/Experimental/Paper2D/Source/Paper2DEditor/Private/SpriteEditor/SpriteDetailsCustomization.cpp
+++ b/Engine/Plugins/Experimental/Paper2D/Source/Paper2DEditor/Private/SpriteEditor/SpriteDetailsCustomization.cpp
@@ -16,7 +16,7 @@ TSharedRef<IDetailCustomization> FSpriteDetailsCustomization::MakeInstance()
 void FSpriteDetailsCustomization::CustomizeDetails(IDetailLayoutBuilder& DetailLayout)
 {
 	// Make sure sprite properties are near the top
-	DetailLayout.EditCategory("Sprite", TEXT(""), ECategoryPriority::TypeSpecific);
+	IDetailCategoryBuilder& SpriteCategory = DetailLayout.EditCategory("Sprite", TEXT(""), ECategoryPriority::TypeSpecific);
 
 	TSharedRef<IPropertyHandle> RenderGeometry = DetailLayout.GetProperty(GET_MEMBER_NAME_CHECKED(UPaperSprite, RenderGeometry));
 
@@ -44,6 +44,26 @@ void FSpriteDetailsCustomization::CustomizeDetails(IDetailLayoutBuilder& DetailL
 // 	DetailLayout.HideProperty(BodySetup3D->GetChildHandle(GET_MEMBER_NAME_CHECKED(UBodySetup, BoneName)));
 // 	DetailLayout.HideProperty(BodySetup3D->GetChildHandle(GET_MEMBER_NAME_CHECKED(UBodySetup, PhysicsType)));
 // 	DetailLayout.HideProperty(BodySetup3D->GetChildHandle(GET_MEMBER_NAME_CHECKED(UBodySetup, bConsiderForBounds)));
+
+	// Show other normal properties in the sprite category so that desired ordering doesn't get messed up
+	SpriteCategory.AddProperty(GET_MEMBER_NAME_CHECKED(UPaperSprite, SourceUV));
+	SpriteCategory.AddProperty(GET_MEMBER_NAME_CHECKED(UPaperSprite, SourceDimension));
+	SpriteCategory.AddProperty(GET_MEMBER_NAME_CHECKED(UPaperSprite, SourceTexture));
+	SpriteCategory.AddProperty(GET_MEMBER_NAME_CHECKED(UPaperSprite, DefaultMaterial));
+	SpriteCategory.AddProperty(GET_MEMBER_NAME_CHECKED(UPaperSprite, Sockets));
+	SpriteCategory.AddProperty(GET_MEMBER_NAME_CHECKED(UPaperSprite, PixelsPerUnrealUnit));
+
+	// Show/hide the experimental atlas group support based on whether or not it is enabled
+	TSharedPtr<IPropertyHandle> AtlasGroupProperty = DetailLayout.GetProperty(GET_MEMBER_NAME_CHECKED(UPaperSprite, AtlasGroup));
+	TAttribute<EVisibility> AtlasGroupPropertyVisibility = TAttribute<EVisibility>::Create(TAttribute<EVisibility>::FGetter::CreateSP(this, &FSpriteDetailsCustomization::GetAtlasGroupVisibility, AtlasGroupProperty));
+	SpriteCategory.AddProperty(AtlasGroupProperty, EPropertyLocation::Advanced).Visibility(AtlasGroupPropertyVisibility);
+
+	// Show/hide the custom pivot point based on the pivot mode
+	TSharedPtr<IPropertyHandle> PivotModeProperty = DetailLayout.GetProperty(GET_MEMBER_NAME_CHECKED(UPaperSprite, PivotMode));
+	TSharedPtr<IPropertyHandle> CustomPivotPointProperty = DetailLayout.GetProperty(GET_MEMBER_NAME_CHECKED(UPaperSprite, CustomPivotPoint));
+	TAttribute<EVisibility> CustomPivotPointVisibility = TAttribute<EVisibility>::Create(TAttribute<EVisibility>::FGetter::CreateSP(this, &FSpriteDetailsCustomization::GetCustomPivotVisibility, PivotModeProperty));
+	SpriteCategory.AddProperty(PivotModeProperty);
+	SpriteCategory.AddProperty(CustomPivotPointProperty).Visibility(CustomPivotPointVisibility);
 }
 
 EVisibility FSpriteDetailsCustomization::PhysicsModeMatches(TSharedPtr<IPropertyHandle> Property, ESpriteCollisionMode::Type DesiredMode) const
@@ -79,3 +99,25 @@ EVisibility FSpriteDetailsCustomization::AnyPhysicsMode(TSharedPtr<IPropertyHand
 	// If there are multiple values, show all properties
 	return EVisibility::Visible;
 }
+
+EVisibility FSpriteDetailsCustomization::GetAtlasGroupVisibility(TSharedPtr<IPropertyHandle> Property) const
+{
+	return GetDefault<UPaperRuntimeSettings>()->bEnableSpriteAtlasGroups ? EVisibility::Visible : EVisibility::Collapsed;
+}
+
+EVisibility FSpriteDetailsCustomization::GetCustomPivotVisibility(TSharedPtr<IPropertyHandle> Property) const
+{
+	if (Property.IsValid())
+	{
+		uint8 ValueAsByte;
+		FPropertyAccess::Result Result = Property->GetValue(/*out*/ ValueAsByte);
+
+		if (Result == FPropertyAccess::Success)
+		{
+			return (((ESpritePivotMode::Type)ValueAsByte) == ESpritePivotMode::Custom) ? EVisibility::Visible : EVisibility::Collapsed;
+		}
+	}
+
+	// If there are multiple values, show all properties
+	return EVisibility::Visible;
+}
diff --git a/Engine/Plugins/Experimental/Paper2D/Source/Paper2DEditor/Private/SpriteEditor/SpriteDetailsCustomization.h b/Engine/Plugins/Experimental/Paper2D/Source/Paper2DEditor/Private/SpriteEditor/SpriteDetailsCustomization.h
--- a/Engine/Plugins/Experimental/Paper2D/Source/Paper2DEditor/Private/SpriteEditor/SpriteDetailsCustomization.h
+++ b/Engine/Plugins/Experimental/Paper2D/Source/Paper2DEditor/Private/SpriteEditor/SpriteDetailsCustomization.h
@@ -20,4 +20,6 @@ public:
 protected:
 	EVisibility PhysicsModeMatches(TSharedPtr<IPropertyHandle> Property, ESpriteCollisionMode::Type DesiredMode) const;
 	EVisibility AnyPhysicsMode(TSharedPtr<IPropertyHandle> Property) const;
+	EVisibility GetAtlasGroupVisibility(TSharedPtr<IPropertyHandle> Property) const;
+	EVisibility GetCustomPivotVisibility(TSharedPtr<IPropertyHandle> Property) const;
 };
