Paper2D: Remove the add sprite button from the sprite editor toolbar until it actually does something

diff --git a/Engine/Plugins/Experimental/Paper2D/Source/Paper2DEditor/Private/SpriteEditor/SpriteEditor.cpp b/Engine/Plugins/Experimental/Paper2D/Source/Paper2DEditor/Private/SpriteEditor/SpriteEditor.cpp
--- a/Engine/Plugins/Experimental/Paper2D/Source/Paper2DEditor/Private/SpriteEditor/SpriteEditor.cpp
+++ b/Engine/Plugins/Experimental/Paper2D/Source/Paper2DEditor/Private/SpriteEditor/SpriteEditor.cpp
@@ -492,7 +492,7 @@ void FSpriteEditor::ExtendToolbar()
 				ToolbarBuilder.AddToolBarButton(FSpriteEditorCommands::Get().EnterSourceRegionEditMode);
 				ToolbarBuilder.AddToolBarButton(FSpriteEditorCommands::Get().EnterCollisionEditMode);
 				ToolbarBuilder.AddToolBarButton(FSpriteEditorCommands::Get().EnterRenderingEditMode);
-				ToolbarBuilder.AddToolBarButton(FSpriteEditorCommands::Get().EnterAddSpriteMode);
+				//@TODO: PAPER2D: Re-enable once it does something: ToolbarBuilder.AddToolBarButton(FSpriteEditorCommands::Get().EnterAddSpriteMode);
 			}
 			ToolbarBuilder.EndSection();
 		}
