Paper2D: Fix undo problems with Paper2D due to missing RF_Transactional flag

diff --git a/Engine/Plugins/Experimental/Paper2D/Source/Paper2D/Classes/PaperFlipbook.h b/Engine/Plugins/Experimental/Paper2D/Source/Paper2D/Classes/PaperFlipbook.h
--- a/Engine/Plugins/Experimental/Paper2D/Source/Paper2D/Classes/PaperFlipbook.h
+++ b/Engine/Plugins/Experimental/Paper2D/Source/Paper2D/Classes/PaperFlipbook.h
@@ -41,4 +41,11 @@ class UPaperFlipbook : public UObject
 	PAPER2D_API int32 GetNumFrames() const;
 	PAPER2D_API float GetDuration() const;
 	PAPER2D_API UPaperSprite* GetSpriteAtTime(float Time) const;
+
+	// UObject interface
+#if WITH_EDITORONLY_DATA
+	virtual void Serialize(FArchive& Ar) override;
+	virtual void PostLoad() override;
+#endif
+	// End of UObject interface
 };
diff --git a/Engine/Plugins/Experimental/Paper2D/Source/Paper2D/Private/PaperCustomVersion.h b/Engine/Plugins/Experimental/Paper2D/Source/Paper2D/Private/PaperCustomVersion.h
--- a/Engine/Plugins/Experimental/Paper2D/Source/Paper2D/Private/PaperCustomVersion.h
+++ b/Engine/Plugins/Experimental/Paper2D/Source/Paper2D/Private/PaperCustomVersion.h
@@ -28,6 +28,9 @@ struct FPaperCustomVersion
 		// Fixed typo in convex hull generation code for 3D custom collision
 		FixTypoIn3DConvexHullCollisionGeneration = 6,
 
+		// Fixed undo issues in several asset types due to missing RF_Transactional flag
+		AddTransactionalToClasses = 7,
+
 		// -----<new versions can be added above this line>-------------------------------------------------
 		VersionPlusOne,
 		LatestVersion = VersionPlusOne - 1
diff --git a/Engine/Plugins/Experimental/Paper2D/Source/Paper2D/Private/PaperFlipbook.cpp b/Engine/Plugins/Experimental/Paper2D/Source/Paper2D/Private/PaperFlipbook.cpp
--- a/Engine/Plugins/Experimental/Paper2D/Source/Paper2D/Private/PaperFlipbook.cpp
+++ b/Engine/Plugins/Experimental/Paper2D/Source/Paper2D/Private/PaperFlipbook.cpp
@@ -2,6 +2,7 @@
 
 #include "Paper2DPrivatePCH.h"
 #include "PaperFlipbook.h"
+#include "PaperCustomVersion.h"
 
 //////////////////////////////////////////////////////////////////////////
 // UPaperFlipbook
@@ -51,4 +52,24 @@ UPaperSprite* UPaperFlipbook::GetSpriteAtTime(float Time) const
 	}
 
 	return NULL;
-}
\ No newline at end of file
+}
+
+#if WITH_EDITORONLY_DATA
+void UPaperFlipbook::Serialize(FArchive& Ar)
+{
+	Super::Serialize(Ar);
+
+	Ar.UsingCustomVersion(FPaperCustomVersion::GUID);
+}
+
+void UPaperFlipbook::PostLoad()
+{
+	Super::PostLoad();
+
+	const int32 PaperVer = GetLinkerCustomVersion(FPaperCustomVersion::GUID);
+	if (PaperVer < FPaperCustomVersion::AddTransactionalToClasses)
+	{
+		SetFlags(RF_Transactional);
+	}
+}
+#endif
\ No newline at end of file
diff --git a/Engine/Plugins/Experimental/Paper2D/Source/Paper2D/Private/PaperSprite.cpp b/Engine/Plugins/Experimental/Paper2D/Source/Paper2D/Private/PaperSprite.cpp
--- a/Engine/Plugins/Experimental/Paper2D/Source/Paper2D/Private/PaperSprite.cpp
+++ b/Engine/Plugins/Experimental/Paper2D/Source/Paper2D/Private/PaperSprite.cpp
@@ -1093,6 +1093,11 @@ void UPaperSprite::PostLoad()
 	bool bRebuildCollision = false;
 	bool bRebuildRenderData = false;
 
+	if (PaperVer < FPaperCustomVersion::AddTransactionalToClasses)
+	{
+		SetFlags(RF_Transactional);
+	}
+
 	if (PaperVer < FPaperCustomVersion::AddPixelsPerUnrealUnit)
 	{
 		PixelsPerUnrealUnit = 1.0f;
diff --git a/Engine/Plugins/Experimental/Paper2D/Source/Paper2D/Private/PaperTileMapRenderComponent.cpp b/Engine/Plugins/Experimental/Paper2D/Source/Paper2D/Private/PaperTileMapRenderComponent.cpp
--- a/Engine/Plugins/Experimental/Paper2D/Source/Paper2D/Private/PaperTileMapRenderComponent.cpp
+++ b/Engine/Plugins/Experimental/Paper2D/Source/Paper2D/Private/PaperTileMapRenderComponent.cpp
@@ -32,6 +32,7 @@ FPrimitiveSceneProxy* UPaperTileMapRenderComponent::CreateSceneProxy()
 void UPaperTileMapRenderComponent::PostInitProperties()
 {
 	TileMap = NewObject<UPaperTileMap>(this);
+	TileMap->SetFlags(RF_Transactional);
 	Super::PostInitProperties();
 }
 
@@ -92,6 +93,7 @@ void UPaperTileMapRenderComponent::PostLoad()
 	{
 		// Create a tile map object and move our old properties over to it
 		TileMap = NewObject<UPaperTileMap>(this);
+		TileMap->SetFlags(RF_Transactional);
 		TileMap->MapWidth = MapWidth_DEPRECATED;
 		TileMap->MapHeight = MapHeight_DEPRECATED;
 		TileMap->TileWidth = TileWidth_DEPRECATED;
diff --git a/Engine/Plugins/Experimental/Paper2D/Source/Paper2DEditor/Private/Atlasing/PaperSpriteAtlasFactory.cpp b/Engine/Plugins/Experimental/Paper2D/Source/Paper2DEditor/Private/Atlasing/PaperSpriteAtlasFactory.cpp
--- a/Engine/Plugins/Experimental/Paper2D/Source/Paper2DEditor/Private/Atlasing/PaperSpriteAtlasFactory.cpp
+++ b/Engine/Plugins/Experimental/Paper2D/Source/Paper2DEditor/Private/Atlasing/PaperSpriteAtlasFactory.cpp
@@ -18,7 +18,7 @@ UPaperSpriteAtlasFactory::UPaperSpriteAtlasFactory(const class FPostConstructIni
 
 UObject* UPaperSpriteAtlasFactory::FactoryCreateNew(UClass* InClass, UObject* InParent, FName InName, EObjectFlags Flags, UObject* Context, FFeedbackContext* Warn)
 {
-	return ConstructObject<UObject>(SupportedClass, InParent, InName, Flags);
+	return ConstructObject<UObject>(SupportedClass, InParent, InName, Flags | RF_Transactional);
 }
 
 bool UPaperSpriteAtlasFactory::CanCreateNew() const
diff --git a/Engine/Plugins/Experimental/Paper2D/Source/Paper2DEditor/Private/PaperFlipbookFactory.cpp b/Engine/Plugins/Experimental/Paper2D/Source/Paper2DEditor/Private/PaperFlipbookFactory.cpp
--- a/Engine/Plugins/Experimental/Paper2D/Source/Paper2DEditor/Private/PaperFlipbookFactory.cpp
+++ b/Engine/Plugins/Experimental/Paper2D/Source/Paper2DEditor/Private/PaperFlipbookFactory.cpp
@@ -17,7 +17,7 @@ UPaperFlipbookFactory::UPaperFlipbookFactory(const class FPostConstructInitializ
 
 UObject* UPaperFlipbookFactory::FactoryCreateNew(UClass* Class, UObject* InParent, FName Name, EObjectFlags Flags, UObject* Context, FFeedbackContext* Warn)
 {
-	UPaperFlipbook* NewFlipbook = ConstructObject<UPaperFlipbook>(Class, InParent, Name, Flags);
+	UPaperFlipbook* NewFlipbook = ConstructObject<UPaperFlipbook>(Class, InParent, Name, Flags | RF_Transactional);
 
 	return NewFlipbook;
 }
diff --git a/Engine/Plugins/Experimental/Paper2D/Source/Paper2DEditor/Private/PaperSpriteFactory.cpp b/Engine/Plugins/Experimental/Paper2D/Source/Paper2DEditor/Private/PaperSpriteFactory.cpp
--- a/Engine/Plugins/Experimental/Paper2D/Source/Paper2DEditor/Private/PaperSpriteFactory.cpp
+++ b/Engine/Plugins/Experimental/Paper2D/Source/Paper2DEditor/Private/PaperSpriteFactory.cpp
@@ -26,7 +26,7 @@ bool UPaperSpriteFactory::ConfigureProperties()
 
 UObject* UPaperSpriteFactory::FactoryCreateNew(UClass* Class, UObject* InParent, FName Name, EObjectFlags Flags, UObject* Context, FFeedbackContext* Warn)
 {
-	UPaperSprite* NewSprite = ConstructObject<UPaperSprite>(Class, InParent, Name, Flags);
+	UPaperSprite* NewSprite = ConstructObject<UPaperSprite>(Class, InParent, Name, Flags | RF_Transactional);
 
 	if (bUseSourceRegion)
 	{
diff --git a/Engine/Plugins/Experimental/Paper2D/Source/PaperJsonImporter/Private/PaperJsonImporterFactory.cpp b/Engine/Plugins/Experimental/Paper2D/Source/PaperJsonImporter/Private/PaperJsonImporterFactory.cpp
--- a/Engine/Plugins/Experimental/Paper2D/Source/PaperJsonImporter/Private/PaperJsonImporterFactory.cpp
+++ b/Engine/Plugins/Experimental/Paper2D/Source/PaperJsonImporter/Private/PaperJsonImporterFactory.cpp
@@ -48,6 +48,8 @@ FText UPaperJsonImporterFactory::GetToolTip() const
 
 UObject* UPaperJsonImporterFactory::FactoryCreateText(UClass* InClass, UObject* InParent, FName InName, EObjectFlags Flags, UObject* Context, const TCHAR* Type, const TCHAR*& Buffer, const TCHAR* BufferEnd, FFeedbackContext* Warn)
 {
+	Flags |= RF_Transactional;
+
 	FEditorDelegates::OnAssetPreImport.Broadcast(this, InClass, InParent, InName, Type);
 
 	FAssetToolsModule& AssetToolsModule = FModuleManager::GetModuleChecked<FAssetToolsModule>("AssetTools");
