Paper2D: Fix intermittent corruption issue when converting old sprites on load due to async texture data creation by blocking

diff --git a/Engine/Plugins/Experimental/Paper2D/Source/Paper2D/Private/PaperSprite.cpp b/Engine/Plugins/Experimental/Paper2D/Source/Paper2D/Private/PaperSprite.cpp
--- a/Engine/Plugins/Experimental/Paper2D/Source/Paper2D/Private/PaperSprite.cpp
+++ b/Engine/Plugins/Experimental/Paper2D/Source/Paper2D/Private/PaperSprite.cpp
@@ -556,12 +556,22 @@ void UPaperSprite::RebuildRenderData()
 	TArray<FVector2D> TriangluatedPoints;
 	Triangulate(RenderGeometry, /*out*/ TriangluatedPoints);
 
-	// Adjust for the pivot and store in the baked geometry buffer
+	// Determine the texture size
 	UTexture2D* EffectiveTexture = GetBakedTexture();
 
-	const float InverseWidth = (EffectiveTexture != nullptr) ? (1.0f / EffectiveTexture->GetSizeX()) : 1.0f;
-	const float InverseHeight = (EffectiveTexture != nullptr) ? (1.0f / EffectiveTexture->GetSizeY()) : 1.0f;
+	FVector2D TextureSize(1.0f, 1.0f);
+	if (EffectiveTexture)
+	{
+		EffectiveTexture->FinishCachePlatformData();
+		const int32 TextureWidth = EffectiveTexture->GetSizeX();
+		const int32 TextureHeight = EffectiveTexture->GetSizeY();
+		check((TextureWidth > 0) && (TextureHeight > 0));
+		TextureSize = FVector2D(TextureWidth, TextureHeight);
+	}
+	const float InverseWidth = 1.0f / TextureSize.X;
+	const float InverseHeight = 1.0f / TextureSize.Y;
 
+	// Adjust for the pivot and store in the baked geometry buffer
 	const FVector2D DeltaUV((BakedSourceTexture != nullptr) ? (BakedSourceUV - SourceUV) : FVector2D::ZeroVector);
 
 	const float UnitsPerPixel = 1.0f / PixelsPerUnrealUnit;
@@ -570,6 +580,7 @@ void UPaperSprite::RebuildRenderData()
 	for (int32 PointIndex = 0; PointIndex < TriangluatedPoints.Num(); ++PointIndex)
 	{
 		const FVector2D& SourcePos = TriangluatedPoints[PointIndex];
+
 		const FVector2D PivotSpacePos = ConvertTextureSpaceToPivotSpace(SourcePos);
 		const FVector2D UV(SourcePos + DeltaUV);
 
diff --git a/Engine/Source/Runtime/Engine/Classes/Engine/Texture.h b/Engine/Source/Runtime/Engine/Classes/Engine/Texture.h
--- a/Engine/Source/Runtime/Engine/Classes/Engine/Texture.h
+++ b/Engine/Source/Runtime/Engine/Classes/Engine/Texture.h
@@ -602,7 +602,7 @@ public:
 	/**
 	 * Blocks on async cache tasks and prepares platform data for use.
 	 */
-	void FinishCachePlatformData();
+	ENGINE_API void FinishCachePlatformData();
 
 	/**
 	 * Forces platform data to be rebuilt.
