Paper2D:  Add more options to InitializeSprite

diff --git a/Engine/Plugins/Experimental/Paper2D/Source/Paper2D/Classes/PaperSprite.h b/Engine/Plugins/Experimental/Paper2D/Source/Paper2D/Classes/PaperSprite.h
--- a/Engine/Plugins/Experimental/Paper2D/Source/Paper2D/Classes/PaperSprite.h
+++ b/Engine/Plugins/Experimental/Paper2D/Source/Paper2D/Classes/PaperSprite.h
@@ -165,8 +165,10 @@ public:
 	void Triangulate(const FSpritePolygonCollection& Source, TArray<FVector2D>& Target);
 
 	// Reinitializes this sprite (NOTE: Does not register existing components in the world)
+	void InitializeSprite(UTexture2D* Texture, float InPixelsPerUnrealUnit);
 	void InitializeSprite(UTexture2D* Texture);
 	void InitializeSprite(UTexture2D* Texture, const FVector2D& Offset, const FVector2D& Dimension);
+	void InitializeSprite(UTexture2D* Texture, const FVector2D& Offset, const FVector2D& Dimension, float InPixelsPerUnit);
 
 	FVector2D GetSourceUV() const { return SourceUV; }
 	FVector2D GetSourceSize() const { return SourceDimension; }
diff --git a/Engine/Plugins/Experimental/Paper2D/Source/Paper2D/Private/PaperSprite.cpp b/Engine/Plugins/Experimental/Paper2D/Source/Paper2D/Private/PaperSprite.cpp
--- a/Engine/Plugins/Experimental/Paper2D/Source/Paper2D/Private/PaperSprite.cpp
+++ b/Engine/Plugins/Experimental/Paper2D/Source/Paper2D/Private/PaperSprite.cpp
@@ -897,23 +897,27 @@ void UPaperSprite::ExtractRectsFromTexture(UTexture2D* Texture, TArray<FIntRect>
 	SpriteTextureBitmap.ExtractRects(/*out*/ OutRects);
 }
 
-void UPaperSprite::InitializeSprite(UTexture2D* Texture)
+void UPaperSprite::InitializeSprite(UTexture2D* Texture, float InPixelsPerUnrealUnit)
 {
 	if (Texture != NULL)
 	{
 		const FVector2D Dimension(Texture->GetSizeX(), Texture->GetSizeY());
-		InitializeSprite(Texture, FVector2D::ZeroVector, Dimension);
+		InitializeSprite(Texture, FVector2D::ZeroVector, Dimension, InPixelsPerUnrealUnit);
 	}
 	else
 	{
-		InitializeSprite(NULL, FVector2D::ZeroVector, FVector2D(1.0f, 1.0f));
+		InitializeSprite(NULL, FVector2D::ZeroVector, FVector2D(1.0f, 1.0f), InPixelsPerUnrealUnit);
 	}
+}
 
-	PixelsPerUnrealUnit = GetDefault<UPaperRuntimeSettings>()->DefaultPixelsPerUnrealUnit;
+void UPaperSprite::InitializeSprite(UTexture2D* Texture)
+{
+	InitializeSprite(Texture, GetDefault<UPaperRuntimeSettings>()->DefaultPixelsPerUnrealUnit);
 }
 
-void UPaperSprite::InitializeSprite(UTexture2D* Texture, const FVector2D& Offset, const FVector2D& Dimension)
+void UPaperSprite::InitializeSprite(UTexture2D* Texture, const FVector2D& Offset, const FVector2D& Dimension, float InPixelsPerUnrealUnit)
 {
+	PixelsPerUnrealUnit = InPixelsPerUnrealUnit;
 	SourceTexture = Texture;
 	SourceUV = Offset;
 	SourceDimension = Dimension;
@@ -921,6 +925,11 @@ void UPaperSprite::InitializeSprite(UTexture2D* Texture, const FVector2D& Offset
 	RebuildRenderData();
 }
 
+void UPaperSprite::InitializeSprite(UTexture2D* Texture, const FVector2D& Offset, const FVector2D& Dimension)
+{
+	InitializeSprite(Texture, Offset, Dimension, GetDefault<UPaperRuntimeSettings>()->DefaultPixelsPerUnrealUnit);
+}
+
 FVector2D UPaperSprite::ConvertTextureSpaceToPivotSpace(FVector2D Input) const
 {
 	const FVector2D Pivot = GetPivotPosition();
