Paper2D: Add an 'early access preview' note to plugin description

diff --git a/Engine/Plugins/Experimental/Paper2D/Paper2D.uplugin b/Engine/Plugins/Experimental/Paper2D/Paper2D.uplugin
--- a/Engine/Plugins/Experimental/Paper2D/Paper2D.uplugin
+++ b/Engine/Plugins/Experimental/Paper2D/Paper2D.uplugin
@@ -7,7 +7,7 @@
 	"CreatedBy" : "Epic Games, Inc.",
 	"CreatedByURL" : "http://epicgames.com",
 	"EngineVersion" : "4.2.0",
-	"Description" : "Paper2D adds tools and assets to help create 2D games including animated sprite assets, tilesets, 2D level editing tools, and more.",
+	"Description" : "[EARLY ACCESS PREVIEW] Paper2D adds tools and assets to help create 2D games including animated sprite assets, tilesets, 2D level editing tools, and more.",
 	"Category" : "2D.Helpers",
 	"EnabledByDefault" : true,
 
