Paper2D: Fix several issues when pixels per uu is not 1.0

diff --git a/Engine/Plugins/Experimental/Paper2D/Source/Paper2D/Classes/PaperSprite.h b/Engine/Plugins/Experimental/Paper2D/Source/Paper2D/Classes/PaperSprite.h
--- a/Engine/Plugins/Experimental/Paper2D/Source/Paper2D/Classes/PaperSprite.h
+++ b/Engine/Plugins/Experimental/Paper2D/Source/Paper2D/Classes/PaperSprite.h
@@ -174,6 +174,7 @@ public:
 	FVector2D GetSourceSize() const { return SourceDimension; }
 	UTexture2D* GetSourceTexture() const { return SourceTexture; }
 	float GetPixelsPerUnrealUnit() const { return PixelsPerUnrealUnit; }
+	float GetUnrealUnitsPerPixel() const { return 1.0f / PixelsPerUnrealUnit; }
 #endif
 
 	// Returns the texture this should be rendered with
diff --git a/Engine/Plugins/Experimental/Paper2D/Source/Paper2D/Private/PaperRenderComponent.cpp b/Engine/Plugins/Experimental/Paper2D/Source/Paper2D/Private/PaperRenderComponent.cpp
--- a/Engine/Plugins/Experimental/Paper2D/Source/Paper2D/Private/PaperRenderComponent.cpp
+++ b/Engine/Plugins/Experimental/Paper2D/Source/Paper2D/Private/PaperRenderComponent.cpp
@@ -100,21 +100,24 @@ FTransform UPaperRenderComponent::GetSocketTransform(FName InSocketName, ERelati
 	{
 		if (FPaperSpriteSocket* Socket = SourceSprite->FindSocket(InSocketName))
 		{
+			FTransform SocketLocalTransform = Socket->LocalTransform;
+			SocketLocalTransform.ScaleTranslation(SourceSprite->GetUnrealUnitsPerPixel());
+
 			switch (TransformSpace)
 			{
 				case RTS_World:
-					return Socket->LocalTransform * ComponentToWorld;
+					return SocketLocalTransform * ComponentToWorld;
 
 				case RTS_Actor:
 					if (const AActor* Actor = GetOwner())
 					{
-						const FTransform SocketTransform = Socket->LocalTransform * ComponentToWorld;
+						const FTransform SocketTransform = SocketLocalTransform * ComponentToWorld;
 						return SocketTransform.GetRelativeTransform(Actor->GetTransform());
 					}
 					break;
 
 				case RTS_Component:
-					return Socket->LocalTransform;
+					return SocketLocalTransform;
 
 				default:
 					check(false);
diff --git a/Engine/Plugins/Experimental/Paper2D/Source/Paper2D/Private/PaperSprite.cpp b/Engine/Plugins/Experimental/Paper2D/Source/Paper2D/Private/PaperSprite.cpp
--- a/Engine/Plugins/Experimental/Paper2D/Source/Paper2D/Private/PaperSprite.cpp
+++ b/Engine/Plugins/Experimental/Paper2D/Source/Paper2D/Private/PaperSprite.cpp
@@ -271,6 +271,11 @@ UPaperSprite::UPaperSprite(const FPostConstructInitializeProperties& PCIP)
 #if WITH_EDITORONLY_DATA
 void UPaperSprite::PostEditChangeProperty(FPropertyChangedEvent& PropertyChangedEvent)
 {
+	if (PixelsPerUnrealUnit <= 0.0f)
+	{
+		PixelsPerUnrealUnit = 1.0f;
+	}
+
 	//@TODO: Determine when these are really needed!
 	TComponentReregisterContext<UPaperRenderComponent> ReregisterStaticComponents;
 	TComponentReregisterContext<UPaperAnimatedRenderComponent> ReregisterAnimatedComponents;
@@ -391,7 +396,7 @@ void UPaperSprite::BuildCustomCollisionData()
 
 
 	// Adjust the collision data to be relative to the pivot and scaled from pixels to uu
-	const float UnitsPerPixel = 1.0f / PixelsPerUnrealUnit;
+	const float UnitsPerPixel = GetUnrealUnitsPerPixel();
 	for (FVector2D& Point : CollisionData)
 	{
 		Point = ConvertTextureSpaceToPivotSpace(Point) * UnitsPerPixel;
@@ -466,7 +471,7 @@ void UPaperSprite::BuildBoundingBoxCollisionData(bool bUseTightBounds)
 	CreatePolygonFromBoundingBox(CollisionGeometry, bUseTightBounds);
 
 	// Bake it to the runtime structure
-	const float UnitsPerPixel = 1.0f / PixelsPerUnrealUnit;
+	const float UnitsPerPixel = GetUnrealUnitsPerPixel();
 
 	switch (SpriteCollisionDomain)
 	{
@@ -574,7 +579,7 @@ void UPaperSprite::RebuildRenderData()
 	// Adjust for the pivot and store in the baked geometry buffer
 	const FVector2D DeltaUV((BakedSourceTexture != nullptr) ? (BakedSourceUV - SourceUV) : FVector2D::ZeroVector);
 
-	const float UnitsPerPixel = 1.0f / PixelsPerUnrealUnit;
+	const float UnitsPerPixel = GetUnrealUnitsPerPixel();
 
 	BakedRenderData.Empty(TriangluatedPoints.Num());
 	for (int32 PointIndex = 0; PointIndex < TriangluatedPoints.Num(); ++PointIndex)
@@ -972,7 +977,7 @@ FVector UPaperSprite::ConvertPivotSpaceToTextureSpace(FVector Input) const
 
 FVector UPaperSprite::ConvertTextureSpaceToWorldSpace(const FVector2D& SourcePoint) const
 {
-	const float UnitsPerPixel = 1.0f / PixelsPerUnrealUnit;
+	const float UnitsPerPixel = GetUnrealUnitsPerPixel();
 
 	const FVector2D SourcePointInUU = SourcePoint * UnitsPerPixel;
 
@@ -995,7 +1000,7 @@ FTransform UPaperSprite::GetPivotToWorld() const
 	const FVector2D Pivot = GetPivotPosition();
 	const FVector2D SourceDims = (SourceTexture != NULL) ? FVector2D(SourceTexture->GetSurfaceWidth(), SourceTexture->GetSurfaceHeight()) : FVector2D::ZeroVector;
 
-	const FVector Translation((Pivot.X * PaperAxisX) + ((SourceDims.Y - Pivot.Y) * PaperAxisY));
+	const FVector Translation(((Pivot.X * PaperAxisX) + ((SourceDims.Y - Pivot.Y) * PaperAxisY)) * GetUnrealUnitsPerPixel());
 	return FTransform(Translation);
 }
 
diff --git a/Engine/Plugins/Experimental/Paper2D/Source/Paper2D/Private/Terrain/PaperTerrainComponent.cpp b/Engine/Plugins/Experimental/Paper2D/Source/Paper2D/Private/Terrain/PaperTerrainComponent.cpp
--- a/Engine/Plugins/Experimental/Paper2D/Source/Paper2D/Private/Terrain/PaperTerrainComponent.cpp
+++ b/Engine/Plugins/Experimental/Paper2D/Source/Paper2D/Private/Terrain/PaperTerrainComponent.cpp
@@ -569,7 +569,7 @@ void UPaperTerrainComponent::SpawnFromPoly(const class UPaperSprite* NewSprite,
 {
 	//@TODO: Need to split geom tools out of UnrealEd to support this outside of the editor
 #if WITH_EDITOR
-	const FVector2D TextureSizeInUnits = TextureSize / NewSprite->GetPixelsPerUnrealUnit();
+	const FVector2D TextureSizeInUnits = TextureSize * NewSprite->GetUnrealUnitsPerPixel();
 
 	if (Poly.Vertices.Num() >= 3)
 	{
diff --git a/Engine/Plugins/Experimental/Paper2D/Source/Paper2DEditor/Private/SpriteEditor/SpriteEditorSelections.h b/Engine/Plugins/Experimental/Paper2D/Source/Paper2DEditor/Private/SpriteEditor/SpriteEditorSelections.h
--- a/Engine/Plugins/Experimental/Paper2D/Source/Paper2DEditor/Private/SpriteEditor/SpriteEditorSelections.h
+++ b/Engine/Plugins/Experimental/Paper2D/Source/Paper2DEditor/Private/SpriteEditor/SpriteEditorSelections.h
@@ -221,7 +221,9 @@ public:
 			default:
 				break;
 			}
-			Result = BoundsVertex.X * PaperAxisX + (SourceDims.Y - BoundsVertex.Y) * PaperAxisY;
+
+			const FVector PixelSpaceResult = BoundsVertex.X * PaperAxisX + (SourceDims.Y - BoundsVertex.Y) * PaperAxisY;
+			Result = PixelSpaceResult * Sprite->GetUnrealUnitsPerPixel();
 		}
 
 		return Result;
@@ -328,7 +330,8 @@ public:
 					const FVector2D SourceDims = (SourceTexture != NULL) ? FVector2D(SourceTexture->GetSurfaceWidth(), SourceTexture->GetSurfaceHeight()) : FVector2D::ZeroVector;
 
 					const FVector2D Result2D = Vertices[TargetVertexIndex];
-					Result = (Result2D.X * PaperAxisX) + ((SourceDims.Y - Result2D.Y) * PaperAxisY);
+					const FVector PixelSpaceResult = (Result2D.X * PaperAxisX) + ((SourceDims.Y - Result2D.Y) * PaperAxisY);
+					Result = PixelSpaceResult * Sprite->GetUnrealUnitsPerPixel();
 				}
 			}
 		}
diff --git a/Engine/Plugins/Experimental/Paper2D/Source/Paper2DEditor/Private/SpriteEditor/SpriteEditorViewportClient.cpp b/Engine/Plugins/Experimental/Paper2D/Source/Paper2DEditor/Private/SpriteEditor/SpriteEditorViewportClient.cpp
--- a/Engine/Plugins/Experimental/Paper2D/Source/Paper2DEditor/Private/SpriteEditor/SpriteEditorViewportClient.cpp
+++ b/Engine/Plugins/Experimental/Paper2D/Source/Paper2DEditor/Private/SpriteEditor/SpriteEditorViewportClient.cpp
@@ -76,6 +76,9 @@ FSpriteEditorViewportClient::FSpriteEditorViewportClient(TWeakPtr<FSpriteEditor>
 	// Create a sprite and render component for the source texture view
 	{
 		UPaperSprite* DummySprite = NewObject<UPaperSprite>();
+		DummySprite->SpriteCollisionDomain = ESpriteCollisionMode::None;
+		DummySprite->PivotMode = ESpritePivotMode::Bottom_Left;
+
 		SourceTextureViewComponent = NewObject<UPaperRenderComponent>();
 		SourceTextureViewComponent->SetSprite(DummySprite);
 		UpdateSourceTextureSpriteFromSprite(GetSpriteBeingEdited());
@@ -94,12 +97,9 @@ void FSpriteEditorViewportClient::UpdateSourceTextureSpriteFromSprite(UPaperSpri
 
 	if (SourceSprite != NULL)
 	{
-		TargetSprite->PivotMode = ESpritePivotMode::Bottom_Left;
-		TargetSprite->PixelsPerUnrealUnit = SourceSprite->PixelsPerUnrealUnit;
-
-		if (SourceSprite->GetSourceTexture() != TargetSprite->GetSourceTexture())
+		if ((SourceSprite->GetSourceTexture() != TargetSprite->GetSourceTexture()) || (TargetSprite->PixelsPerUnrealUnit != SourceSprite->PixelsPerUnrealUnit))
 		{
-			TargetSprite->InitializeSprite(SourceSprite->SourceTexture);
+			TargetSprite->InitializeSprite(SourceSprite->SourceTexture, SourceSprite->PixelsPerUnrealUnit);
 		}
 	}
 	else
@@ -741,9 +741,10 @@ void FSpriteEditorViewportClient::NotifySpriteBeingEditedHasChanged()
 	// Update components to know about the new sprite being edited
 	UPaperSprite* Sprite = GetSpriteBeingEdited();
 
-	SourceTextureViewComponent->SetSprite(Sprite);
 	RenderSpriteComponent->SetSprite(Sprite);
+	UpdateSourceTextureSpriteFromSprite(Sprite);
 
+	//
 	bDeferZoomToSprite = true;
 }
 
