Paper2D: Add temporary option to avoid collapsing unnecessary vertices when triangulating a sprite's geometry (useful for adding extra verts to improve terrain spline deformation)

diff --git a/Engine/Plugins/Experimental/Paper2D/Source/Paper2D/Classes/SpriteEditorOnlyTypes.h b/Engine/Plugins/Experimental/Paper2D/Source/Paper2D/Classes/SpriteEditorOnlyTypes.h
--- a/Engine/Plugins/Experimental/Paper2D/Source/Paper2D/Classes/SpriteEditorOnlyTypes.h
+++ b/Engine/Plugins/Experimental/Paper2D/Source/Paper2D/Classes/SpriteEditorOnlyTypes.h
@@ -56,23 +56,23 @@ struct FSpritePolygonCollection
 	UPROPERTY(Category=PolygonData, EditAnywhere)
 	TEnumAsByte<ESpritePolygonMode::Type> GeometryType;
 
+	// Experimental: Hint to the triangulation routine that extra vertices should be preserved
+	UPROPERTY(Category=PolygonData, EditAnywhere, AdvancedDisplay)
+	bool bAvoidVertexMerging;
+
 	// Alpha threshold for a transparent pixel (range 0..1, anything equal or below this value will be considered unimportant)
 	UPROPERTY(Category=PolygonData, EditAnywhere, AdvancedDisplay)
 	float AlphaThreshold;
 
-	//@TODO: Document
+	// This is the threshold below which multiple vertices will be merged together when doing shrink-wrapping.  Higher values result in fewer vertices.
 	UPROPERTY(Category=PolygonData, EditAnywhere, AdvancedDisplay)
 	float SimplifyEpsilon;
 
-	// Number of vertices to include as part of the outline (only used when Mode is ShrinkWrapped)
-// 	UPROPERTY(Category=PolygonData, EditAnywhere)
-// 	int32 VertexCount;
-
 	FSpritePolygonCollection()
 		: GeometryType(ESpritePolygonMode::TightBoundingBox)
+		, bAvoidVertexMerging(false)
 		, AlphaThreshold(0.0f)
 		, SimplifyEpsilon(2.0f)
-		//, VertexCount(20)
 	{
 	}
 
diff --git a/Engine/Plugins/Experimental/Paper2D/Source/Paper2D/Private/PaperSprite.cpp b/Engine/Plugins/Experimental/Paper2D/Source/Paper2D/Private/PaperSprite.cpp
--- a/Engine/Plugins/Experimental/Paper2D/Source/Paper2D/Private/PaperSprite.cpp
+++ b/Engine/Plugins/Experimental/Paper2D/Source/Paper2D/Private/PaperSprite.cpp
@@ -880,7 +880,7 @@ void UPaperSprite::Triangulate(const FSpritePolygonCollection& Source, TArray<FV
 
 			// Attempt to triangulate this polygon
 			TArray<FClipSMTriangle> GeneratedTriangles;
-			if (TriangulatePoly(/*out*/ GeneratedTriangles, ClipPolygon))
+			if (TriangulatePoly(/*out*/ GeneratedTriangles, ClipPolygon, Source.bAvoidVertexMerging))
 			{
 				// Convert the triangles back to our 2D data structure
 				for (int32 TriangleIndex = 0; TriangleIndex < GeneratedTriangles.Num(); ++TriangleIndex)
