Paper2D: Add a controls help text string for each mode in the sprite editor

diff --git a/Engine/Plugins/Experimental/Paper2D/Source/Paper2DEditor/Private/SpriteEditor/SpriteEditorViewportClient.cpp b/Engine/Plugins/Experimental/Paper2D/Source/Paper2DEditor/Private/SpriteEditor/SpriteEditorViewportClient.cpp
--- a/Engine/Plugins/Experimental/Paper2D/Source/Paper2DEditor/Private/SpriteEditor/SpriteEditorViewportClient.cpp
+++ b/Engine/Plugins/Experimental/Paper2D/Source/Paper2DEditor/Private/SpriteEditor/SpriteEditorViewportClient.cpp
@@ -439,6 +439,9 @@ void FSpriteEditorViewportClient::DrawCanvas(FViewport& Viewport, FSceneView& Vi
 
 	int32 YPos = 6;
 
+	static const FText GeomHelpStr = LOCTEXT("GeomEditHelp", "Select an edge and press Insert to add a vertex.\nSelect an edge or vertex and press Delete to remove it.\n");
+	static const FText SourceRegionHelpStr = LOCTEXT("SourceRegionHelp", "Drag handles to adjust source region\nDouble-click on an image region to select all connected pixels");
+
 	switch (CurrentMode)
 	{
 	case ESpriteEditorMode::ViewMode:
@@ -446,7 +449,9 @@ void FSpriteEditorViewportClient::DrawCanvas(FViewport& Viewport, FSceneView& Vi
 
 		// Display the pivot
 		{
-			const FText PivotStr = FText::Format(LOCTEXT("PivotPosition", "Pivot: {0}, {1}"), FText::AsNumber(Sprite->CustomPivotPoint.X), FText::AsNumber(Sprite->CustomPivotPoint.Y));
+			FNumberFormattingOptions NoDigitGroupingFormat;
+			NoDigitGroupingFormat.UseGrouping = false;
+			const FText PivotStr = FText::Format(LOCTEXT("PivotPosition", "Pivot: ({0}, {1})"), FText::AsNumber(Sprite->CustomPivotPoint.X, &NoDigitGroupingFormat), FText::AsNumber(Sprite->CustomPivotPoint.Y, &NoDigitGroupingFormat));
 			FCanvasTextItem TextItem(FVector2D(6, YPos), PivotStr, GEngine->GetSmallFont(), FLinearColor::White);
 			TextItem.EnableShadow(FLinearColor::Black);
 			TextItem.Draw(&Canvas);
@@ -459,6 +464,14 @@ void FSpriteEditorViewportClient::DrawCanvas(FViewport& Viewport, FSceneView& Vi
 		break;
 	case ESpriteEditorMode::EditCollisionMode:
 		{
+ 			// Display tool help
+			{
+				FCanvasTextItem TextItem(FVector2D(6, YPos), GeomHelpStr, GEngine->GetSmallFont(), FLinearColor::White);
+				TextItem.EnableShadow(FLinearColor::Black);
+				TextItem.Draw(&Canvas);
+				YPos += 36;
+			}
+
 			// Draw the custom collision geometry
 			const FLinearColor CollisionColor(1.0f, 1.0f, 0.0f, 1.0f);
 			DrawGeometry(Viewport, View, Canvas, Sprite->CollisionGeometry, CollisionColor, false);
@@ -467,6 +480,14 @@ void FSpriteEditorViewportClient::DrawCanvas(FViewport& Viewport, FSceneView& Vi
 		break;
 	case ESpriteEditorMode::EditRenderingGeomMode:
 		{
+ 			// Display tool help
+			{
+				FCanvasTextItem TextItem(FVector2D(6, YPos), GeomHelpStr, GEngine->GetSmallFont(), FLinearColor::White);
+				TextItem.EnableShadow(FLinearColor::Black);
+				TextItem.Draw(&Canvas);
+				YPos += 36;
+			}
+
 			// Draw the custom render geometry
 			const FLinearColor RenderGeomColor(1.0f, 0.2f, 0.0f, 1.0f);
 			DrawGeometry(Viewport, View, Canvas, Sprite->RenderGeometry, RenderGeomColor, true);
@@ -475,6 +496,14 @@ void FSpriteEditorViewportClient::DrawCanvas(FViewport& Viewport, FSceneView& Vi
 		break;
 	case ESpriteEditorMode::EditSourceRegionMode:
 		{
+			// Display tool help
+			{
+				FCanvasTextItem TextItem(FVector2D(6, YPos), SourceRegionHelpStr, GEngine->GetSmallFont(), FLinearColor::White);
+				TextItem.EnableShadow(FLinearColor::Black);
+				TextItem.Draw(&Canvas);
+				YPos += 18;
+			}
+
 			// Draw the custom render geometry
 			const FLinearColor BoundsColor(1.0f, 1.0f, 1.0f, 0.8f);
 			DrawSourceRegion(Viewport, View, Canvas, BoundsColor, true);
