Paper2D: Add note to PaperAxis values

diff --git a/Engine/Plugins/Experimental/Paper2D/Source/Paper2D/Public/Paper2DModule.h b/Engine/Plugins/Experimental/Paper2D/Source/Paper2D/Public/Paper2DModule.h
--- a/Engine/Plugins/Experimental/Paper2D/Source/Paper2D/Public/Paper2DModule.h
+++ b/Engine/Plugins/Experimental/Paper2D/Source/Paper2D/Public/Paper2DModule.h
@@ -11,6 +11,10 @@ class IPaper2DModuleInterface : public IModuleInterface
 {
 };
 
+// NOTE: Experimental: Do not adjust these values, many aspects of the editor or runtime assume PaperX=Right(+X) PaperY=Up(+Z) PaperZ=PaperX cross PaperY=Out(-Y)
+// The ability to adjust these values may be removed entirely in the future, or made to work properly
+//@TODO: PAPER2D: Synchronize with the Box2D axes in the engine (or remove the ability to customize them)
+
 // The 3D vector that corresponds to the sprite X axis
 extern PAPER2D_API FVector PaperAxisX;
 
