Paper2D: Hook up new sprite editor toolbar icons

diff --git a/Engine/Plugins/Experimental/Paper2D/Source/Paper2DEditor/Private/PaperStyle.cpp b/Engine/Plugins/Experimental/Paper2D/Source/Paper2DEditor/Private/PaperStyle.cpp
--- a/Engine/Plugins/Experimental/Paper2D/Source/Paper2DEditor/Private/PaperStyle.cpp
+++ b/Engine/Plugins/Experimental/Paper2D/Source/Paper2DEditor/Private/PaperStyle.cpp
@@ -65,20 +65,26 @@ TSharedRef<FSlateStyleSet> FPaperStyle::Create()
 		Style.Set("SpriteEditor.SetShowPivot", new IMAGE_BRUSH(TEXT("Icons/icon_StaticMeshEd_ShowPivot_40x"), Icon40x40));
 		Style.Set("SpriteEditor.SetShowPivot.Small", new IMAGE_BRUSH(TEXT("Icons/icon_StaticMeshEd_ShowPivot_40x"), Icon20x20));
 
-		Style.Set("SpriteEditor.EnterViewMode", new IMAGE_BRUSH(TEXT("Icons/icon_TextureEd_AlphaChannel_40x"), Icon40x40));
-		Style.Set("SpriteEditor.EnterViewMode.Small", new IMAGE_BRUSH(TEXT("Icons/icon_TextureEd_AlphaChannel_40x"), Icon20x20));
+		Style.Set("SpriteEditor.AddPolygon", new IMAGE_BRUSH(TEXT("Icons/icon_Paper2D_AddPolygon_40x"), Icon40x40));
+		Style.Set("SpriteEditor.AddPolygon.Small", new IMAGE_BRUSH(TEXT("Icons/icon_Paper2D_AddPolygon_40x"), Icon20x20));
 
-		Style.Set("SpriteEditor.EnterCollisionEditMode", new FSlateImageBrush(TEXT("texture://Paper2D/SlateBrushes/EditCollisionGeom.EditCollisionGeom"), Icon40x40));
-		Style.Set("SpriteEditor.EnterCollisionEditMode.Small", new FSlateImageBrush(TEXT("texture://Paper2D/SlateBrushes/EditCollisionGeom.EditCollisionGeom"), Icon20x20));
+		Style.Set("SpriteEditor.SnapAllVertices", new IMAGE_BRUSH(TEXT("Icons/icon_Paper2D_SnapToPixelGrid_40x"), Icon40x40));
+		Style.Set("SpriteEditor.SnapAllVertices.Small", new IMAGE_BRUSH(TEXT("Icons/icon_Paper2D_SnapToPixelGrid_40x"), Icon20x20));
 
-		Style.Set("SpriteEditor.EnterSourceRegionEditMode", new FSlateImageBrush(TEXT("texture://Paper2D/SlateBrushes/EditSourceRegion.EditSourceRegion"), Icon40x40));
-		Style.Set("SpriteEditor.EnterSourceRegionEditMode.Small", new FSlateImageBrush(TEXT("texture://Paper2D/SlateBrushes/EditSourceRegion.EditSourceRegion"), Icon20x20));
+		Style.Set("SpriteEditor.EnterViewMode", new IMAGE_BRUSH(TEXT("Icons/icon_Paper2D_ViewSprite_40x"), Icon40x40));
+		Style.Set("SpriteEditor.EnterViewMode.Small", new IMAGE_BRUSH(TEXT("Icons/icon_Paper2D_ViewSprite_40x"), Icon20x20));
 
-		Style.Set("SpriteEditor.EnterRenderingEditMode", new FSlateImageBrush(TEXT("texture://Paper2D/SlateBrushes/EditRenderGeomButton.EditRenderGeomButton"), Icon40x40));
-		Style.Set("SpriteEditor.EnterRenderingEditMode.Small", new FSlateImageBrush(TEXT("texture://Paper2D/SlateBrushes/EditRenderGeomButton.EditRenderGeomButton"), Icon20x20));
+		Style.Set("SpriteEditor.EnterCollisionEditMode", new IMAGE_BRUSH("Icons/icon_Paper2D_EditCollision_40x", Icon40x40));
+		Style.Set("SpriteEditor.EnterCollisionEditMode.Small", new IMAGE_BRUSH("Icons/icon_Paper2D_EditCollision_40x", Icon20x20));
 
-		Style.Set("SpriteEditor.EnterAddSpriteMode", new FSlateImageBrush(TEXT("texture://Paper2D/SlateBrushes/AddSpriteButton.AddSpriteButton"), Icon40x40));
-		Style.Set("SpriteEditor.EnterAddSpriteMode.Small", new FSlateImageBrush(TEXT("texture://Paper2D/SlateBrushes/AddSpriteButton.AddSpriteButton"), Icon20x20));
+		Style.Set("SpriteEditor.EnterSourceRegionEditMode", new IMAGE_BRUSH("Icons/icon_Paper2D_EditSourceRegion_40x", Icon40x40));
+		Style.Set("SpriteEditor.EnterSourceRegionEditMode.Small", new IMAGE_BRUSH("Icons/icon_Paper2D_EditSourceRegion_40x", Icon20x20));
+
+		Style.Set("SpriteEditor.EnterRenderingEditMode", new IMAGE_BRUSH("Icons/icon_Paper2D_RenderGeom_40x", Icon40x40));
+		Style.Set("SpriteEditor.EnterRenderingEditMode.Small", new IMAGE_BRUSH("Icons/icon_Paper2D_RenderGeom_40x", Icon20x20));
+
+		Style.Set("SpriteEditor.EnterAddSpriteMode", new IMAGE_BRUSH("Icons/icon_Paper2D_AddSprite_40x", Icon40x40));
+		Style.Set("SpriteEditor.EnterAddSpriteMode.Small", new IMAGE_BRUSH("Icons/icon_Paper2D_AddSprite_40x", Icon20x20));
 	}
 
 	// Flipbook editor
