Paper2D: Add stub documentation for sprite and flipbook editors

diff --git a/Engine/Plugins/Experimental/Paper2D/Source/Paper2DEditor/Private/FlipbookEditor/FlipbookEditor.cpp b/Engine/Plugins/Experimental/Paper2D/Source/Paper2DEditor/Private/FlipbookEditor/FlipbookEditor.cpp
--- a/Engine/Plugins/Experimental/Paper2D/Source/Paper2DEditor/Private/FlipbookEditor/FlipbookEditor.cpp
+++ b/Engine/Plugins/Experimental/Paper2D/Source/Paper2DEditor/Private/FlipbookEditor/FlipbookEditor.cpp
@@ -353,6 +353,11 @@ FString FFlipbookEditor::GetWorldCentricTabPrefix() const
 	return TEXT("FlipbookEditor");
 }
 
+FString FFlipbookEditor::GetDocumentationLink() const
+{
+	return TEXT("Engine/Paper2D/FlipbookEditor");
+}
+
 FLinearColor FFlipbookEditor::GetWorldCentricTabColorScale() const
 {
 	return FLinearColor::White;
diff --git a/Engine/Plugins/Experimental/Paper2D/Source/Paper2DEditor/Private/FlipbookEditor/FlipbookEditor.h b/Engine/Plugins/Experimental/Paper2D/Source/Paper2DEditor/Private/FlipbookEditor/FlipbookEditor.h
--- a/Engine/Plugins/Experimental/Paper2D/Source/Paper2DEditor/Private/FlipbookEditor/FlipbookEditor.h
+++ b/Engine/Plugins/Experimental/Paper2D/Source/Paper2DEditor/Private/FlipbookEditor/FlipbookEditor.h
@@ -22,6 +22,7 @@ public:
 	virtual FText GetToolkitName() const override;
 	virtual FLinearColor GetWorldCentricTabColorScale() const override;
 	virtual FString GetWorldCentricTabPrefix() const override;
+	virtual FString GetDocumentationLink() const override;
 	// End of FAssetEditorToolkit
 
 	// FSerializableObject interface
diff --git a/Engine/Plugins/Experimental/Paper2D/Source/Paper2DEditor/Private/SpriteEditor/SpriteEditor.cpp b/Engine/Plugins/Experimental/Paper2D/Source/Paper2DEditor/Private/SpriteEditor/SpriteEditor.cpp
--- a/Engine/Plugins/Experimental/Paper2D/Source/Paper2DEditor/Private/SpriteEditor/SpriteEditor.cpp
+++ b/Engine/Plugins/Experimental/Paper2D/Source/Paper2DEditor/Private/SpriteEditor/SpriteEditor.cpp
@@ -434,6 +434,11 @@ FString FSpriteEditor::GetWorldCentricTabPrefix() const
 	return TEXT("SpriteEditor");
 }
 
+FString FSpriteEditor::GetDocumentationLink() const
+{
+	return TEXT("Engine/Paper2D/SpriteEditor");
+}
+
 FLinearColor FSpriteEditor::GetWorldCentricTabColorScale() const
 {
 	return FLinearColor::White;
diff --git a/Engine/Plugins/Experimental/Paper2D/Source/Paper2DEditor/Private/SpriteEditor/SpriteEditor.h b/Engine/Plugins/Experimental/Paper2D/Source/Paper2DEditor/Private/SpriteEditor/SpriteEditor.h
--- a/Engine/Plugins/Experimental/Paper2D/Source/Paper2DEditor/Private/SpriteEditor/SpriteEditor.h
+++ b/Engine/Plugins/Experimental/Paper2D/Source/Paper2DEditor/Private/SpriteEditor/SpriteEditor.h
@@ -22,6 +22,7 @@ public:
 	virtual FText GetToolkitName() const override;
 	virtual FLinearColor GetWorldCentricTabColorScale() const override;
 	virtual FString GetWorldCentricTabPrefix() const override;
+	virtual FString GetDocumentationLink() const override;
 	// End of FAssetEditorToolkit
 
 	// FSerializableObject interface
