Paper2D: Code style fixes

diff --git a/Engine/Plugins/Experimental/Paper2D/Source/Paper2DEditor/Private/FlipbookEditor/SFlipbookTimeline.cpp b/Engine/Plugins/Experimental/Paper2D/Source/Paper2DEditor/Private/FlipbookEditor/SFlipbookTimeline.cpp
--- a/Engine/Plugins/Experimental/Paper2D/Source/Paper2DEditor/Private/FlipbookEditor/SFlipbookTimeline.cpp
+++ b/Engine/Plugins/Experimental/Paper2D/Source/Paper2DEditor/Private/FlipbookEditor/SFlipbookTimeline.cpp
@@ -6,10 +6,9 @@
 class STrackHandle : public SImage
 {
 public:
-
 	SLATE_BEGIN_ARGS(STrackHandle)
 		: _SlateUnitsPerFrame(1)
-		, _FlipbookBeingEdited((class UPaperFlipbook*)NULL)
+		, _FlipbookBeingEdited(nullptr)
 		, _KeyFrameIdx(INDEX_NONE)
 	{}
 
@@ -29,16 +28,16 @@ public:
 		bDragging = false;
 		StartingFrameRun = INDEX_NONE;
 
-		SImage::Construct( SImage::FArguments() );
+		SImage::Construct(SImage::FArguments());
 	}
 
-	virtual FReply OnMouseButtonDown( const FGeometry& MyGeometry, const FPointerEvent& MouseEvent ) override
+	virtual FReply OnMouseButtonDown(const FGeometry& MyGeometry, const FPointerEvent& MouseEvent) override
 	{
-		if ( MouseEvent.GetEffectingButton() == EKeys::LeftMouseButton )
+		if (MouseEvent.GetEffectingButton() == EKeys::LeftMouseButton)
 		{
 			DistanceDragged = 0;
 			StartingFrameRun = INDEX_NONE;
-			return FReply::Handled().CaptureMouse( SharedThis(this) ).UseHighPrecisionMouseMovement( SharedThis(this) );
+			return FReply::Handled().CaptureMouse(SharedThis(this)).UseHighPrecisionMouseMovement(SharedThis(this));
 		}
 		else
 		{
@@ -47,18 +46,18 @@ public:
 		}
 	}
 
-	virtual FReply OnMouseButtonUp( const FGeometry& MyGeometry, const FPointerEvent& MouseEvent ) override
+	virtual FReply OnMouseButtonUp(const FGeometry& MyGeometry, const FPointerEvent& MouseEvent) override
 	{
-		if ( MouseEvent.GetEffectingButton() == EKeys::LeftMouseButton && this->HasMouseCapture() )
+		if ((MouseEvent.GetEffectingButton() == EKeys::LeftMouseButton) && this->HasMouseCapture())
 		{
-			if ( bDragging && StartingFrameRun != INDEX_NONE )
+			if (bDragging && (StartingFrameRun != INDEX_NONE))
 			{
 				UPaperFlipbook* Flipbook = FlipbookBeingEdited.Get();
-				if ( Flipbook && Flipbook->KeyFrames.IsValidIndex(KeyFrameIdx) )
+				if (Flipbook && Flipbook->KeyFrames.IsValidIndex(KeyFrameIdx))
 				{
 					FPaperFlipbookKeyFrame& KeyFrame = Flipbook->KeyFrames[KeyFrameIdx];
 
-					if ( KeyFrame.FrameRun != StartingFrameRun )
+					if (KeyFrame.FrameRun != StartingFrameRun)
 					{
 						Flipbook->MarkPackageDirty();
 						Flipbook->PostEditChange();
@@ -82,12 +81,12 @@ public:
 		}
 	}
 
-	virtual FReply OnMouseMove( const FGeometry& MyGeometry, const FPointerEvent& MouseEvent ) override
+	virtual FReply OnMouseMove(const FGeometry& MyGeometry, const FPointerEvent& MouseEvent) override
 	{
-		if ( this->HasMouseCapture() )
+		if (this->HasMouseCapture())
 		{
 			UPaperFlipbook* Flipbook = FlipbookBeingEdited.Get();
-			if ( Flipbook && Flipbook->KeyFrames.IsValidIndex(KeyFrameIdx) )
+			if (Flipbook && Flipbook->KeyFrames.IsValidIndex(KeyFrameIdx))
 			{
 				FPaperFlipbookKeyFrame& KeyFrame = Flipbook->KeyFrames[KeyFrameIdx];
 
@@ -95,7 +94,7 @@ public:
 
 				if (!bDragging)
 				{
-					if ( FMath::Abs(DistanceDragged) > SlateDragStartDistance )
+					if (FMath::Abs(DistanceDragged) > SlateDragStartDistance)
 					{
 						StartingFrameRun = KeyFrame.FrameRun;
 						bDragging = true;
@@ -104,10 +103,10 @@ public:
 				else
 				{
 					float LocalSlateUnitsPerFrame = SlateUnitsPerFrame.Get();
-					if ( LocalSlateUnitsPerFrame != 0 )
+					if (LocalSlateUnitsPerFrame != 0)
 					{
-						KeyFrame.FrameRun = StartingFrameRun + ( DistanceDragged / LocalSlateUnitsPerFrame);
-						KeyFrame.FrameRun = FMath::Max<int32>( 1, KeyFrame.FrameRun );
+						KeyFrame.FrameRun = StartingFrameRun + (DistanceDragged / LocalSlateUnitsPerFrame);
+						KeyFrame.FrameRun = FMath::Max<int32>(1, KeyFrame.FrameRun);
 					}
 				}
 			}
@@ -118,11 +117,9 @@ public:
 		return FReply::Unhandled();
 	}
 
-	virtual FCursorReply OnCursorQuery( const FGeometry& MyGeometry, const FPointerEvent& CursorEvent ) const override
+	virtual FCursorReply OnCursorQuery(const FGeometry& MyGeometry, const FPointerEvent& CursorEvent) const override
 	{
-		return bDragging ? 
-			FCursorReply::Cursor( EMouseCursor::None ) :
-			FCursorReply::Cursor( EMouseCursor::ResizeLeftRight );
+		return FCursorReply::Cursor(bDragging ? EMouseCursor::None : EMouseCursor::ResizeLeftRight);
 	}
 
 private:
@@ -141,7 +138,7 @@ class STimelineTrack : public SCompoundWidget
 public:
 	SLATE_BEGIN_ARGS(STimelineTrack)
 		: _SlateUnitsPerFrame(1)
-		, _FlipbookBeingEdited((class UPaperFlipbook*)NULL)
+		, _FlipbookBeingEdited(nullptr)
 	{}
 
 		SLATE_ATTRIBUTE( float, SlateUnitsPerFrame )
@@ -165,13 +162,13 @@ public:
 		Rebuild();
 	}
 
-	void Tick( const FGeometry& AllottedGeometry, const double InCurrentTime, const float InDeltaTime ) override
+	void Tick(const FGeometry& AllottedGeometry, const double InCurrentTime, const float InDeltaTime) override
 	{
 		SCompoundWidget::Tick(AllottedGeometry, InCurrentTime, InDeltaTime);
 
 		UPaperFlipbook* Flipbook = FlipbookBeingEdited.Get();
 		int32 NewNumKeyframes = (Flipbook != NULL ? Flipbook->KeyFrames.Num() : 0);
-		if ( NewNumKeyframes != NumKeyframesFromLastRebuild )
+		if (NewNumKeyframes != NumKeyframesFromLastRebuild)
 		{
 			Rebuild();
 		}
@@ -182,10 +179,9 @@ private:
 	{
 		MainBoxPtr->ClearChildren();
 
-		UPaperFlipbook* Flipbook = FlipbookBeingEdited.Get();
-		if ( Flipbook )
+		if (UPaperFlipbook* Flipbook = FlipbookBeingEdited.Get())
 		{
-			for ( int32 KeyFrameIdx = 0; KeyFrameIdx < Flipbook->KeyFrames.Num(); ++KeyFrameIdx )
+			for (int32 KeyFrameIdx = 0; KeyFrameIdx < Flipbook->KeyFrames.Num(); ++KeyFrameIdx)
 			{
 				MainBoxPtr->AddSlot()
 				.AutoWidth()
@@ -225,7 +221,7 @@ private:
 	FOptionalSize GetFrameWidth(int32 KeyFrameIdx) const
 	{
 		UPaperFlipbook* Flipbook = FlipbookBeingEdited.Get();
-		if ( Flipbook && Flipbook->KeyFrames.IsValidIndex(KeyFrameIdx) )
+		if (Flipbook && Flipbook->KeyFrames.IsValidIndex(KeyFrameIdx))
 		{
 			const FPaperFlipbookKeyFrame& KeyFrame = Flipbook->KeyFrames[KeyFrameIdx];
 			return FMath::Max<float>(0, KeyFrame.FrameRun * SlateUnitsPerFrame.Get() - HandleWidth);
@@ -294,13 +290,13 @@ public:
 		Rebuild();
 	}
 
-	void Tick( const FGeometry& AllottedGeometry, const double InCurrentTime, const float InDeltaTime ) override
+	void Tick(const FGeometry& AllottedGeometry, const double InCurrentTime, const float InDeltaTime) override
 	{
 		SCompoundWidget::Tick(AllottedGeometry, InCurrentTime, InDeltaTime);
 
 		UPaperFlipbook* Flipbook = FlipbookBeingEdited.Get();
 		int32 NewNumFrames = (Flipbook != NULL ? Flipbook->GetNumFrames() : 0);
-		if ( NewNumFrames != NumFramesFromLastRebuild )
+		if (NewNumFrames != NumFramesFromLastRebuild)
 		{
 			Rebuild();
 		}
@@ -348,7 +344,7 @@ private:
 
 private:
 	TAttribute<float> SlateUnitsPerFrame;
-	TAttribute< class UPaperFlipbook* > FlipbookBeingEdited;
+	TAttribute<class UPaperFlipbook*> FlipbookBeingEdited;
 	TAttribute<float> PlayTime;
 
 	TSharedPtr<SHorizontalBox> MainBoxPtr;
diff --git a/Engine/Plugins/Experimental/Paper2D/Source/Paper2DEditor/Private/FlipbookEditor/SFlipbookTimeline.h b/Engine/Plugins/Experimental/Paper2D/Source/Paper2DEditor/Private/FlipbookEditor/SFlipbookTimeline.h
--- a/Engine/Plugins/Experimental/Paper2D/Source/Paper2DEditor/Private/FlipbookEditor/SFlipbookTimeline.h
+++ b/Engine/Plugins/Experimental/Paper2D/Source/Paper2DEditor/Private/FlipbookEditor/SFlipbookTimeline.h
@@ -8,18 +8,18 @@ class SFlipbookTimeline : public SCompoundWidget
 {
 public:
 	SLATE_BEGIN_ARGS(SFlipbookTimeline)
-		: _FlipbookBeingEdited((class UPaperFlipbook*)NULL)
+		: _FlipbookBeingEdited(nullptr)
 		, _PlayTime(0)
 	{}
 
-		SLATE_ATTRIBUTE( class UPaperFlipbook*, FlipbookBeingEdited )
-		SLATE_ATTRIBUTE( float, PlayTime )
+		SLATE_ATTRIBUTE(class UPaperFlipbook*, FlipbookBeingEdited)
+		SLATE_ATTRIBUTE(float, PlayTime)
 
 	SLATE_END_ARGS()
 
 	void Construct(const FArguments& InArgs);
 
 private:
-	TAttribute< class UPaperFlipbook* > FlipbookBeingEdited;
-	TAttribute< float > PlayTime;
+	TAttribute<class UPaperFlipbook*> FlipbookBeingEdited;
+	TAttribute<float> PlayTime;
 };
