Paper2D: Split flipbook widgets out into separate files

diff --git a/Engine/Plugins/Experimental/Paper2D/Source/Paper2DEditor/Private/FlipbookEditor/SFlipbookTimeline.cpp b/Engine/Plugins/Experimental/Paper2D/Source/Paper2DEditor/Private/FlipbookEditor/SFlipbookTimeline.cpp
--- a/Engine/Plugins/Experimental/Paper2D/Source/Paper2DEditor/Private/FlipbookEditor/SFlipbookTimeline.cpp
+++ b/Engine/Plugins/Experimental/Paper2D/Source/Paper2DEditor/Private/FlipbookEditor/SFlipbookTimeline.cpp
@@ -3,355 +3,12 @@
 #include "Paper2DEditorPrivatePCH.h"
 #include "SFlipbookTimeline.h"
 
-class STrackHandle : public SImage
-{
-public:
-	SLATE_BEGIN_ARGS(STrackHandle)
-		: _SlateUnitsPerFrame(1)
-		, _FlipbookBeingEdited(nullptr)
-		, _KeyFrameIdx(INDEX_NONE)
-	{}
-
-		SLATE_ATTRIBUTE( float, SlateUnitsPerFrame )
-		SLATE_ATTRIBUTE( class UPaperFlipbook*, FlipbookBeingEdited )
-		SLATE_ARGUMENT( int32, KeyFrameIdx )
-		
-	SLATE_END_ARGS()
-
-	void Construct(const FArguments& InArgs)
-	{
-		SlateUnitsPerFrame = InArgs._SlateUnitsPerFrame;
-		FlipbookBeingEdited = InArgs._FlipbookBeingEdited;
-		KeyFrameIdx = InArgs._KeyFrameIdx;
-
-		DistanceDragged = 0;
-		bDragging = false;
-		StartingFrameRun = INDEX_NONE;
-
-		SImage::Construct(SImage::FArguments());
-	}
-
-	virtual FReply OnMouseButtonDown(const FGeometry& MyGeometry, const FPointerEvent& MouseEvent) override
-	{
-		if (MouseEvent.GetEffectingButton() == EKeys::LeftMouseButton)
-		{
-			DistanceDragged = 0;
-			StartingFrameRun = INDEX_NONE;
-			return FReply::Handled().CaptureMouse(SharedThis(this)).UseHighPrecisionMouseMovement(SharedThis(this));
-		}
-		else
-		{
-
-			return FReply::Unhandled();
-		}
-	}
-
-	virtual FReply OnMouseButtonUp(const FGeometry& MyGeometry, const FPointerEvent& MouseEvent) override
-	{
-		if ((MouseEvent.GetEffectingButton() == EKeys::LeftMouseButton) && this->HasMouseCapture())
-		{
-			if (bDragging && (StartingFrameRun != INDEX_NONE))
-			{
-				UPaperFlipbook* Flipbook = FlipbookBeingEdited.Get();
-				if (Flipbook && Flipbook->KeyFrames.IsValidIndex(KeyFrameIdx))
-				{
-					FPaperFlipbookKeyFrame& KeyFrame = Flipbook->KeyFrames[KeyFrameIdx];
-
-					if (KeyFrame.FrameRun != StartingFrameRun)
-					{
-						Flipbook->MarkPackageDirty();
-						Flipbook->PostEditChange();
-					}
-				}
-			}
-
-			bDragging = false;
-
-			FIntPoint NewMousePos(
-				(MyGeometry.AbsolutePosition.X + MyGeometry.Size.X / 2) * MyGeometry.Scale,
-				(MyGeometry.AbsolutePosition.Y + MyGeometry.Size.Y / 2) * MyGeometry.Scale
-				);
-
-			return FReply::Handled().ReleaseMouseCapture().SetMousePos(NewMousePos);
-
-		}
-		else
-		{
-			return FReply::Unhandled();
-		}
-	}
-
-	virtual FReply OnMouseMove(const FGeometry& MyGeometry, const FPointerEvent& MouseEvent) override
-	{
-		if (this->HasMouseCapture())
-		{
-			UPaperFlipbook* Flipbook = FlipbookBeingEdited.Get();
-			if (Flipbook && Flipbook->KeyFrames.IsValidIndex(KeyFrameIdx))
-			{
-				FPaperFlipbookKeyFrame& KeyFrame = Flipbook->KeyFrames[KeyFrameIdx];
-
-				DistanceDragged += MouseEvent.GetCursorDelta().X;
-
-				if (!bDragging)
-				{
-					if (FMath::Abs(DistanceDragged) > SlateDragStartDistance)
-					{
-						StartingFrameRun = KeyFrame.FrameRun;
-						bDragging = true;
-					}
-				}
-				else
-				{
-					float LocalSlateUnitsPerFrame = SlateUnitsPerFrame.Get();
-					if (LocalSlateUnitsPerFrame != 0)
-					{
-						KeyFrame.FrameRun = StartingFrameRun + (DistanceDragged / LocalSlateUnitsPerFrame);
-						KeyFrame.FrameRun = FMath::Max<int32>(1, KeyFrame.FrameRun);
-					}
-				}
-			}
-
-			return FReply::Handled();
-		}
-
-		return FReply::Unhandled();
-	}
-
-	virtual FCursorReply OnCursorQuery(const FGeometry& MyGeometry, const FPointerEvent& CursorEvent) const override
-	{
-		return FCursorReply::Cursor(bDragging ? EMouseCursor::None : EMouseCursor::ResizeLeftRight);
-	}
-
-private:
-	float DistanceDragged;
-	int32 StartingFrameRun;
-	bool bDragging;
-
-	TAttribute<float> SlateUnitsPerFrame;
-	TAttribute<class UPaperFlipbook*> FlipbookBeingEdited;
-	int32 KeyFrameIdx;
-};
-
-
-class STimelineTrack : public SCompoundWidget
-{
-public:
-	SLATE_BEGIN_ARGS(STimelineTrack)
-		: _SlateUnitsPerFrame(1)
-		, _FlipbookBeingEdited(nullptr)
-	{}
-
-		SLATE_ATTRIBUTE( float, SlateUnitsPerFrame )
-		SLATE_ATTRIBUTE( class UPaperFlipbook*, FlipbookBeingEdited )
-
-	SLATE_END_ARGS()
-
-	void Construct(const FArguments& InArgs)
-	{
-		SlateUnitsPerFrame = InArgs._SlateUnitsPerFrame;
-		FlipbookBeingEdited = InArgs._FlipbookBeingEdited;
-
-		HandleWidth = 4;
-		NumKeyframesFromLastRebuild = 0;
-
-		ChildSlot
-		[
-			SAssignNew(MainBoxPtr, SHorizontalBox)
-		];
-
-		Rebuild();
-	}
-
-	void Tick(const FGeometry& AllottedGeometry, const double InCurrentTime, const float InDeltaTime) override
-	{
-		SCompoundWidget::Tick(AllottedGeometry, InCurrentTime, InDeltaTime);
-
-		UPaperFlipbook* Flipbook = FlipbookBeingEdited.Get();
-		int32 NewNumKeyframes = (Flipbook != NULL ? Flipbook->KeyFrames.Num() : 0);
-		if (NewNumKeyframes != NumKeyframesFromLastRebuild)
-		{
-			Rebuild();
-		}
-	}
-
-private:
-	void Rebuild()
-	{
-		MainBoxPtr->ClearChildren();
-
-		if (UPaperFlipbook* Flipbook = FlipbookBeingEdited.Get())
-		{
-			for (int32 KeyFrameIdx = 0; KeyFrameIdx < Flipbook->KeyFrames.Num(); ++KeyFrameIdx)
-			{
-				MainBoxPtr->AddSlot()
-				.AutoWidth()
-				[
-					SNew(SBox)
-					.WidthOverride( TAttribute<FOptionalSize>::Create( TAttribute<FOptionalSize>::FGetter::CreateSP(this, &STimelineTrack::GetFrameWidth, KeyFrameIdx) ) )
-					[
-						SNew(SBorder)
-						.BorderImage( FEditorStyle::GetBrush("PropertyTable.CurrentCellBorder") )
-						[
-							SNullWidget::NullWidget
-						]
-					]
-				];
-
-				MainBoxPtr->AddSlot()
-				.AutoWidth()
-				[
-					SNew(SBox).WidthOverride(HandleWidth)
-					[
-						SNew(STrackHandle)
-						.SlateUnitsPerFrame(SlateUnitsPerFrame)
-						.FlipbookBeingEdited(FlipbookBeingEdited)
-						.KeyFrameIdx(KeyFrameIdx)
-					]
-				];
-			}
-
-			NumKeyframesFromLastRebuild = Flipbook->KeyFrames.Num();
-		}
-		else
-		{
-			NumKeyframesFromLastRebuild = 0;
-		}
-	}
-
-	FOptionalSize GetFrameWidth(int32 KeyFrameIdx) const
-	{
-		UPaperFlipbook* Flipbook = FlipbookBeingEdited.Get();
-		if (Flipbook && Flipbook->KeyFrames.IsValidIndex(KeyFrameIdx))
-		{
-			const FPaperFlipbookKeyFrame& KeyFrame = Flipbook->KeyFrames[KeyFrameIdx];
-			return FMath::Max<float>(0, KeyFrame.FrameRun * SlateUnitsPerFrame.Get() - HandleWidth);
-		}
-		else
-		{
-			return 1;
-		}
-	}
-
-private:
-	TAttribute<float> SlateUnitsPerFrame;
-	TAttribute< class UPaperFlipbook* > FlipbookBeingEdited;
-
-	TSharedPtr<SHorizontalBox> MainBoxPtr;
-
-	int32 NumKeyframesFromLastRebuild;
-	float HandleWidth;
-};
-
-
-
-class STimelineHeader : public SCompoundWidget
-{
-public:
-	SLATE_BEGIN_ARGS(STimelineHeader)
-		: _SlateUnitsPerFrame(1)
-		, _FlipbookBeingEdited((class UPaperFlipbook*)NULL)
-		, _PlayTime(0)
-	{}
-
-		SLATE_ATTRIBUTE( float, SlateUnitsPerFrame )
-		SLATE_ATTRIBUTE( class UPaperFlipbook*, FlipbookBeingEdited )
-		SLATE_ATTRIBUTE( float, PlayTime )
-
-	SLATE_END_ARGS()
-
-	void Construct(const FArguments& InArgs)
-	{
-		SlateUnitsPerFrame = InArgs._SlateUnitsPerFrame;
-		FlipbookBeingEdited = InArgs._FlipbookBeingEdited;
-		PlayTime = InArgs._PlayTime;
-
-		NumFramesFromLastRebuild = 0;
-
-		ChildSlot
-		[
-			SNew(SOverlay)
-
-			+SOverlay::Slot()
-			[
-				SAssignNew(MainBoxPtr, SHorizontalBox)
-			]
-
-			+SOverlay::Slot()
-			.Padding( TAttribute<FMargin>(this, &STimelineHeader::GetPlayTimePadding) )
-			[
-				SNew(SBox)
-				.WidthOverride(SlateUnitsPerFrame.Get())
-				[
-					SNew(SImage)
-				]
-			]
-		];
-
-		Rebuild();
-	}
-
-	void Tick(const FGeometry& AllottedGeometry, const double InCurrentTime, const float InDeltaTime) override
-	{
-		SCompoundWidget::Tick(AllottedGeometry, InCurrentTime, InDeltaTime);
-
-		UPaperFlipbook* Flipbook = FlipbookBeingEdited.Get();
-		int32 NewNumFrames = (Flipbook != NULL ? Flipbook->GetNumFrames() : 0);
-		if (NewNumFrames != NumFramesFromLastRebuild)
-		{
-			Rebuild();
-		}
-	}
-
-private:
-	void Rebuild()
-	{
-		MainBoxPtr->ClearChildren();
-
-		UPaperFlipbook* Flipbook = FlipbookBeingEdited.Get();
-		float LocalSlateUnitsPerFrame = SlateUnitsPerFrame.Get();
-		if ( Flipbook && LocalSlateUnitsPerFrame > 0 )
-		{
-			const int32 NumFrames = Flipbook->GetNumFrames();
-			for ( int32 FrameIdx = 0; FrameIdx < NumFrames; ++FrameIdx )
-			{
-				MainBoxPtr->AddSlot()
-					.AutoWidth()
-					[
-						SNew(SBox)
-						.WidthOverride(LocalSlateUnitsPerFrame)
-						.HAlign(HAlign_Center)
-						[
-							SNew(STextBlock)
-							.Text(FText::AsNumber(FrameIdx).ToString())
-						]
-					];
-			}
-
-			NumFramesFromLastRebuild = NumFrames;
-		}
-		else
-		{
-			NumFramesFromLastRebuild = 0;
-		}
-	}
-
-	FMargin GetPlayTimePadding() const
-	{
-		float PlayTimePadding = PlayTime.Get() - SlateUnitsPerFrame.Get() / 2;
-		return FMargin(PlayTimePadding, 0, 0, 0);
-	}
-	
-
-private:
-	TAttribute<float> SlateUnitsPerFrame;
-	TAttribute<class UPaperFlipbook*> FlipbookBeingEdited;
-	TAttribute<float> PlayTime;
-
-	TSharedPtr<SHorizontalBox> MainBoxPtr;
-
-	int32 NumFramesFromLastRebuild;
-};
+#include "SFlipbookTrackHandle.h"
+#include "STimelineHeader.h"
+#include "STimelineTrack.h"
 
+//////////////////////////////////////////////////////////////////////////
+// SFlipbookTimeline
 
 void SFlipbookTimeline::Construct(const FArguments& InArgs)
 {
diff --git a/Engine/Plugins/Experimental/Paper2D/Source/Paper2DEditor/Private/FlipbookEditor/SFlipbookTrackHandle.h b/Engine/Plugins/Experimental/Paper2D/Source/Paper2DEditor/Private/FlipbookEditor/SFlipbookTrackHandle.h
new file mode 100644
--- /dev/null
+++ b/Engine/Plugins/Experimental/Paper2D/Source/Paper2DEditor/Private/FlipbookEditor/SFlipbookTrackHandle.h
@@ -0,0 +1,133 @@
+// Copyright 1998-2014 Epic Games, Inc. All Rights Reserved.
+
+//////////////////////////////////////////////////////////////////////////
+// SFlipbookTrackHandle
+
+class SFlipbookTrackHandle : public SImage
+{
+public:
+	SLATE_BEGIN_ARGS(SFlipbookTrackHandle)
+		: _SlateUnitsPerFrame(1)
+		, _FlipbookBeingEdited(nullptr)
+		, _KeyFrameIdx(INDEX_NONE)
+	{}
+
+		SLATE_ATTRIBUTE( float, SlateUnitsPerFrame )
+		SLATE_ATTRIBUTE( class UPaperFlipbook*, FlipbookBeingEdited )
+		SLATE_ARGUMENT( int32, KeyFrameIdx )
+		
+	SLATE_END_ARGS()
+
+	void Construct(const FArguments& InArgs)
+	{
+		SlateUnitsPerFrame = InArgs._SlateUnitsPerFrame;
+		FlipbookBeingEdited = InArgs._FlipbookBeingEdited;
+		KeyFrameIdx = InArgs._KeyFrameIdx;
+
+		DistanceDragged = 0;
+		bDragging = false;
+		StartingFrameRun = INDEX_NONE;
+
+		SImage::Construct(SImage::FArguments());
+	}
+
+	virtual FReply OnMouseButtonDown(const FGeometry& MyGeometry, const FPointerEvent& MouseEvent) override
+	{
+		if (MouseEvent.GetEffectingButton() == EKeys::LeftMouseButton)
+		{
+			DistanceDragged = 0;
+			StartingFrameRun = INDEX_NONE;
+			return FReply::Handled().CaptureMouse(SharedThis(this)).UseHighPrecisionMouseMovement(SharedThis(this));
+		}
+		else
+		{
+
+			return FReply::Unhandled();
+		}
+	}
+
+	virtual FReply OnMouseButtonUp(const FGeometry& MyGeometry, const FPointerEvent& MouseEvent) override
+	{
+		if ((MouseEvent.GetEffectingButton() == EKeys::LeftMouseButton) && this->HasMouseCapture())
+		{
+			if (bDragging && (StartingFrameRun != INDEX_NONE))
+			{
+				UPaperFlipbook* Flipbook = FlipbookBeingEdited.Get();
+				if (Flipbook && Flipbook->KeyFrames.IsValidIndex(KeyFrameIdx))
+				{
+					FPaperFlipbookKeyFrame& KeyFrame = Flipbook->KeyFrames[KeyFrameIdx];
+
+					if (KeyFrame.FrameRun != StartingFrameRun)
+					{
+						Flipbook->MarkPackageDirty();
+						Flipbook->PostEditChange();
+					}
+				}
+			}
+
+			bDragging = false;
+
+			FIntPoint NewMousePos(
+				(MyGeometry.AbsolutePosition.X + MyGeometry.Size.X / 2) * MyGeometry.Scale,
+				(MyGeometry.AbsolutePosition.Y + MyGeometry.Size.Y / 2) * MyGeometry.Scale
+				);
+
+			return FReply::Handled().ReleaseMouseCapture().SetMousePos(NewMousePos);
+
+		}
+		else
+		{
+			return FReply::Unhandled();
+		}
+	}
+
+	virtual FReply OnMouseMove(const FGeometry& MyGeometry, const FPointerEvent& MouseEvent) override
+	{
+		if (this->HasMouseCapture())
+		{
+			UPaperFlipbook* Flipbook = FlipbookBeingEdited.Get();
+			if (Flipbook && Flipbook->KeyFrames.IsValidIndex(KeyFrameIdx))
+			{
+				FPaperFlipbookKeyFrame& KeyFrame = Flipbook->KeyFrames[KeyFrameIdx];
+
+				DistanceDragged += MouseEvent.GetCursorDelta().X;
+
+				if (!bDragging)
+				{
+					if (FMath::Abs(DistanceDragged) > SlateDragStartDistance)
+					{
+						StartingFrameRun = KeyFrame.FrameRun;
+						bDragging = true;
+					}
+				}
+				else
+				{
+					float LocalSlateUnitsPerFrame = SlateUnitsPerFrame.Get();
+					if (LocalSlateUnitsPerFrame != 0)
+					{
+						KeyFrame.FrameRun = StartingFrameRun + (DistanceDragged / LocalSlateUnitsPerFrame);
+						KeyFrame.FrameRun = FMath::Max<int32>(1, KeyFrame.FrameRun);
+					}
+				}
+			}
+
+			return FReply::Handled();
+		}
+
+		return FReply::Unhandled();
+	}
+
+	virtual FCursorReply OnCursorQuery(const FGeometry& MyGeometry, const FPointerEvent& CursorEvent) const override
+	{
+		return FCursorReply::Cursor(bDragging ? EMouseCursor::None : EMouseCursor::ResizeLeftRight);
+	}
+
+private:
+	float DistanceDragged;
+	int32 StartingFrameRun;
+	bool bDragging;
+
+	TAttribute<float> SlateUnitsPerFrame;
+	TAttribute<class UPaperFlipbook*> FlipbookBeingEdited;
+	int32 KeyFrameIdx;
+};
diff --git a/Engine/Plugins/Experimental/Paper2D/Source/Paper2DEditor/Private/FlipbookEditor/STimelineHeader.h b/Engine/Plugins/Experimental/Paper2D/Source/Paper2DEditor/Private/FlipbookEditor/STimelineHeader.h
new file mode 100644
--- /dev/null
+++ b/Engine/Plugins/Experimental/Paper2D/Source/Paper2DEditor/Private/FlipbookEditor/STimelineHeader.h
@@ -0,0 +1,112 @@
+// Copyright 1998-2014 Epic Games, Inc. All Rights Reserved.
+
+//////////////////////////////////////////////////////////////////////////
+// STimelineHeader
+
+class STimelineHeader : public SCompoundWidget
+{
+public:
+	SLATE_BEGIN_ARGS(STimelineHeader)
+		: _SlateUnitsPerFrame(1)
+		, _FlipbookBeingEdited((class UPaperFlipbook*)NULL)
+		, _PlayTime(0)
+	{}
+
+	SLATE_ATTRIBUTE(float, SlateUnitsPerFrame)
+		SLATE_ATTRIBUTE(class UPaperFlipbook*, FlipbookBeingEdited)
+		SLATE_ATTRIBUTE(float, PlayTime)
+
+		SLATE_END_ARGS()
+
+		void Construct(const FArguments& InArgs)
+	{
+			SlateUnitsPerFrame = InArgs._SlateUnitsPerFrame;
+			FlipbookBeingEdited = InArgs._FlipbookBeingEdited;
+			PlayTime = InArgs._PlayTime;
+
+			NumFramesFromLastRebuild = 0;
+
+			ChildSlot
+				[
+					SNew(SOverlay)
+
+					+ SOverlay::Slot()
+					[
+						SAssignNew(MainBoxPtr, SHorizontalBox)
+					]
+
+					+ SOverlay::Slot()
+						.Padding(TAttribute<FMargin>(this, &STimelineHeader::GetPlayTimePadding))
+						[
+							SNew(SBox)
+							.WidthOverride(SlateUnitsPerFrame.Get())
+							[
+								SNew(SImage)
+							]
+						]
+				];
+
+			Rebuild();
+		}
+
+	void Tick(const FGeometry& AllottedGeometry, const double InCurrentTime, const float InDeltaTime) override
+	{
+		SCompoundWidget::Tick(AllottedGeometry, InCurrentTime, InDeltaTime);
+
+		UPaperFlipbook* Flipbook = FlipbookBeingEdited.Get();
+		int32 NewNumFrames = (Flipbook != NULL ? Flipbook->GetNumFrames() : 0);
+		if (NewNumFrames != NumFramesFromLastRebuild)
+		{
+			Rebuild();
+		}
+	}
+
+private:
+	void Rebuild()
+	{
+		MainBoxPtr->ClearChildren();
+
+		UPaperFlipbook* Flipbook = FlipbookBeingEdited.Get();
+		float LocalSlateUnitsPerFrame = SlateUnitsPerFrame.Get();
+		if (Flipbook && LocalSlateUnitsPerFrame > 0)
+		{
+			const int32 NumFrames = Flipbook->GetNumFrames();
+			for (int32 FrameIdx = 0; FrameIdx < NumFrames; ++FrameIdx)
+			{
+				MainBoxPtr->AddSlot()
+					.AutoWidth()
+					[
+						SNew(SBox)
+						.WidthOverride(LocalSlateUnitsPerFrame)
+						.HAlign(HAlign_Center)
+						[
+							SNew(STextBlock)
+							.Text(FText::AsNumber(FrameIdx).ToString())
+						]
+					];
+			}
+
+			NumFramesFromLastRebuild = NumFrames;
+		}
+		else
+		{
+			NumFramesFromLastRebuild = 0;
+		}
+	}
+
+	FMargin GetPlayTimePadding() const
+	{
+		float PlayTimePadding = PlayTime.Get() - SlateUnitsPerFrame.Get() / 2;
+		return FMargin(PlayTimePadding, 0, 0, 0);
+	}
+
+
+private:
+	TAttribute<float> SlateUnitsPerFrame;
+	TAttribute<class UPaperFlipbook*> FlipbookBeingEdited;
+	TAttribute<float> PlayTime;
+
+	TSharedPtr<SHorizontalBox> MainBoxPtr;
+
+	int32 NumFramesFromLastRebuild;
+};
diff --git a/Engine/Plugins/Experimental/Paper2D/Source/Paper2DEditor/Private/FlipbookEditor/STimelineTrack.h b/Engine/Plugins/Experimental/Paper2D/Source/Paper2DEditor/Private/FlipbookEditor/STimelineTrack.h
new file mode 100644
--- /dev/null
+++ b/Engine/Plugins/Experimental/Paper2D/Source/Paper2DEditor/Private/FlipbookEditor/STimelineTrack.h
@@ -0,0 +1,113 @@
+// Copyright 1998-2014 Epic Games, Inc. All Rights Reserved.
+
+//////////////////////////////////////////////////////////////////////////
+// STimelineTrack
+
+class STimelineTrack : public SCompoundWidget
+{
+public:
+	SLATE_BEGIN_ARGS(STimelineTrack)
+		: _SlateUnitsPerFrame(1)
+		, _FlipbookBeingEdited(nullptr)
+	{}
+
+		SLATE_ATTRIBUTE( float, SlateUnitsPerFrame )
+		SLATE_ATTRIBUTE( class UPaperFlipbook*, FlipbookBeingEdited )
+
+	SLATE_END_ARGS()
+
+	void Construct(const FArguments& InArgs)
+	{
+		SlateUnitsPerFrame = InArgs._SlateUnitsPerFrame;
+		FlipbookBeingEdited = InArgs._FlipbookBeingEdited;
+
+		HandleWidth = 4;
+		NumKeyframesFromLastRebuild = 0;
+
+		ChildSlot
+		[
+			SAssignNew(MainBoxPtr, SHorizontalBox)
+		];
+
+		Rebuild();
+	}
+
+	void Tick(const FGeometry& AllottedGeometry, const double InCurrentTime, const float InDeltaTime) override
+	{
+		SCompoundWidget::Tick(AllottedGeometry, InCurrentTime, InDeltaTime);
+
+		UPaperFlipbook* Flipbook = FlipbookBeingEdited.Get();
+		int32 NewNumKeyframes = (Flipbook != NULL ? Flipbook->KeyFrames.Num() : 0);
+		if (NewNumKeyframes != NumKeyframesFromLastRebuild)
+		{
+			Rebuild();
+		}
+	}
+
+private:
+	void Rebuild()
+	{
+		MainBoxPtr->ClearChildren();
+
+		if (UPaperFlipbook* Flipbook = FlipbookBeingEdited.Get())
+		{
+			for (int32 KeyFrameIdx = 0; KeyFrameIdx < Flipbook->KeyFrames.Num(); ++KeyFrameIdx)
+			{
+				MainBoxPtr->AddSlot()
+				.AutoWidth()
+				[
+					SNew(SBox)
+					.WidthOverride( TAttribute<FOptionalSize>::Create( TAttribute<FOptionalSize>::FGetter::CreateSP(this, &STimelineTrack::GetFrameWidth, KeyFrameIdx) ) )
+					[
+						SNew(SBorder)
+						.BorderImage( FEditorStyle::GetBrush("PropertyTable.CurrentCellBorder") )
+						[
+							SNullWidget::NullWidget
+						]
+					]
+				];
+
+				MainBoxPtr->AddSlot()
+				.AutoWidth()
+				[
+					SNew(SBox).WidthOverride(HandleWidth)
+					[
+						SNew(SFlipbookTrackHandle)
+						.SlateUnitsPerFrame(SlateUnitsPerFrame)
+						.FlipbookBeingEdited(FlipbookBeingEdited)
+						.KeyFrameIdx(KeyFrameIdx)
+					]
+				];
+			}
+
+			NumKeyframesFromLastRebuild = Flipbook->KeyFrames.Num();
+		}
+		else
+		{
+			NumKeyframesFromLastRebuild = 0;
+		}
+	}
+
+	FOptionalSize GetFrameWidth(int32 KeyFrameIdx) const
+	{
+		UPaperFlipbook* Flipbook = FlipbookBeingEdited.Get();
+		if (Flipbook && Flipbook->KeyFrames.IsValidIndex(KeyFrameIdx))
+		{
+			const FPaperFlipbookKeyFrame& KeyFrame = Flipbook->KeyFrames[KeyFrameIdx];
+			return FMath::Max<float>(0, KeyFrame.FrameRun * SlateUnitsPerFrame.Get() - HandleWidth);
+		}
+		else
+		{
+			return 1;
+		}
+	}
+
+private:
+	TAttribute<float> SlateUnitsPerFrame;
+	TAttribute< class UPaperFlipbook* > FlipbookBeingEdited;
+
+	TSharedPtr<SHorizontalBox> MainBoxPtr;
+
+	int32 NumKeyframesFromLastRebuild;
+	float HandleWidth;
+};
