Paper2D: Add style brush for flipbook timeline grab handle and fix track checkerboard background

diff --git a/Engine/Plugins/Experimental/Paper2D/Source/Paper2DEditor/Private/FlipbookEditor/SFlipbookTrackHandle.h b/Engine/Plugins/Experimental/Paper2D/Source/Paper2DEditor/Private/FlipbookEditor/SFlipbookTrackHandle.h
--- a/Engine/Plugins/Experimental/Paper2D/Source/Paper2DEditor/Private/FlipbookEditor/SFlipbookTrackHandle.h
+++ b/Engine/Plugins/Experimental/Paper2D/Source/Paper2DEditor/Private/FlipbookEditor/SFlipbookTrackHandle.h
@@ -3,6 +3,7 @@
 //////////////////////////////////////////////////////////////////////////
 // SFlipbookTrackHandle
 
+// This is the grab handle at the end of a frame region, which can be dragged to change the duration
 class SFlipbookTrackHandle : public SImage
 {
 public:
@@ -28,7 +29,9 @@ public:
 		bDragging = false;
 		StartingFrameRun = INDEX_NONE;
 
-		SImage::Construct(SImage::FArguments());
+		SImage::Construct(
+			SImage::FArguments()
+			.Image(FEditorStyle::GetBrush("FlipbookEditor.RegionGrabHandle")));
 	}
 
 	virtual FReply OnMouseButtonDown(const FGeometry& MyGeometry, const FPointerEvent& MouseEvent) override
diff --git a/Engine/Plugins/Experimental/Paper2D/Source/Paper2DEditor/Private/FlipbookEditor/STimelineHeader.h b/Engine/Plugins/Experimental/Paper2D/Source/Paper2DEditor/Private/FlipbookEditor/STimelineHeader.h
--- a/Engine/Plugins/Experimental/Paper2D/Source/Paper2DEditor/Private/FlipbookEditor/STimelineHeader.h
+++ b/Engine/Plugins/Experimental/Paper2D/Source/Paper2DEditor/Private/FlipbookEditor/STimelineHeader.h
@@ -3,6 +3,7 @@
 //////////////////////////////////////////////////////////////////////////
 // STimelineHeader
 
+// This is the bar above the timeline which (will someday) shows the frame ticks and current time
 class STimelineHeader : public SCompoundWidget
 {
 public:
@@ -15,39 +16,38 @@ public:
 	SLATE_ATTRIBUTE(float, SlateUnitsPerFrame)
 		SLATE_ATTRIBUTE(class UPaperFlipbook*, FlipbookBeingEdited)
 		SLATE_ATTRIBUTE(float, PlayTime)
+	SLATE_END_ARGS()
 
-		SLATE_END_ARGS()
-
-		void Construct(const FArguments& InArgs)
+	void Construct(const FArguments& InArgs)
 	{
-			SlateUnitsPerFrame = InArgs._SlateUnitsPerFrame;
-			FlipbookBeingEdited = InArgs._FlipbookBeingEdited;
-			PlayTime = InArgs._PlayTime;
-
-			NumFramesFromLastRebuild = 0;
-
-			ChildSlot
+		SlateUnitsPerFrame = InArgs._SlateUnitsPerFrame;
+		FlipbookBeingEdited = InArgs._FlipbookBeingEdited;
+		PlayTime = InArgs._PlayTime;
+
+		NumFramesFromLastRebuild = 0;
+
+		ChildSlot
+		[
+			SNew(SOverlay)
+
+			+ SOverlay::Slot()
+			[
+				SAssignNew(MainBoxPtr, SHorizontalBox)
+			]
+
+			+ SOverlay::Slot()
+			.Padding(TAttribute<FMargin>(this, &STimelineHeader::GetPlayTimePadding))
+			[
+				SNew(SBox)
+				.WidthOverride(SlateUnitsPerFrame.Get())
 				[
-					SNew(SOverlay)
-
-					+ SOverlay::Slot()
-					[
-						SAssignNew(MainBoxPtr, SHorizontalBox)
-					]
+					SNew(SSpacer)
+				]
+			]
+		];
 
-					+ SOverlay::Slot()
-						.Padding(TAttribute<FMargin>(this, &STimelineHeader::GetPlayTimePadding))
-						[
-							SNew(SBox)
-							.WidthOverride(SlateUnitsPerFrame.Get())
-							[
-								SNew(SImage)
-							]
-						]
-				];
-
-			Rebuild();
-		}
+		Rebuild();
+	}
 
 	void Tick(const FGeometry& AllottedGeometry, const double InCurrentTime, const float InDeltaTime) override
 	{
diff --git a/Engine/Plugins/Experimental/Paper2D/Source/Paper2DEditor/Private/FlipbookEditor/STimelineTrack.h b/Engine/Plugins/Experimental/Paper2D/Source/Paper2DEditor/Private/FlipbookEditor/STimelineTrack.h
--- a/Engine/Plugins/Experimental/Paper2D/Source/Paper2DEditor/Private/FlipbookEditor/STimelineTrack.h
+++ b/Engine/Plugins/Experimental/Paper2D/Source/Paper2DEditor/Private/FlipbookEditor/STimelineTrack.h
@@ -21,7 +21,7 @@ public:
 		SlateUnitsPerFrame = InArgs._SlateUnitsPerFrame;
 		FlipbookBeingEdited = InArgs._FlipbookBeingEdited;
 
-		HandleWidth = 4;
+		HandleWidth = 12;
 		NumKeyframesFromLastRebuild = 0;
 
 		ChildSlot
diff --git a/Engine/Plugins/Experimental/Paper2D/Source/Paper2DEditor/Private/PaperStyle.cpp b/Engine/Plugins/Experimental/Paper2D/Source/Paper2DEditor/Private/PaperStyle.cpp
--- a/Engine/Plugins/Experimental/Paper2D/Source/Paper2DEditor/Private/PaperStyle.cpp
+++ b/Engine/Plugins/Experimental/Paper2D/Source/Paper2DEditor/Private/PaperStyle.cpp
@@ -5,7 +5,8 @@
 #include "SlateStyle.h"
 #include "EditorStyle.h"
 
-#define IMAGE_BRUSH( RelativePath, ... ) FSlateImageBrush( Style.RootToContentDir( RelativePath, TEXT(".png") ), __VA_ARGS__ )
+#define IMAGE_BRUSH(RelativePath, ...) FSlateImageBrush(Style.RootToContentDir(RelativePath, TEXT(".png")), __VA_ARGS__)
+#define BOX_BRUSH(RelativePath, ...) FSlateBoxBrush(Style.RootToContentDir(RelativePath, TEXT(".png")), __VA_ARGS__)
 
 //////////////////////////////////////////////////////////////////////////
 // FPaperStyle
@@ -97,6 +98,8 @@ TSharedRef<FSlateStyleSet> FPaperStyle::Create()
 		Style.Set("FlipbookEditor.SetShowCollision.Small", new IMAGE_BRUSH(TEXT("Icons/icon_StaticMeshEd_Collision_40x"), Icon20x20));
 		Style.Set("FlipbookEditor.SetShowPivot", new IMAGE_BRUSH(TEXT("Icons/icon_StaticMeshEd_ShowPivot_40x"), Icon40x40));
 		Style.Set("FlipbookEditor.SetShowPivot.Small", new IMAGE_BRUSH(TEXT("Icons/icon_StaticMeshEd_ShowPivot_40x"), Icon20x20));
+
+		Style.Set("FlipbookEditor.RegionGrabHandle", new BOX_BRUSH("Sequencer/ScrubHandleWhole", FMargin(6.f / 13.f, 10 / 24.f, 6 / 13.f, 10 / 24.f)));
 	}
 
 	return StyleRef;
