Paper2D: Add UPaperFlipbook::GetSpriteAtFrame

diff --git a/Engine/Plugins/Experimental/Paper2D/Source/Paper2D/Classes/PaperFlipbook.h b/Engine/Plugins/Experimental/Paper2D/Source/Paper2D/Classes/PaperFlipbook.h
--- a/Engine/Plugins/Experimental/Paper2D/Source/Paper2D/Classes/PaperFlipbook.h
+++ b/Engine/Plugins/Experimental/Paper2D/Source/Paper2D/Classes/PaperFlipbook.h
@@ -41,6 +41,7 @@ class UPaperFlipbook : public UObject
 	PAPER2D_API int32 GetNumFrames() const;
 	PAPER2D_API float GetDuration() const;
 	PAPER2D_API UPaperSprite* GetSpriteAtTime(float Time) const;
+	PAPER2D_API UPaperSprite* GetSpriteAtFrame(int32 FrameIndex) const;
 
 	// UObject interface
 #if WITH_EDITORONLY_DATA
diff --git a/Engine/Plugins/Experimental/Paper2D/Source/Paper2D/Private/PaperFlipbook.cpp b/Engine/Plugins/Experimental/Paper2D/Source/Paper2D/Private/PaperFlipbook.cpp
--- a/Engine/Plugins/Experimental/Paper2D/Source/Paper2D/Private/PaperFlipbook.cpp
+++ b/Engine/Plugins/Experimental/Paper2D/Source/Paper2D/Private/PaperFlipbook.cpp
@@ -54,6 +54,11 @@ UPaperSprite* UPaperFlipbook::GetSpriteAtTime(float Time) const
 	return NULL;
 }
 
+UPaperSprite* UPaperFlipbook::GetSpriteAtFrame(int32 FrameIndex) const
+{
+	return KeyFrames.IsValidIndex(FrameIndex) ? KeyFrames[FrameIndex].Sprite : nullptr;
+}
+
 #if WITH_EDITORONLY_DATA
 void UPaperFlipbook::Serialize(FArchive& Ar)
 {
