Paper2D: Color border of keyframes based on whether there is a sprite or not

diff --git a/Engine/Plugins/Experimental/Paper2D/Source/Paper2DEditor/Private/FlipbookEditor/STimelineTrack.h b/Engine/Plugins/Experimental/Paper2D/Source/Paper2DEditor/Private/FlipbookEditor/STimelineTrack.h
--- a/Engine/Plugins/Experimental/Paper2D/Source/Paper2DEditor/Private/FlipbookEditor/STimelineTrack.h
+++ b/Engine/Plugins/Experimental/Paper2D/Source/Paper2DEditor/Private/FlipbookEditor/STimelineTrack.h
@@ -49,18 +49,28 @@ private:
 	{
 		MainBoxPtr->ClearChildren();
 
+		// Color each region based on whether a sprite has been set or not for it
+		const auto BorderColorDelegate = [](UPaperFlipbook* ThisFlipbook, int32 TestIndex) -> FSlateColor
+		{
+			const bool bFrameValid = ThisFlipbook->GetSpriteAtFrame(TestIndex) != nullptr;
+			return bFrameValid ? FLinearColor::White : FLinearColor::Black;
+		};
+
+		// Create the sections for each keyframe
 		if (UPaperFlipbook* Flipbook = FlipbookBeingEdited.Get())
 		{
 			for (int32 KeyFrameIdx = 0; KeyFrameIdx < Flipbook->KeyFrames.Num(); ++KeyFrameIdx)
 			{
 				MainBoxPtr->AddSlot()
 				.AutoWidth()
+				.Padding(FMargin(0.0f, 7.0f, 0.0f, 7.0f))
 				[
 					SNew(SBox)
 					.WidthOverride( TAttribute<FOptionalSize>::Create( TAttribute<FOptionalSize>::FGetter::CreateSP(this, &STimelineTrack::GetFrameWidth, KeyFrameIdx) ) )
 					[
 						SNew(SBorder)
-						.BorderImage( FEditorStyle::GetBrush("PropertyTable.CurrentCellBorder") )
+						.BorderImage(FEditorStyle::GetBrush("FlipbookEditor.RegionBorder"))
+						.BorderBackgroundColor_Static(BorderColorDelegate, Flipbook, KeyFrameIdx)
 						[
 							SNullWidget::NullWidget
 						]
diff --git a/Engine/Plugins/Experimental/Paper2D/Source/Paper2DEditor/Private/PaperStyle.cpp b/Engine/Plugins/Experimental/Paper2D/Source/Paper2DEditor/Private/PaperStyle.cpp
--- a/Engine/Plugins/Experimental/Paper2D/Source/Paper2DEditor/Private/PaperStyle.cpp
+++ b/Engine/Plugins/Experimental/Paper2D/Source/Paper2DEditor/Private/PaperStyle.cpp
@@ -100,6 +100,7 @@ TSharedRef<FSlateStyleSet> FPaperStyle::Create()
 		Style.Set("FlipbookEditor.SetShowPivot.Small", new IMAGE_BRUSH(TEXT("Icons/icon_StaticMeshEd_ShowPivot_40x"), Icon20x20));
 
 		Style.Set("FlipbookEditor.RegionGrabHandle", new BOX_BRUSH("Sequencer/ScrubHandleWhole", FMargin(6.f / 13.f, 10 / 24.f, 6 / 13.f, 10 / 24.f)));
+		Style.Set("FlipbookEditor.RegionBorder", new BOX_BRUSH("Common/CurrentCellBorder", FMargin(4.f / 16.f), FLinearColor(1.0f, 1.0f, 1.0f, 1.0f)));
 	}
 
 	return StyleRef;
