Paper2D: Add context menu to flipbook keyframes, currently allows deletion and duplication

diff --git a/Engine/Plugins/Experimental/Paper2D/Source/Paper2DEditor/Private/FlipbookEditor/FlipbookEditor.cpp b/Engine/Plugins/Experimental/Paper2D/Source/Paper2DEditor/Private/FlipbookEditor/FlipbookEditor.cpp
--- a/Engine/Plugins/Experimental/Paper2D/Source/Paper2DEditor/Private/FlipbookEditor/FlipbookEditor.cpp
+++ b/Engine/Plugins/Experimental/Paper2D/Source/Paper2DEditor/Private/FlipbookEditor/FlipbookEditor.cpp
@@ -206,8 +206,9 @@ TSharedRef<SDockTab> FFlipbookEditor::SpawnTab_Viewport(const FSpawnTabArgs& Arg
 			.Padding(0, 8, 0, 0)
 			.AutoHeight()
 			[
-				SNew(SFlipbookTimeline)
+				SNew(SFlipbookTimeline, GetToolkitCommands())
 				.FlipbookBeingEdited(this, &FFlipbookEditor::GetFlipbookBeingEdited)
+				.OnSelectionChanged(this, &FFlipbookEditor::SetSelection)
 			]
 		];
 }
@@ -254,6 +255,7 @@ void FFlipbookEditor::InitFlipbookEditor(const EToolkitMode::Type Mode, const TS
 	FAssetEditorManager::Get().CloseOtherEditors(InitFlipbook, this);
 	FlipbookBeingEdited = InitFlipbook;
 	PlayTime = 0;
+	CurrentSelectedKeyframe = INDEX_NONE;
 
 	FFlipbookEditorCommands::Register();
 
@@ -308,24 +310,17 @@ void FFlipbookEditor::InitFlipbookEditor(const EToolkitMode::Type Mode, const TS
 
 void FFlipbookEditor::BindCommands()
 {
-// 	const FFlipbookEditorCommands& Commands = FFlipbookEditorCommands::Get();
-// 
-// 	const TSharedRef<FUICommandList>& UICommandList = GetToolkitCommands();
-// 
-// 	UICommandList->MapAction( FGenericCommands::Get().Delete,
-// 		FExecuteAction::CreateSP( this, &FStaticMeshEditor::DeleteSelectedSockets ),
-// 		FCanExecuteAction::CreateSP( this, &FStaticMeshEditor::HasSelectedSockets ));
-// 
-// 	UICommandList->MapAction( FGenericCommands::Get().Undo, 
-// 		FExecuteAction::CreateSP( this, &FStaticMeshEditor::UndoAction ) );
-// 
-// 	UICommandList->MapAction( FGenericCommands::Get().Redo, 
-// 		FExecuteAction::CreateSP( this, &FStaticMeshEditor::RedoAction ) );
-// 
-// 	UICommandList->MapAction(
-// 		FGenericCommands::Get().Duplicate,
-// 		FExecuteAction::CreateSP(this, &FStaticMeshEditor::DuplicateSelectedSocket),
-// 		FCanExecuteAction::CreateSP(this, &FStaticMeshEditor::HasSelectedSockets));
+	const FFlipbookEditorCommands& Commands = FFlipbookEditorCommands::Get();
+
+	const TSharedRef<FUICommandList>& UICommandList = GetToolkitCommands();
+
+	UICommandList->MapAction(FGenericCommands::Get().Delete,
+		FExecuteAction::CreateSP(this, &FFlipbookEditor::DeleteSelection),
+		FCanExecuteAction::CreateSP(this, &FFlipbookEditor::HasValidSelection));
+
+	UICommandList->MapAction(FGenericCommands::Get().Duplicate,
+		FExecuteAction::CreateSP(this, &FFlipbookEditor::DuplicateSelection),
+		FCanExecuteAction::CreateSP(this, &FFlipbookEditor::HasValidSelection));
 }
 
 FName FFlipbookEditor::GetToolkitFName() const
@@ -411,6 +406,43 @@ void FFlipbookEditor::ExtendToolbar()
  	AddToolbarExtender(Paper2DEditorModule->GetFlipbookEditorToolBarExtensibilityManager()->GetAllExtenders());
 }
 
+void FFlipbookEditor::DeleteSelection()
+{
+	if (FlipbookBeingEdited->KeyFrames.IsValidIndex(CurrentSelectedKeyframe))
+	{
+		const FScopedTransaction Transaction(LOCTEXT("DeleteKeyframe", "Delete Keyframe"));
+		FlipbookBeingEdited->Modify();
+
+		FlipbookBeingEdited->KeyFrames.RemoveAt(CurrentSelectedKeyframe);
+
+		CurrentSelectedKeyframe = INDEX_NONE;
+	}
+}
+
+void FFlipbookEditor::DuplicateSelection()
+{
+	if (FlipbookBeingEdited->KeyFrames.IsValidIndex(CurrentSelectedKeyframe))
+	{
+		const FScopedTransaction Transaction(LOCTEXT("DuplicateKeyframe", "Duplicate Keyframe"));
+		FlipbookBeingEdited->Modify();
+
+		FPaperFlipbookKeyFrame NewFrame = FlipbookBeingEdited->KeyFrames[CurrentSelectedKeyframe];
+		FlipbookBeingEdited->KeyFrames.Insert(NewFrame, CurrentSelectedKeyframe);
+
+		CurrentSelectedKeyframe = INDEX_NONE;
+	}
+}
+
+void FFlipbookEditor::SetSelection(int32 NewSelection)
+{
+	CurrentSelectedKeyframe = NewSelection;
+}
+
+bool FFlipbookEditor::HasValidSelection() const
+{
+	return FlipbookBeingEdited->KeyFrames.IsValidIndex(CurrentSelectedKeyframe);
+}
+
 //////////////////////////////////////////////////////////////////////////
 
 #undef LOCTEXT_NAMESPACE
\ No newline at end of file
diff --git a/Engine/Plugins/Experimental/Paper2D/Source/Paper2DEditor/Private/FlipbookEditor/FlipbookEditor.h b/Engine/Plugins/Experimental/Paper2D/Source/Paper2DEditor/Private/FlipbookEditor/FlipbookEditor.h
--- a/Engine/Plugins/Experimental/Paper2D/Source/Paper2DEditor/Private/FlipbookEditor/FlipbookEditor.h
+++ b/Engine/Plugins/Experimental/Paper2D/Source/Paper2DEditor/Private/FlipbookEditor/FlipbookEditor.h
@@ -37,6 +37,7 @@ protected:
 	UPaperFlipbook* FlipbookBeingEdited;
 	TSharedPtr<class SFlipbookEditorViewport> ViewportPtr;
 	float PlayTime;
+	int32 CurrentSelectedKeyframe;
 
 protected:
 	void BindCommands();
@@ -47,4 +48,9 @@ protected:
 
 	TSharedRef<SDockTab> SpawnTab_Viewport(const FSpawnTabArgs& Args);
 	TSharedRef<SDockTab> SpawnTab_Details(const FSpawnTabArgs& Args);
+
+	void DeleteSelection();
+	void DuplicateSelection();
+	void SetSelection(int32 NewSelection);
+	bool HasValidSelection() const;
 };
\ No newline at end of file
diff --git a/Engine/Plugins/Experimental/Paper2D/Source/Paper2DEditor/Private/FlipbookEditor/FlipbookEditorCommands.h b/Engine/Plugins/Experimental/Paper2D/Source/Paper2DEditor/Private/FlipbookEditor/FlipbookEditorCommands.h
--- a/Engine/Plugins/Experimental/Paper2D/Source/Paper2DEditor/Private/FlipbookEditor/FlipbookEditorCommands.h
+++ b/Engine/Plugins/Experimental/Paper2D/Source/Paper2DEditor/Private/FlipbookEditor/FlipbookEditorCommands.h
@@ -26,4 +26,6 @@ public:
 
 	// View Menu Commands
 	TSharedPtr<FUICommandInfo> SetShowPivot;
+
+	// 
 };
\ No newline at end of file
diff --git a/Engine/Plugins/Experimental/Paper2D/Source/Paper2DEditor/Private/FlipbookEditor/SFlipbookTimeline.cpp b/Engine/Plugins/Experimental/Paper2D/Source/Paper2DEditor/Private/FlipbookEditor/SFlipbookTimeline.cpp
--- a/Engine/Plugins/Experimental/Paper2D/Source/Paper2DEditor/Private/FlipbookEditor/SFlipbookTimeline.cpp
+++ b/Engine/Plugins/Experimental/Paper2D/Source/Paper2DEditor/Private/FlipbookEditor/SFlipbookTimeline.cpp
@@ -3,6 +3,11 @@
 #include "Paper2DEditorPrivatePCH.h"
 #include "SFlipbookTimeline.h"
 
+#define LOCTEXT_NAMESPACE "FlipbookEditor"
+
+//////////////////////////////////////////////////////////////////////////
+// Inline widgets
+
 #include "SFlipbookTrackHandle.h"
 #include "STimelineHeader.h"
 #include "STimelineTrack.h"
@@ -10,10 +15,11 @@
 //////////////////////////////////////////////////////////////////////////
 // SFlipbookTimeline
 
-void SFlipbookTimeline::Construct(const FArguments& InArgs)
+void SFlipbookTimeline::Construct(const FArguments& InArgs, TSharedPtr<const FUICommandList> InCommandList)
 {
 	FlipbookBeingEdited = InArgs._FlipbookBeingEdited;
 	PlayTime = InArgs._PlayTime;
+	OnSelectionChanged = InArgs._OnSelectionChanged;
 
 	const int32 SlateUnitsPerFrame = 32;
 	const int32 FrameHeight = 48;
@@ -40,11 +46,16 @@ void SFlipbookTimeline::Construct(const FArguments& InArgs)
 			[
 				SNew(SBox).HeightOverride(FrameHeight)
 				[
-					SNew(STimelineTrack)
+					SNew(STimelineTrack, InCommandList)
 					.SlateUnitsPerFrame(SlateUnitsPerFrame)
 					.FlipbookBeingEdited(FlipbookBeingEdited)
+					.OnSelectionChanged(OnSelectionChanged)
 				]
 			]
 		]
 	];
-}
\ No newline at end of file
+}
+
+//////////////////////////////////////////////////////////////////////////
+
+#undef LOCTEXT_NAMESPACE
\ No newline at end of file
diff --git a/Engine/Plugins/Experimental/Paper2D/Source/Paper2DEditor/Private/FlipbookEditor/SFlipbookTimeline.h b/Engine/Plugins/Experimental/Paper2D/Source/Paper2DEditor/Private/FlipbookEditor/SFlipbookTimeline.h
--- a/Engine/Plugins/Experimental/Paper2D/Source/Paper2DEditor/Private/FlipbookEditor/SFlipbookTimeline.h
+++ b/Engine/Plugins/Experimental/Paper2D/Source/Paper2DEditor/Private/FlipbookEditor/SFlipbookTimeline.h
@@ -4,6 +4,9 @@
 
 struct FSpriteAnimationFrame;
 
+// Called when the selection changes
+DECLARE_DELEGATE_OneParam(FOnFlipbookKeyframeSelectionChanged, int32);
+
 class SFlipbookTimeline : public SCompoundWidget
 {
 public:
@@ -14,12 +17,13 @@ public:
 
 		SLATE_ATTRIBUTE(class UPaperFlipbook*, FlipbookBeingEdited)
 		SLATE_ATTRIBUTE(float, PlayTime)
-
+		SLATE_EVENT(FOnFlipbookKeyframeSelectionChanged, OnSelectionChanged)
 	SLATE_END_ARGS()
 
-	void Construct(const FArguments& InArgs);
+	void Construct(const FArguments& InArgs, TSharedPtr<const FUICommandList> InCommandList);
 
 private:
 	TAttribute<class UPaperFlipbook*> FlipbookBeingEdited;
 	TAttribute<float> PlayTime;
+	FOnFlipbookKeyframeSelectionChanged OnSelectionChanged;
 };
diff --git a/Engine/Plugins/Experimental/Paper2D/Source/Paper2DEditor/Private/FlipbookEditor/STimelineTrack.h b/Engine/Plugins/Experimental/Paper2D/Source/Paper2DEditor/Private/FlipbookEditor/STimelineTrack.h
--- a/Engine/Plugins/Experimental/Paper2D/Source/Paper2DEditor/Private/FlipbookEditor/STimelineTrack.h
+++ b/Engine/Plugins/Experimental/Paper2D/Source/Paper2DEditor/Private/FlipbookEditor/STimelineTrack.h
@@ -13,13 +13,17 @@ public:
 
 		SLATE_ATTRIBUTE( float, SlateUnitsPerFrame )
 		SLATE_ATTRIBUTE( class UPaperFlipbook*, FlipbookBeingEdited )
+		SLATE_EVENT( FOnFlipbookKeyframeSelectionChanged, OnSelectionChanged )
 
 	SLATE_END_ARGS()
 
-	void Construct(const FArguments& InArgs)
+	void Construct(const FArguments& InArgs, TSharedPtr<const FUICommandList> InCommandList)
 	{
+		CommandList = InCommandList;
 		SlateUnitsPerFrame = InArgs._SlateUnitsPerFrame;
 		FlipbookBeingEdited = InArgs._FlipbookBeingEdited;
+		SelectedFrame = INDEX_NONE;
+		OnSelectionChanged = InArgs._OnSelectionChanged;
 
 		HandleWidth = 12;
 		NumKeyframesFromLastRebuild = 0;
@@ -44,15 +48,29 @@ public:
 		}
 	}
 
+	FReply KeyframeOnMouseButtonUp(const FGeometry& MyGeometry, const FPointerEvent& MouseEvent, int32 FrameIndex)
+	{
+		if (MouseEvent.GetEffectingButton() == EKeys::RightMouseButton)
+		{
+			TSharedRef<SWidget> MenuContents = GenerateContextMenu(FrameIndex);
+			FSlateApplication::Get().PushMenu(AsShared(), MenuContents, MouseEvent.GetScreenSpacePosition(), FPopupTransitionEffect(FPopupTransitionEffect::ContextMenu));
+
+			return FReply::Handled();
+		}
+
+		return FReply::Unhandled();
+	}
+
 private:
 	void Rebuild()
 	{
 		MainBoxPtr->ClearChildren();
 
 		// Color each region based on whether a sprite has been set or not for it
-		const auto BorderColorDelegate = [](UPaperFlipbook* ThisFlipbook, int32 TestIndex) -> FSlateColor
+		const auto BorderColorDelegate = [](TAttribute<UPaperFlipbook*> ThisFlipbookPtr, int32 TestIndex) -> FSlateColor
 		{
-			const bool bFrameValid = ThisFlipbook->GetSpriteAtFrame(TestIndex) != nullptr;
+			UPaperFlipbook* FlipbookPtr = ThisFlipbookPtr.Get();
+			const bool bFrameValid = (FlipbookPtr != nullptr) && (FlipbookPtr->GetSpriteAtFrame(TestIndex) != nullptr);
 			return bFrameValid ? FLinearColor::White : FLinearColor::Black;
 		};
 
@@ -70,7 +88,8 @@ private:
 					[
 						SNew(SBorder)
 						.BorderImage(FEditorStyle::GetBrush("FlipbookEditor.RegionBorder"))
-						.BorderBackgroundColor_Static(BorderColorDelegate, Flipbook, KeyFrameIdx)
+						.BorderBackgroundColor_Static(BorderColorDelegate, FlipbookBeingEdited, KeyFrameIdx)
+						.OnMouseButtonUp(this, &STimelineTrack::KeyframeOnMouseButtonUp, KeyFrameIdx)
 						[
 							SNullWidget::NullWidget
 						]
@@ -112,7 +131,26 @@ private:
 		}
 	}
 
+	TSharedRef<SWidget> GenerateContextMenu(int32 FrameIndex)
+	{
+		SelectedFrame = FrameIndex;
+		OnSelectionChanged.ExecuteIfBound(SelectedFrame);
+
+		FMenuBuilder MenuBuilder(true, CommandList);
+		MenuBuilder.BeginSection("KeyframeActions", LOCTEXT("KeyframeActionsSectionHeader", "Keyframe Actions"));
+
+// 		MenuBuilder.AddMenuEntry(FGenericCommands::Get().Cut);
+// 		MenuBuilder.AddMenuEntry(FGenericCommands::Get().Copy);
+// 		MenuBuilder.AddMenuEntry(FGenericCommands::Get().Paste);
+		MenuBuilder.AddMenuEntry(FGenericCommands::Get().Duplicate);
+		MenuBuilder.AddMenuEntry(FGenericCommands::Get().Delete);
+
+		MenuBuilder.EndSection();
+
+		return MenuBuilder.MakeWidget();
+	}
 private:
+	int32 SelectedFrame;
 	TAttribute<float> SlateUnitsPerFrame;
 	TAttribute< class UPaperFlipbook* > FlipbookBeingEdited;
 
@@ -120,4 +158,7 @@ private:
 
 	int32 NumKeyframesFromLastRebuild;
 	float HandleWidth;
+
+	FOnFlipbookKeyframeSelectionChanged OnSelectionChanged;
+	TSharedPtr<const FUICommandList> CommandList;
 };
