Paper2D: Rename pass to better match engine standards (phase 1)

- Rename UPaperRenderComponent to UPaperSpriteComponent
- Rename UPaperAnimatedRenderComponent to UPaperFlipbookComponent
- Rename FPaperAnimatedRenderSceneProxy to FPaperFlipbookSceneProxy

- Fix APaperCharacter to properly prevent inherited Mesh component from being created

diff --git a/Engine/Config/BaseEngine.ini b/Engine/Config/BaseEngine.ini
--- a/Engine/Config/BaseEngine.ini
+++ b/Engine/Config/BaseEngine.ini
@@ -650,6 +650,10 @@ bUseStreamingPause=false
 ; Paper2D modifications (4.2 era)
 +TaggedPropertyRedirects=(ClassName="PaperRenderComponent",OldPropertyName="TestMaterial",NewPropertyName="MaterialOverride")
 
+; Paper2D modifications (4.3 era)
++ActiveClassRedirects=(OldClassName="PaperRenderComponent",NewClassName="/Script/Paper2D.PaperSpriteComponent")
++ActiveClassRedirects=(OldClassName="PaperAnimatedRenderComponent",NewClassName="/Script/Paper2D.PaperFlipbookComponent")
+
 ; collision response container refactor
 +TaggedPropertyRedirects=(ClassName="CollisionResponseContainer",OldPropertyName="RigidBody",NewPropertyName="PhysicsBody")
 +TaggedPropertyRedirects=(ClassName="CollisionResponseContainer",OldPropertyName="Static",NewPropertyName="WorldStatic")
diff --git a/Engine/Plugins/Experimental/Paper2D/Source/Paper2D/Classes/PaperAnimatedRenderComponent.h b/Engine/Plugins/Experimental/Paper2D/Source/Paper2D/Classes/PaperAnimatedRenderComponent.h
deleted file mode 100644
diff --git a/Engine/Plugins/Experimental/Paper2D/Source/Paper2D/Classes/PaperCharacter.h b/Engine/Plugins/Experimental/Paper2D/Source/Paper2D/Classes/PaperCharacter.h
--- a/Engine/Plugins/Experimental/Paper2D/Source/Paper2D/Classes/PaperCharacter.h
+++ b/Engine/Plugins/Experimental/Paper2D/Source/Paper2D/Classes/PaperCharacter.h
@@ -4,23 +4,17 @@
 
 #include "PaperCharacter.generated.h"
 
-class UPaperAnimatedRenderComponent;
-
-//=============================================================================
-// Paper Characters Exactly like Characters but are sprite based.
-//=============================================================================
-UCLASS(MinimalAPI)
+// APaperCharacter behaves like ACharacter, but uses a UPaperFlipbookComponent instead of a USkeletalMeshComponent as a visual representation
+// Note: The variable named Mesh will not be set up on this actor!
+UCLASS(MinimalAPI, EarlyAccessPreview)
 class APaperCharacter : public ACharacter
 {
 	GENERATED_UCLASS_BODY()
 	
+	// Name of the Sprite component
 	static FName SpriteComponentName;
+
 	/** The main skeletal mesh associated with this Character (optional sub-object). */
 	UPROPERTY(Category = Character, VisibleAnywhere, BlueprintReadOnly)
-	TSubobjectPtr<class UPaperAnimatedRenderComponent> Sprite;
-
-protected:
-
-public:	
-
+	TSubobjectPtr<class UPaperFlipbookComponent> Sprite;
 };
diff --git a/Engine/Plugins/Experimental/Paper2D/Source/Paper2D/Classes/PaperFlipbookActor.h b/Engine/Plugins/Experimental/Paper2D/Source/Paper2D/Classes/PaperFlipbookActor.h
--- a/Engine/Plugins/Experimental/Paper2D/Source/Paper2D/Classes/PaperFlipbookActor.h
+++ b/Engine/Plugins/Experimental/Paper2D/Source/Paper2D/Classes/PaperFlipbookActor.h
@@ -2,7 +2,7 @@
 
 #pragma once
 
-#include "PaperAnimatedRenderComponent.h"
+#include "PaperFlipbookComponent.h"
 #include "PaperFlipbookActor.generated.h"
 
 UCLASS(MinimalAPI)
@@ -11,7 +11,7 @@ class APaperFlipbookActor : public AActor
 	GENERATED_UCLASS_BODY()
 
 	UPROPERTY(Category=Sprite, VisibleAnywhere)
-	TSubobjectPtr<class UPaperAnimatedRenderComponent> RenderComponent;
+	TSubobjectPtr<class UPaperFlipbookComponent> RenderComponent;
 
 	// AActor interface
 #if WITH_EDITOR
diff --git a/Engine/Plugins/Experimental/Paper2D/Source/Paper2D/Classes/PaperFlipbookComponent.h b/Engine/Plugins/Experimental/Paper2D/Source/Paper2D/Classes/PaperFlipbookComponent.h
new file mode 100644
--- /dev/null
+++ b/Engine/Plugins/Experimental/Paper2D/Source/Paper2D/Classes/PaperFlipbookComponent.h
@@ -0,0 +1,67 @@
+// Copyright 1998-2014 Epic Games, Inc. All Rights Reserved.
+
+#pragma once
+
+#include "PaperFlipbook.h"
+#include "PaperFlipbookComponent.generated.h"
+
+UCLASS(MinimalAPI, ShowCategories=(Mobility), EarlyAccessPreview, meta=(BlueprintSpawnableComponent))
+class UPaperFlipbookComponent : public UPrimitiveComponent
+{
+	GENERATED_UCLASS_BODY()
+
+protected:
+	UPROPERTY(Category=Sprite, EditAnywhere, meta=(DisplayThumbnail = "true"))
+	UPaperFlipbook* SourceFlipbook;
+
+	UPROPERTY(Category=Sprite, EditAnywhere)
+	UMaterialInterface* Material;
+
+	UPROPERTY(Category=Sprite, EditAnywhere, BlueprintReadWrite)
+	float PlayRate;
+public:
+	/** Change the flipbook used by this instance. */
+	UFUNCTION(BlueprintCallable, Category="Sprite")
+	virtual bool SetFlipbook(class UPaperFlipbook* NewFlipbook);
+
+	/** Gets the flipbook used by this instance. */
+	UFUNCTION(BlueprintPure, Category="Sprite")
+	virtual UPaperFlipbook* GetFlipbook();
+
+	/** Gets the material used by this instance */
+	UFUNCTION(BlueprintCallable, Category="Sprite")
+	UMaterialInterface* GetSpriteMaterial() const;
+
+	/** Set color of the sprite */
+	UFUNCTION(BlueprintCallable, Category="Sprite")
+	void SetSpriteColor(FLinearColor NewColor);
+
+	//@TODO: Hackery for demo
+	UFUNCTION(BlueprintCallable, Category="Sprite")
+	void SetCurrentTime(float NewTime);
+
+protected:
+	float AccumulatedTime;
+	int32 CachedFrameIndex;
+
+	UPROPERTY(BlueprintReadOnly, Interp, Category=Sprite)
+	FLinearColor SpriteColor;
+
+protected:
+	void CalculateCurrentFrame();
+public:
+	// UActorComponent interface
+	virtual void TickComponent(float DeltaTime, enum ELevelTick TickType, FActorComponentTickFunction* ThisTickFunction) override;
+	virtual void SendRenderDynamicData_Concurrent() override;
+	virtual const UObject* AdditionalStatObject() const override;
+	// End of UActorComponent interface
+
+	// UPrimitiveComponent interface
+	virtual FPrimitiveSceneProxy* CreateSceneProxy() override;
+	virtual FBoxSphereBounds CalcBounds(const FTransform & LocalToWorld) const override;
+	// End of UPrimitiveComponent interface
+};
+
+// Allow the old name to continue to work for one release
+DEPRECATED(4.3, "UPaperAnimatedRenderComponent has been renamed to UPaperFlipbookComponent")
+typedef UPaperFlipbookComponent UPaperAnimatedRenderComponent;
diff --git a/Engine/Plugins/Experimental/Paper2D/Source/Paper2D/Classes/PaperRenderActor.h b/Engine/Plugins/Experimental/Paper2D/Source/Paper2D/Classes/PaperRenderActor.h
--- a/Engine/Plugins/Experimental/Paper2D/Source/Paper2D/Classes/PaperRenderActor.h
+++ b/Engine/Plugins/Experimental/Paper2D/Source/Paper2D/Classes/PaperRenderActor.h
@@ -2,7 +2,7 @@
 
 #pragma once
 
-#include "PaperRenderComponent.h"
+#include "PaperSpriteComponent.h"
 #include "PaperRenderActor.generated.h"
 
 UCLASS(MinimalAPI)
@@ -11,7 +11,7 @@ class APaperRenderActor : public AActor
 	GENERATED_UCLASS_BODY()
 
 	UPROPERTY(Category=Sprite, VisibleAnywhere)
-	TSubobjectPtr<class UPaperRenderComponent> RenderComponent;
+	TSubobjectPtr<class UPaperSpriteComponent> RenderComponent;
 
 	// AActor interface
 #if WITH_EDITOR
diff --git a/Engine/Plugins/Experimental/Paper2D/Source/Paper2D/Classes/PaperRenderComponent.h b/Engine/Plugins/Experimental/Paper2D/Source/Paper2D/Classes/PaperRenderComponent.h
deleted file mode 100644
diff --git a/Engine/Plugins/Experimental/Paper2D/Source/Paper2D/Classes/PaperSpriteComponent.h b/Engine/Plugins/Experimental/Paper2D/Source/Paper2D/Classes/PaperSpriteComponent.h
new file mode 100644
--- /dev/null
+++ b/Engine/Plugins/Experimental/Paper2D/Source/Paper2D/Classes/PaperSpriteComponent.h
@@ -0,0 +1,73 @@
+// Copyright 1998-2014 Epic Games, Inc. All Rights Reserved.
+
+#pragma once
+
+#include "PaperSprite.h"
+
+#include "PaperSpriteComponent.generated.h"
+
+UCLASS(MinimalAPI, ShowCategories=(Mobility), EarlyAccessPreview, meta=(BlueprintSpawnableComponent))
+class UPaperSpriteComponent : public UPrimitiveComponent
+{
+	GENERATED_UCLASS_BODY()
+
+protected:
+	// The sprite asset used by this component
+	UPROPERTY(Category=Sprite, EditAnywhere, BlueprintReadOnly, meta=(DisplayThumbnail = "true"))
+	UPaperSprite* SourceSprite;
+
+	// The material override for this sprite component (if any)
+	UPROPERTY(Category=Sprite, EditAnywhere, BlueprintReadOnly)
+	UMaterialInterface* MaterialOverride;
+
+	// The color of the sprite (passed to the sprite material as a vertex color)
+	UPROPERTY(BlueprintReadOnly, Interp, Category=Sprite)
+	FLinearColor SpriteColor;
+
+public:
+	/** Change the PaperSprite used by this instance. */
+	UFUNCTION(BlueprintCallable, Category="Sprite")
+	virtual bool SetSprite(class UPaperSprite* NewSprite);
+
+	/** Gets the PaperSprite used by this instance. */
+	UFUNCTION(BlueprintPure, Category="Sprite")
+	virtual UPaperSprite* GetSprite();
+
+	/** Set color of the sprite */
+	UFUNCTION(BlueprintCallable, Category="Sprite")
+	void SetSpriteColor(FLinearColor NewColor);
+
+	// Returns the wireframe color to use for this component.
+	FLinearColor GetWireframeColor() const;
+
+public:
+	// UObject interface
+#if WITH_EDITOR
+	virtual void PostEditChangeProperty(FPropertyChangedEvent& PropertyChangedEvent) override;
+#endif
+	// End of UObject interface
+
+	// UActorComponent interface
+	virtual void SendRenderDynamicData_Concurrent() override;
+	virtual const UObject* AdditionalStatObject() const override;
+	// End of UActorComponent interface
+
+	// USceneComponent interface
+	virtual bool HasAnySockets() const override;
+	virtual FTransform GetSocketTransform(FName InSocketName, ERelativeTransformSpace TransformSpace = RTS_World) const override;
+	virtual void QuerySupportedSockets(TArray<FComponentSocketDescription>& OutSockets) const override;
+	// End of USceneComponent interface
+
+	// UPrimitiveComponent interface
+	virtual FPrimitiveSceneProxy* CreateSceneProxy() override;
+	virtual FBoxSphereBounds CalcBounds(const FTransform & LocalToWorld) const override;
+	virtual class UBodySetup* GetBodySetup() override;
+	// End of UPrimitiveComponent interface
+
+protected:
+	friend class FPaperRenderSceneProxy;
+};
+
+// Allow the old name to continue to work for one release
+DEPRECATED(4.3, "UPaperRenderComponent has been renamed to UPaperSpriteComponent")
+typedef UPaperSpriteComponent UPaperRenderComponent;
diff --git a/Engine/Plugins/Experimental/Paper2D/Source/Paper2D/Private/PaperAnimatedRenderComponent.cpp b/Engine/Plugins/Experimental/Paper2D/Source/Paper2D/Private/PaperAnimatedRenderComponent.cpp
deleted file mode 100644
diff --git a/Engine/Plugins/Experimental/Paper2D/Source/Paper2D/Private/PaperAnimatedRenderSceneProxy.cpp b/Engine/Plugins/Experimental/Paper2D/Source/Paper2D/Private/PaperAnimatedRenderSceneProxy.cpp
deleted file mode 100644
diff --git a/Engine/Plugins/Experimental/Paper2D/Source/Paper2D/Private/PaperAnimatedRenderSceneProxy.h b/Engine/Plugins/Experimental/Paper2D/Source/Paper2D/Private/PaperAnimatedRenderSceneProxy.h
deleted file mode 100644
diff --git a/Engine/Plugins/Experimental/Paper2D/Source/Paper2D/Private/PaperCharacter.cpp b/Engine/Plugins/Experimental/Paper2D/Source/Paper2D/Private/PaperCharacter.cpp
--- a/Engine/Plugins/Experimental/Paper2D/Source/Paper2D/Private/PaperCharacter.cpp
+++ b/Engine/Plugins/Experimental/Paper2D/Source/Paper2D/Private/PaperCharacter.cpp
@@ -1,41 +1,34 @@
 // Copyright 1998-2014 Epic Games, Inc. All Rights Reserved.
 
-/*=============================================================================
-	PaperCharacter.cpp: APaperCharacter implementation
-	A character that uses a paper sprite.
-=============================================================================*/
-
 #include "Paper2DPrivatePCH.h"
 
-
 FName APaperCharacter::SpriteComponentName(TEXT("Sprite0"));
 
+//////////////////////////////////////////////////////////////////////////
+// APaperCharacter
+
 APaperCharacter::APaperCharacter(const class FPostConstructInitializeProperties& PCIP)
-	: Super(PCIP)
+	: Super(PCIP.DoNotCreateDefaultSubobject(ACharacter::MeshComponentName))
 {
-	// Prevent creation of mesh
-	PCIP.DoNotCreateDefaultSubobject(ACharacter::MeshComponentName);
-	
-	Sprite = PCIP.CreateOptionalDefaultSubobject<UPaperAnimatedRenderComponent>(this, APaperCharacter::SpriteComponentName);
+	// Try to create the sprite component
+	Sprite = PCIP.CreateOptionalDefaultSubobject<UPaperFlipbookComponent>(this, APaperCharacter::SpriteComponentName);
 	if (Sprite)
 	{
 		Sprite->AlwaysLoadOnClient = true;
 		Sprite->AlwaysLoadOnServer = true;
 		Sprite->bOwnerNoSee = false;
-		//Mesh->MeshComponentUpdateFlag = EMeshComponentUpdateFlag::AlwaysTickPose;
-		//Mesh->bCastDynamicShadow = true;
 		Sprite->bAffectDynamicIndirectLighting = true;
 		Sprite->PrimaryComponentTick.TickGroup = TG_PrePhysics;
+
 		// force tick after movement component updates
 		if (CharacterMovement)
 		{
 			Sprite->PrimaryComponentTick.AddPrerequisite(this, CharacterMovement->PrimaryComponentTick);
 		}
+
 		Sprite->AttachParent = CapsuleComponent;
 		static FName CollisionProfileName(TEXT("CharacterMesh"));
 		Sprite->SetCollisionProfileName(CollisionProfileName);
 		Sprite->bGenerateOverlapEvents = false;
 	}
 }
-
-
diff --git a/Engine/Plugins/Experimental/Paper2D/Source/Paper2D/Private/PaperFlipbookActor.cpp b/Engine/Plugins/Experimental/Paper2D/Source/Paper2D/Private/PaperFlipbookActor.cpp
--- a/Engine/Plugins/Experimental/Paper2D/Source/Paper2D/Private/PaperFlipbookActor.cpp
+++ b/Engine/Plugins/Experimental/Paper2D/Source/Paper2D/Private/PaperFlipbookActor.cpp
@@ -9,7 +9,7 @@
 APaperFlipbookActor::APaperFlipbookActor(const FPostConstructInitializeProperties& PCIP)
 	: Super(PCIP)
 {
-	RenderComponent = PCIP.CreateDefaultSubobject<UPaperAnimatedRenderComponent>(this, TEXT("RenderComponent"));
+	RenderComponent = PCIP.CreateDefaultSubobject<UPaperFlipbookComponent>(this, TEXT("RenderComponent"));
 
 	RootComponent = RenderComponent;
 }
diff --git a/Engine/Plugins/Experimental/Paper2D/Source/Paper2D/Private/PaperFlipbookComponent.cpp b/Engine/Plugins/Experimental/Paper2D/Source/Paper2D/Private/PaperFlipbookComponent.cpp
new file mode 100644
--- /dev/null
+++ b/Engine/Plugins/Experimental/Paper2D/Source/Paper2D/Private/PaperFlipbookComponent.cpp
@@ -0,0 +1,192 @@
+// Copyright 1998-2014 Epic Games, Inc. All Rights Reserved.
+
+#include "Paper2DPrivatePCH.h"
+#include "PaperFlipbookSceneProxy.h"
+#include "PaperFlipbookComponent.h"
+
+//////////////////////////////////////////////////////////////////////////
+// UPaperFlipbookComponent
+
+UPaperFlipbookComponent::UPaperFlipbookComponent(const FPostConstructInitializeProperties& PCIP)
+	: Super(PCIP)
+{
+	BodyInstance.SetCollisionEnabled(ECollisionEnabled::NoCollision);
+
+	SetCollisionProfileName(UCollisionProfile::NoCollision_ProfileName);
+	static ConstructorHelpers::FObjectFinder<UMaterial> DefaultMaterial(TEXT("/Paper2D/DefaultSpriteMaterial"));
+	Material = DefaultMaterial.Object;
+
+	SpriteColor = FLinearColor::White;
+
+	Mobility = EComponentMobility::Movable;
+	PrimaryComponentTick.bCanEverTick = true;
+	PrimaryComponentTick.TickGroup = TG_DuringPhysics;
+	bTickInEditor = true;
+	
+	CachedFrameIndex = INDEX_NONE;
+	AccumulatedTime = 0.0f;
+	PlayRate = 1.0f;
+}
+
+FPrimitiveSceneProxy* UPaperFlipbookComponent::CreateSceneProxy()
+{
+	FPaperFlipbookSceneProxy* NewProxy = new FPaperFlipbookSceneProxy(this);
+
+	CalculateCurrentFrame();
+	UPaperSprite* SpriteToSend = NULL;
+	if ((SourceFlipbook != NULL) && (CachedFrameIndex >= 0) && (CachedFrameIndex < SourceFlipbook->KeyFrames.Num()))
+	{
+		SpriteToSend = SourceFlipbook->KeyFrames[CachedFrameIndex].Sprite;
+	}
+
+	FSpriteDrawCallRecord DrawCall;
+	DrawCall.BuildFromSprite(SpriteToSend);
+	DrawCall.Color = SpriteColor;
+	NewProxy->SetDrawCall_RenderThread(DrawCall);
+	return NewProxy;
+}
+
+FBoxSphereBounds UPaperFlipbookComponent::CalcBounds(const FTransform & LocalToWorld) const
+{
+	const float NewScale = LocalToWorld.GetScale3D().GetMax() * 300.0f;//(Sprite ? (float)FMath::Max(Sprite->GetSizeX(),Sprite->GetSizeY()) : 1.0f);  //@TODO: Compute a realistic bound
+	return FBoxSphereBounds(LocalToWorld.GetTranslation(), FVector(NewScale,NewScale,NewScale), FMath::Sqrt(3.0f * FMath::Square(NewScale)));
+}
+
+void UPaperFlipbookComponent::CalculateCurrentFrame()
+{
+	if (SourceFlipbook != NULL && SourceFlipbook->FramesPerSecond > 0)
+	{
+		int32 LastCachedFrame = CachedFrameIndex;
+
+		float SumTime = 0.0f;
+		int32 FrameIndex;
+		for (FrameIndex = 0; FrameIndex < SourceFlipbook->KeyFrames.Num(); ++FrameIndex)
+		{
+			SumTime += SourceFlipbook->KeyFrames[FrameIndex].FrameRun / SourceFlipbook->FramesPerSecond;
+
+			if (AccumulatedTime < SumTime)
+			{
+				break;
+			}
+		}
+
+		if (AccumulatedTime >= SumTime)
+		{
+			AccumulatedTime = FMath::Fmod(AccumulatedTime, SumTime);
+			//@TODO: Could have gone farther than this!!!
+			CachedFrameIndex = 0;
+		}
+		else
+		{
+			CachedFrameIndex = FrameIndex;
+		}
+
+		if (CachedFrameIndex != LastCachedFrame)
+		{
+			// Indicate we need to send new dynamic data.
+			MarkRenderDynamicDataDirty();
+		}
+	}
+}
+
+void UPaperFlipbookComponent::TickComponent(float DeltaTime, enum ELevelTick TickType, FActorComponentTickFunction* ThisTickFunction)
+{
+	// Indicate we need to send new dynamic data.
+	MarkRenderDynamicDataDirty();
+	
+	AccumulatedTime += DeltaTime * PlayRate;
+	CalculateCurrentFrame();
+
+	//Super::TickComponent(DeltaTime, TickType, ThisTickFunction);
+}
+
+void UPaperFlipbookComponent::SendRenderDynamicData_Concurrent()
+{
+	if (SceneProxy != NULL)
+	{
+		UPaperSprite* SpriteToSend = NULL;
+		if ((SourceFlipbook != NULL) && (CachedFrameIndex >= 0) && (CachedFrameIndex < SourceFlipbook->KeyFrames.Num()))
+		{
+			SpriteToSend = SourceFlipbook->KeyFrames[CachedFrameIndex].Sprite;
+		}
+
+		FSpriteDrawCallRecord DrawCall;
+		DrawCall.BuildFromSprite(SpriteToSend);
+		DrawCall.Color = SpriteColor;
+
+		ENQUEUE_UNIQUE_RENDER_COMMAND_TWOPARAMETER(
+				FSendPaperRenderComponentDynamicData,
+				FPaperRenderSceneProxy*,InSceneProxy,(FPaperRenderSceneProxy*)SceneProxy,
+				FSpriteDrawCallRecord,InSpriteToSend,DrawCall,
+			{
+				InSceneProxy->SetDrawCall_RenderThread(InSpriteToSend);
+			});
+	}
+}
+
+bool UPaperFlipbookComponent::SetFlipbook(class UPaperFlipbook* NewFlipbook)
+{
+	if (NewFlipbook != SourceFlipbook)
+	{
+		// Don't allow changing the sprite if we are "static".
+		AActor* Owner = GetOwner();
+		if (!IsRegistered() || (Owner == NULL) || (Mobility != EComponentMobility::Static))
+		{
+			SourceFlipbook = NewFlipbook;
+
+			// We need to also reset the frame and time also
+			AccumulatedTime = 0.0f;
+			CalculateCurrentFrame();
+
+			// Need to send this to render thread at some point
+			MarkRenderStateDirty();
+
+			// Update physics representation right away
+			RecreatePhysicsState();
+
+			// Since we have new mesh, we need to update bounds
+			UpdateBounds();
+
+			return true;
+		}
+	}
+
+	return false;
+}
+
+UPaperFlipbook* UPaperFlipbookComponent::GetFlipbook()
+{
+	return SourceFlipbook;
+}
+
+UMaterialInterface* UPaperFlipbookComponent::GetSpriteMaterial() const
+{
+	return Material;
+}
+
+void UPaperFlipbookComponent::SetSpriteColor(FLinearColor NewColor)
+{
+	// Can't set color on a static component
+	if (!(IsRegistered() && (Mobility == EComponentMobility::Static)) && (SpriteColor != NewColor))
+	{
+		SpriteColor = NewColor;
+
+		//@TODO: Should we send immediately?
+		MarkRenderDynamicDataDirty();
+	}
+}
+
+void UPaperFlipbookComponent::SetCurrentTime(float NewTime)
+{
+	// Can't set color on a static component
+	if (!(IsRegistered() && (Mobility == EComponentMobility::Static)) && (NewTime != AccumulatedTime))
+	{
+		AccumulatedTime = NewTime;
+		CalculateCurrentFrame();
+	}
+}
+
+const UObject* UPaperFlipbookComponent::AdditionalStatObject() const
+{
+	return SourceFlipbook;
+}
diff --git a/Engine/Plugins/Experimental/Paper2D/Source/Paper2D/Private/PaperFlipbookSceneProxy.cpp b/Engine/Plugins/Experimental/Paper2D/Source/Paper2D/Private/PaperFlipbookSceneProxy.cpp
new file mode 100644
--- /dev/null
+++ b/Engine/Plugins/Experimental/Paper2D/Source/Paper2D/Private/PaperFlipbookSceneProxy.cpp
@@ -0,0 +1,20 @@
+// Copyright 1998-2014 Epic Games, Inc. All Rights Reserved.
+
+#include "Paper2DPrivatePCH.h"
+#include "PaperFlipbookSceneProxy.h"
+
+//////////////////////////////////////////////////////////////////////////
+// FPaperFlipbookSceneProxy
+
+FPaperFlipbookSceneProxy::FPaperFlipbookSceneProxy(const UPaperFlipbookComponent* InComponent)
+	: FPaperRenderSceneProxy(InComponent)
+{
+	//@TODO: PAPER2D: WireframeColor = RenderComp->GetWireframeColor();
+
+	Material = InComponent->GetSpriteMaterial();
+
+	if (Material)
+	{
+		MaterialRelevance = Material->GetRelevance();
+	}
+}
diff --git a/Engine/Plugins/Experimental/Paper2D/Source/Paper2D/Private/PaperFlipbookSceneProxy.h b/Engine/Plugins/Experimental/Paper2D/Source/Paper2D/Private/PaperFlipbookSceneProxy.h
new file mode 100644
--- /dev/null
+++ b/Engine/Plugins/Experimental/Paper2D/Source/Paper2D/Private/PaperFlipbookSceneProxy.h
@@ -0,0 +1,14 @@
+// Copyright 1998-2014 Epic Games, Inc. All Rights Reserved.
+
+#pragma once
+
+#include "PaperRenderSceneProxy.h"
+
+//////////////////////////////////////////////////////////////////////////
+// FPaperFlipbookSceneProxy
+
+class FPaperFlipbookSceneProxy : public FPaperRenderSceneProxy
+{
+public:
+	FPaperFlipbookSceneProxy(const UPaperFlipbookComponent* InComponent);
+};
diff --git a/Engine/Plugins/Experimental/Paper2D/Source/Paper2D/Private/PaperRenderActor.cpp b/Engine/Plugins/Experimental/Paper2D/Source/Paper2D/Private/PaperRenderActor.cpp
--- a/Engine/Plugins/Experimental/Paper2D/Source/Paper2D/Private/PaperRenderActor.cpp
+++ b/Engine/Plugins/Experimental/Paper2D/Source/Paper2D/Private/PaperRenderActor.cpp
@@ -9,7 +9,7 @@
 APaperRenderActor::APaperRenderActor(const FPostConstructInitializeProperties& PCIP)
 	: Super(PCIP)
 {
-	RenderComponent = PCIP.CreateDefaultSubobject<UPaperRenderComponent>(this, TEXT("RenderComponent"));
+	RenderComponent = PCIP.CreateDefaultSubobject<UPaperSpriteComponent>(this, TEXT("RenderComponent"));
 	RenderComponent->SetCollisionProfileName(UCollisionProfile::BlockAll_ProfileName);
 	RenderComponent->Mobility = EComponentMobility::Static;
 
diff --git a/Engine/Plugins/Experimental/Paper2D/Source/Paper2D/Private/PaperRenderComponent.cpp b/Engine/Plugins/Experimental/Paper2D/Source/Paper2D/Private/PaperRenderComponent.cpp
deleted file mode 100644
diff --git a/Engine/Plugins/Experimental/Paper2D/Source/Paper2D/Private/PaperRenderSceneProxy.cpp b/Engine/Plugins/Experimental/Paper2D/Source/Paper2D/Private/PaperRenderSceneProxy.cpp
--- a/Engine/Plugins/Experimental/Paper2D/Source/Paper2D/Private/PaperRenderSceneProxy.cpp
+++ b/Engine/Plugins/Experimental/Paper2D/Source/Paper2D/Private/PaperRenderSceneProxy.cpp
@@ -168,7 +168,7 @@ FPaperRenderSceneProxy::FPaperRenderSceneProxy(const UPrimitiveComponent* InComp
 	, WireframeColor(FLinearColor::White)
 	, CollisionResponse(InComponent->GetCollisionResponseToChannels())
 {
-	if (const UPaperRenderComponent* RenderComp = Cast<const UPaperRenderComponent>(InComponent))
+	if (const UPaperSpriteComponent* RenderComp = Cast<const UPaperSpriteComponent>(InComponent))
 	{
 		WireframeColor = RenderComp->GetWireframeColor();
 
@@ -216,7 +216,7 @@ void FPaperRenderSceneProxy::DrawDynamicElements(FPrimitiveDrawInterface* PDI, c
 
 	checkSlow(IsInRenderingThread());
 
-	if (SourceSprite != NULL)//const UPaperRenderComponent* TypedFoo = Cast<const UPaperRenderComponent>(OwnerComponent))
+	if (SourceSprite != NULL)
 	{
 		// Show 3D physics
 		if ((View->Family->EngineShowFlags.Collision /*@TODO: && bIsCollisionEnabled*/) && AllowDebugViewmodes())
diff --git a/Engine/Plugins/Experimental/Paper2D/Source/Paper2D/Private/PaperSprite.cpp b/Engine/Plugins/Experimental/Paper2D/Source/Paper2D/Private/PaperSprite.cpp
--- a/Engine/Plugins/Experimental/Paper2D/Source/Paper2D/Private/PaperSprite.cpp
+++ b/Engine/Plugins/Experimental/Paper2D/Source/Paper2D/Private/PaperSprite.cpp
@@ -276,9 +276,9 @@ void UPaperSprite::PostEditChangeProperty(FPropertyChangedEvent& PropertyChanged
 		PixelsPerUnrealUnit = 1.0f;
 	}
 
-	//@TODO: Determine when these are really needed!
-	TComponentReregisterContext<UPaperRenderComponent> ReregisterStaticComponents;
-	TComponentReregisterContext<UPaperAnimatedRenderComponent> ReregisterAnimatedComponents;
+	//@TODO: Determine when these are really needed, as they're seriously expensive!
+	TComponentReregisterContext<UPaperSpriteComponent> ReregisterStaticComponents;
+	TComponentReregisterContext<UPaperFlipbookComponent> ReregisterAnimatedComponents;
 
 	// Update the pivot
 	if (PivotMode != ESpritePivotMode::Custom)
diff --git a/Engine/Plugins/Experimental/Paper2D/Source/Paper2D/Private/PaperSpriteComponent.cpp b/Engine/Plugins/Experimental/Paper2D/Source/Paper2D/Private/PaperSpriteComponent.cpp
new file mode 100644
--- /dev/null
+++ b/Engine/Plugins/Experimental/Paper2D/Source/Paper2D/Private/PaperSpriteComponent.cpp
@@ -0,0 +1,195 @@
+// Copyright 1998-2014 Epic Games, Inc. All Rights Reserved.
+
+#include "Paper2DPrivatePCH.h"
+#include "PaperRenderSceneProxy.h"
+#include "PaperSpriteComponent.h"
+#include "PhysicsEngine/BodySetup2D.h"
+
+//////////////////////////////////////////////////////////////////////////
+// UPaperSpriteComponent
+
+UPaperSpriteComponent::UPaperSpriteComponent(const FPostConstructInitializeProperties& PCIP)
+	: Super(PCIP)
+{
+	SetCollisionProfileName(UCollisionProfile::BlockAllDynamic_ProfileName);
+
+	MaterialOverride = nullptr;
+
+	SpriteColor = FLinearColor::White;
+}
+
+#if WITH_EDITOR
+void UPaperSpriteComponent::PostEditChangeProperty(FPropertyChangedEvent& PropertyChangedEvent)
+{
+	FBodyInstanceEditorHelpers::EnsureConsistentMobilitySimulationSettingsOnPostEditChange(this, PropertyChangedEvent);
+
+	Super::PostEditChangeProperty(PropertyChangedEvent);
+}
+#endif
+
+FPrimitiveSceneProxy* UPaperSpriteComponent::CreateSceneProxy()
+{
+	FPaperRenderSceneProxy* NewProxy = new FPaperRenderSceneProxy(this);
+	FSpriteDrawCallRecord DrawCall;
+	DrawCall.BuildFromSprite(SourceSprite);
+	DrawCall.Color = SpriteColor;
+	NewProxy->SetDrawCall_RenderThread(DrawCall);
+	return NewProxy;
+}
+
+FBoxSphereBounds UPaperSpriteComponent::CalcBounds(const FTransform& LocalToWorld) const
+{
+	if (SourceSprite != NULL)
+	{
+		// Graphics bounds.
+		FBoxSphereBounds NewBounds = SourceSprite->GetRenderBounds().TransformBy(LocalToWorld);
+
+		// Add bounds of collision geometry (if present).
+		if (UBodySetup* BodySetup = SourceSprite->BodySetup)
+		{
+			const FBox AggGeomBox = BodySetup->AggGeom.CalcAABB(LocalToWorld);
+			if (AggGeomBox.IsValid)
+			{
+				NewBounds = Union(NewBounds,FBoxSphereBounds(AggGeomBox));
+			}
+		}
+
+		// Apply bounds scale
+		NewBounds.BoxExtent *= BoundsScale;
+		NewBounds.SphereRadius *= BoundsScale;
+
+		return NewBounds;
+	}
+	else
+	{
+		return FBoxSphereBounds(LocalToWorld.GetLocation(), FVector::ZeroVector, 0.f);
+	}
+}
+
+void UPaperSpriteComponent::SendRenderDynamicData_Concurrent()
+{
+	if (SceneProxy != NULL)
+	{
+		FSpriteDrawCallRecord DrawCall;
+		DrawCall.BuildFromSprite(SourceSprite);
+		DrawCall.Color = SpriteColor;
+
+		ENQUEUE_UNIQUE_RENDER_COMMAND_TWOPARAMETER(
+				FSendPaperRenderComponentDynamicData,
+				FPaperRenderSceneProxy*,InSceneProxy,(FPaperRenderSceneProxy*)SceneProxy,
+				FSpriteDrawCallRecord,InSpriteToSend,DrawCall,
+			{
+				InSceneProxy->SetDrawCall_RenderThread(InSpriteToSend);
+			});
+	}
+}
+
+bool UPaperSpriteComponent::HasAnySockets() const
+{
+	if (SourceSprite != NULL)
+	{
+		return SourceSprite->HasAnySockets();
+	}
+
+	return false;
+}
+
+FTransform UPaperSpriteComponent::GetSocketTransform(FName InSocketName, ERelativeTransformSpace TransformSpace) const
+{
+	if (SourceSprite != NULL)
+	{
+		if (FPaperSpriteSocket* Socket = SourceSprite->FindSocket(InSocketName))
+		{
+			FTransform SocketLocalTransform = Socket->LocalTransform;
+			SocketLocalTransform.ScaleTranslation(SourceSprite->GetUnrealUnitsPerPixel());
+
+			switch (TransformSpace)
+			{
+				case RTS_World:
+					return SocketLocalTransform * ComponentToWorld;
+
+				case RTS_Actor:
+					if (const AActor* Actor = GetOwner())
+					{
+						const FTransform SocketTransform = SocketLocalTransform * ComponentToWorld;
+						return SocketTransform.GetRelativeTransform(Actor->GetTransform());
+					}
+					break;
+
+				case RTS_Component:
+					return SocketLocalTransform;
+
+				default:
+					check(false);
+			}
+		}
+	}
+
+	return Super::GetSocketTransform(InSocketName, TransformSpace);
+}
+
+void UPaperSpriteComponent::QuerySupportedSockets(TArray<FComponentSocketDescription>& OutSockets) const
+{
+	if (SourceSprite != NULL)
+	{
+		return SourceSprite->QuerySupportedSockets(OutSockets);
+	}
+}
+
+UBodySetup* UPaperSpriteComponent::GetBodySetup()
+{
+	return (SourceSprite != nullptr) ? SourceSprite->BodySetup : nullptr;
+}
+
+bool UPaperSpriteComponent::SetSprite(class UPaperSprite* NewSprite)
+{
+	if (NewSprite != SourceSprite)
+	{
+		// Don't allow changing the sprite if we are "static".
+		AActor* Owner = GetOwner();
+		if (!IsRegistered() || (Owner == NULL) || (Mobility != EComponentMobility::Static))
+		{
+			SourceSprite = NewSprite;
+
+			// Need to send this to render thread at some point
+			MarkRenderStateDirty();
+
+			// Update physics representation right away
+			RecreatePhysicsState();
+
+			// Since we have new mesh, we need to update bounds
+			UpdateBounds();
+
+			return true;
+		}
+	}
+
+	return false;
+}
+
+UPaperSprite* UPaperSpriteComponent::GetSprite()
+{
+	return SourceSprite;
+}
+
+void UPaperSpriteComponent::SetSpriteColor(FLinearColor NewColor)
+{
+	// Can't set color on a static component
+	if (!(IsRegistered() && (Mobility == EComponentMobility::Static)) && (SpriteColor != NewColor))
+	{
+		SpriteColor = NewColor;
+
+		//@TODO: Should we send immediately?
+		MarkRenderDynamicDataDirty();
+	}
+}
+
+FLinearColor UPaperSpriteComponent::GetWireframeColor() const
+{
+	return FLinearColor::Yellow;
+}
+
+const UObject* UPaperSpriteComponent::AdditionalStatObject() const
+{
+	return SourceSprite;
+}
