Paper2D: Rename pass to better match engine standards (phase 2)

- Rename APaperRenderActor to APaperSpriteActor

- Fix up many references to the old names that were working via the deprecated typedefs

diff --git a/Engine/Config/BaseEngine.ini b/Engine/Config/BaseEngine.ini
--- a/Engine/Config/BaseEngine.ini
+++ b/Engine/Config/BaseEngine.ini
@@ -653,6 +653,7 @@ bUseStreamingPause=false
 ; Paper2D modifications (4.3 era)
 +ActiveClassRedirects=(OldClassName="PaperRenderComponent",NewClassName="/Script/Paper2D.PaperSpriteComponent")
 +ActiveClassRedirects=(OldClassName="PaperAnimatedRenderComponent",NewClassName="/Script/Paper2D.PaperFlipbookComponent")
++ActiveClassRedirects=(OldClassName="PaperRenderActor",NewClassName="/Script/Paper2D.PaperSpriteActor")
 
 ; collision response container refactor
 +TaggedPropertyRedirects=(ClassName="CollisionResponseContainer",OldPropertyName="RigidBody",NewPropertyName="PhysicsBody")
diff --git a/Engine/Plugins/Experimental/Paper2D/Source/Paper2D/Classes/PaperRenderActor.h b/Engine/Plugins/Experimental/Paper2D/Source/Paper2D/Classes/PaperRenderActor.h
deleted file mode 100644
diff --git a/Engine/Plugins/Experimental/Paper2D/Source/Paper2D/Classes/PaperSpriteActor.h b/Engine/Plugins/Experimental/Paper2D/Source/Paper2D/Classes/PaperSpriteActor.h
new file mode 100644
--- /dev/null
+++ b/Engine/Plugins/Experimental/Paper2D/Source/Paper2D/Classes/PaperSpriteActor.h
@@ -0,0 +1,25 @@
+// Copyright 1998-2014 Epic Games, Inc. All Rights Reserved.
+
+#pragma once
+
+#include "PaperSpriteComponent.h"
+#include "PaperSpriteActor.generated.h"
+
+UCLASS(MinimalAPI)
+class APaperSpriteActor : public AActor
+{
+	GENERATED_UCLASS_BODY()
+
+	UPROPERTY(Category=Sprite, VisibleAnywhere)
+	TSubobjectPtr<class UPaperSpriteComponent> RenderComponent;
+
+	// AActor interface
+#if WITH_EDITOR
+	virtual bool GetReferencedContentObjects(TArray<UObject*>& Objects) const override;
+#endif
+	// End of AActor interface
+};
+
+// Allow the old name to continue to work for one release
+DEPRECATED(4.3, "APaperRenderActor has been renamed to APaperSpriteActor")
+typedef APaperSpriteActor APaperRenderActor;
diff --git a/Engine/Plugins/Experimental/Paper2D/Source/Paper2D/Private/PaperFlipbookActor.cpp b/Engine/Plugins/Experimental/Paper2D/Source/Paper2D/Private/PaperFlipbookActor.cpp
--- a/Engine/Plugins/Experimental/Paper2D/Source/Paper2D/Private/PaperFlipbookActor.cpp
+++ b/Engine/Plugins/Experimental/Paper2D/Source/Paper2D/Private/PaperFlipbookActor.cpp
@@ -1,7 +1,7 @@
 // Copyright 1998-2014 Epic Games, Inc. All Rights Reserved.
 
 #include "Paper2DPrivatePCH.h"
-#include "PaperRenderActor.h"
+#include "PaperFlipbookActor.h"
 
 //////////////////////////////////////////////////////////////////////////
 // APaperFlipbookActor
diff --git a/Engine/Plugins/Experimental/Paper2D/Source/Paper2D/Private/PaperRenderActor.cpp b/Engine/Plugins/Experimental/Paper2D/Source/Paper2D/Private/PaperRenderActor.cpp
deleted file mode 100644
diff --git a/Engine/Plugins/Experimental/Paper2D/Source/Paper2D/Private/PaperSpriteActor.cpp b/Engine/Plugins/Experimental/Paper2D/Source/Paper2D/Private/PaperSpriteActor.cpp
new file mode 100644
--- /dev/null
+++ b/Engine/Plugins/Experimental/Paper2D/Source/Paper2D/Private/PaperSpriteActor.cpp
@@ -0,0 +1,28 @@
+// Copyright 1998-2014 Epic Games, Inc. All Rights Reserved.
+
+#include "Paper2DPrivatePCH.h"
+#include "PaperSpriteActor.h"
+
+//////////////////////////////////////////////////////////////////////////
+// APaperSpriteActor
+
+APaperSpriteActor::APaperSpriteActor(const FPostConstructInitializeProperties& PCIP)
+	: Super(PCIP)
+{
+	RenderComponent = PCIP.CreateDefaultSubobject<UPaperSpriteComponent>(this, TEXT("RenderComponent"));
+	RenderComponent->SetCollisionProfileName(UCollisionProfile::BlockAll_ProfileName);
+	RenderComponent->Mobility = EComponentMobility::Static;
+
+	RootComponent = RenderComponent;
+}
+
+#if WITH_EDITOR
+bool APaperSpriteActor::GetReferencedContentObjects(TArray<UObject*>& Objects) const
+{
+	if (UPaperSprite* SourceSprite = RenderComponent->GetSprite())
+	{
+		Objects.Add(SourceSprite);
+	}
+	return true;
+}
+#endif
diff --git a/Engine/Plugins/Experimental/Paper2D/Source/Paper2D/Private/PaperSpriteComponent.cpp b/Engine/Plugins/Experimental/Paper2D/Source/Paper2D/Private/PaperSpriteComponent.cpp
--- a/Engine/Plugins/Experimental/Paper2D/Source/Paper2D/Private/PaperSpriteComponent.cpp
+++ b/Engine/Plugins/Experimental/Paper2D/Source/Paper2D/Private/PaperSpriteComponent.cpp
@@ -75,7 +75,7 @@ void UPaperSpriteComponent::SendRenderDynamicData_Concurrent()
 		DrawCall.Color = SpriteColor;
 
 		ENQUEUE_UNIQUE_RENDER_COMMAND_TWOPARAMETER(
-				FSendPaperRenderComponentDynamicData,
+				FSendPaperSpriteComponentDynamicData,
 				FPaperRenderSceneProxy*,InSceneProxy,(FPaperRenderSceneProxy*)SceneProxy,
 				FSpriteDrawCallRecord,InSpriteToSend,DrawCall,
 			{
diff --git a/Engine/Plugins/Experimental/Paper2D/Source/Paper2DEditor/Private/FlipbookEditor/FlipbookEditorViewportClient.cpp b/Engine/Plugins/Experimental/Paper2D/Source/Paper2DEditor/Private/FlipbookEditor/FlipbookEditorViewportClient.cpp
--- a/Engine/Plugins/Experimental/Paper2D/Source/Paper2DEditor/Private/FlipbookEditor/FlipbookEditorViewportClient.cpp
+++ b/Engine/Plugins/Experimental/Paper2D/Source/Paper2DEditor/Private/FlipbookEditor/FlipbookEditorViewportClient.cpp
@@ -23,7 +23,7 @@ FFlipbookEditorViewportClient::FFlipbookEditorViewportClient(const TAttribute<UP
 	SetRealtime(true);
 
 	// Create a render component for the sprite being edited
-	AnimatedRenderComponent = NewObject<UPaperAnimatedRenderComponent>();
+	AnimatedRenderComponent = NewObject<UPaperFlipbookComponent>();
 	AnimatedRenderComponent->SetFlipbook(FlipbookBeingEdited.Get());
 	AnimatedRenderComponent->UpdateBounds();
 	PreviewScene->AddComponent(AnimatedRenderComponent.Get(), FTransform::Identity);
diff --git a/Engine/Plugins/Experimental/Paper2D/Source/Paper2DEditor/Private/FlipbookEditor/FlipbookEditorViewportClient.h b/Engine/Plugins/Experimental/Paper2D/Source/Paper2DEditor/Private/FlipbookEditor/FlipbookEditorViewportClient.h
--- a/Engine/Plugins/Experimental/Paper2D/Source/Paper2DEditor/Private/FlipbookEditor/FlipbookEditorViewportClient.h
+++ b/Engine/Plugins/Experimental/Paper2D/Source/Paper2DEditor/Private/FlipbookEditor/FlipbookEditorViewportClient.h
@@ -47,7 +47,7 @@ private:
 	TWeakObjectPtr<class UPaperFlipbook> FlipbookBeingEditedLastFrame;
 
 	// Render component for the sprite being edited
-	TWeakObjectPtr<UPaperAnimatedRenderComponent> AnimatedRenderComponent;
+	TWeakObjectPtr<UPaperFlipbookComponent> AnimatedRenderComponent;
 
 	// Should we show the sprite pivot?
 	bool bShowPivot;
diff --git a/Engine/Plugins/Experimental/Paper2D/Source/Paper2DEditor/Private/Paper2DEditorModule.cpp b/Engine/Plugins/Experimental/Paper2D/Source/Paper2DEditor/Private/Paper2DEditorModule.cpp
--- a/Engine/Plugins/Experimental/Paper2D/Source/Paper2DEditor/Private/Paper2DEditorModule.cpp
+++ b/Engine/Plugins/Experimental/Paper2D/Source/Paper2DEditor/Private/Paper2DEditorModule.cpp
@@ -88,10 +88,10 @@ public:
 		RegisterAssetTypeAction(AssetTools, MakeShareable(new FAtlasAssetTypeActions));
 
 		PaperSpriteBroker = MakeShareable(new FPaperSpriteAssetBroker);
-		FComponentAssetBrokerage::RegisterBroker(PaperSpriteBroker, UPaperRenderComponent::StaticClass(), true, true);
+		FComponentAssetBrokerage::RegisterBroker(PaperSpriteBroker, UPaperSpriteComponent::StaticClass(), true, true);
 
 		PaperFlipbookBroker = MakeShareable(new FPaperFlipbookAssetBroker);
-		FComponentAssetBrokerage::RegisterBroker(PaperFlipbookBroker, UPaperAnimatedRenderComponent::StaticClass(), true, true);
+		FComponentAssetBrokerage::RegisterBroker(PaperFlipbookBroker, UPaperFlipbookComponent::StaticClass(), true, true);
 
 		// Register the details customizations
 		{
diff --git a/Engine/Plugins/Experimental/Paper2D/Source/Paper2DEditor/Private/PaperAnimatedSpriteActorFactory.cpp b/Engine/Plugins/Experimental/Paper2D/Source/Paper2DEditor/Private/PaperAnimatedSpriteActorFactory.cpp
--- a/Engine/Plugins/Experimental/Paper2D/Source/Paper2DEditor/Private/PaperAnimatedSpriteActorFactory.cpp
+++ b/Engine/Plugins/Experimental/Paper2D/Source/Paper2DEditor/Private/PaperAnimatedSpriteActorFactory.cpp
@@ -20,7 +20,7 @@ void UPaperFlipbookActorFactory::PostSpawnActor(UObject* Asset, AActor* NewActor
 		GEditor->SetActorLabelUnique(NewActor, Flipbook->GetName());
 
 		APaperFlipbookActor* TypedActor = CastChecked<APaperFlipbookActor>(NewActor);
-		UPaperAnimatedRenderComponent* RenderComponent = TypedActor->RenderComponent;
+		UPaperFlipbookComponent* RenderComponent = TypedActor->RenderComponent;
 		check(RenderComponent);
 
 		RenderComponent->UnregisterComponent();
diff --git a/Engine/Plugins/Experimental/Paper2D/Source/Paper2DEditor/Private/PaperFlipbookAssetBroker.h b/Engine/Plugins/Experimental/Paper2D/Source/Paper2DEditor/Private/PaperFlipbookAssetBroker.h
--- a/Engine/Plugins/Experimental/Paper2D/Source/Paper2DEditor/Private/PaperFlipbookAssetBroker.h
+++ b/Engine/Plugins/Experimental/Paper2D/Source/Paper2DEditor/Private/PaperFlipbookAssetBroker.h
@@ -17,7 +17,7 @@ public:
 
 	virtual bool AssignAssetToComponent(UActorComponent* InComponent, UObject* InAsset) override
 	{
-		if (UPaperAnimatedRenderComponent* RenderComp = Cast<UPaperAnimatedRenderComponent>(InComponent))
+		if (UPaperFlipbookComponent* RenderComp = Cast<UPaperFlipbookComponent>(InComponent))
 		{
 			UPaperFlipbook* Flipbook = Cast<UPaperFlipbook>(InAsset);
 
@@ -33,7 +33,7 @@ public:
 
 	virtual UObject* GetAssetFromComponent(UActorComponent* InComponent) override
 	{
-		if (UPaperAnimatedRenderComponent* RenderComp = Cast<UPaperAnimatedRenderComponent>(InComponent))
+		if (UPaperFlipbookComponent* RenderComp = Cast<UPaperFlipbookComponent>(InComponent))
 		{
 			return RenderComp->GetFlipbook();
 		}
diff --git a/Engine/Plugins/Experimental/Paper2D/Source/Paper2DEditor/Private/PaperSpriteActorFactory.cpp b/Engine/Plugins/Experimental/Paper2D/Source/Paper2DEditor/Private/PaperSpriteActorFactory.cpp
--- a/Engine/Plugins/Experimental/Paper2D/Source/Paper2DEditor/Private/PaperSpriteActorFactory.cpp
+++ b/Engine/Plugins/Experimental/Paper2D/Source/Paper2DEditor/Private/PaperSpriteActorFactory.cpp
@@ -10,7 +10,7 @@ UPaperSpriteActorFactory::UPaperSpriteActorFactory(const FPostConstructInitializ
 	: Super(PCIP)
 {
 	DisplayName = NSLOCTEXT("Paper2D", "PaperSpriteFactoryDisplayName", "Add Sprite");
-	NewActorClass = APaperRenderActor::StaticClass();
+	NewActorClass = APaperSpriteActor::StaticClass();
 }
 
 void UPaperSpriteActorFactory::PostSpawnActor(UObject* Asset, AActor* NewActor)
@@ -19,8 +19,8 @@ void UPaperSpriteActorFactory::PostSpawnActor(UObject* Asset, AActor* NewActor)
 	{
 		GEditor->SetActorLabelUnique(NewActor, Sprite->GetName());
 
-		APaperRenderActor* TypedActor = CastChecked<APaperRenderActor>(NewActor);
-		UPaperRenderComponent* RenderComponent = TypedActor->RenderComponent;
+		APaperSpriteActor* TypedActor = CastChecked<APaperSpriteActor>(NewActor);
+		UPaperSpriteComponent* RenderComponent = TypedActor->RenderComponent;
 		check(RenderComponent);
 
 		RenderComponent->UnregisterComponent();
@@ -31,7 +31,7 @@ void UPaperSpriteActorFactory::PostSpawnActor(UObject* Asset, AActor* NewActor)
 
 void UPaperSpriteActorFactory::PostCreateBlueprint(UObject* Asset, AActor* CDO)
 {
-	checkf(false, TEXT("APaperRenderActor isn't blueprintable; how did you get here?"));
+	checkf(false, TEXT("APaperSpriteActor isn't blueprintable; how did you get here?"));
 }
 
 bool UPaperSpriteActorFactory::CanCreateActorFrom(const FAssetData& AssetData, FText& OutErrorMsg)
diff --git a/Engine/Plugins/Experimental/Paper2D/Source/Paper2DEditor/Private/PaperSpriteAssetBroker.h b/Engine/Plugins/Experimental/Paper2D/Source/Paper2DEditor/Private/PaperSpriteAssetBroker.h
--- a/Engine/Plugins/Experimental/Paper2D/Source/Paper2DEditor/Private/PaperSpriteAssetBroker.h
+++ b/Engine/Plugins/Experimental/Paper2D/Source/Paper2DEditor/Private/PaperSpriteAssetBroker.h
@@ -17,7 +17,7 @@ public:
 
 	virtual bool AssignAssetToComponent(UActorComponent* InComponent, UObject* InAsset) override
 	{
-		if (UPaperRenderComponent* RenderComp = Cast<UPaperRenderComponent>(InComponent))
+		if (UPaperSpriteComponent* RenderComp = Cast<UPaperSpriteComponent>(InComponent))
 		{
 			UPaperSprite* Sprite = Cast<UPaperSprite>(InAsset);
 
@@ -33,7 +33,7 @@ public:
 
 	virtual UObject* GetAssetFromComponent(UActorComponent* InComponent) override
 	{
-		if (UPaperRenderComponent* RenderComp = Cast<UPaperRenderComponent>(InComponent))
+		if (UPaperSpriteComponent* RenderComp = Cast<UPaperSpriteComponent>(InComponent))
 		{
 			return RenderComp->GetSprite();
 		}
diff --git a/Engine/Plugins/Experimental/Paper2D/Source/Paper2DEditor/Private/SpriteEditor/SpriteEditorViewportClient.cpp b/Engine/Plugins/Experimental/Paper2D/Source/Paper2DEditor/Private/SpriteEditor/SpriteEditorViewportClient.cpp
--- a/Engine/Plugins/Experimental/Paper2D/Source/Paper2DEditor/Private/SpriteEditor/SpriteEditorViewportClient.cpp
+++ b/Engine/Plugins/Experimental/Paper2D/Source/Paper2DEditor/Private/SpriteEditor/SpriteEditorViewportClient.cpp
@@ -66,7 +66,7 @@ FSpriteEditorViewportClient::FSpriteEditorViewportClient(TWeakPtr<FSpriteEditor>
 
 	// Create a render component for the sprite being edited
 	{
-		RenderSpriteComponent = NewObject<UPaperRenderComponent>();
+		RenderSpriteComponent = NewObject<UPaperSpriteComponent>();
 		UPaperSprite* Sprite = GetSpriteBeingEdited();
 		RenderSpriteComponent->SetSprite(Sprite);
 
@@ -79,7 +79,7 @@ FSpriteEditorViewportClient::FSpriteEditorViewportClient(TWeakPtr<FSpriteEditor>
 		DummySprite->SpriteCollisionDomain = ESpriteCollisionMode::None;
 		DummySprite->PivotMode = ESpritePivotMode::Bottom_Left;
 
-		SourceTextureViewComponent = NewObject<UPaperRenderComponent>();
+		SourceTextureViewComponent = NewObject<UPaperSpriteComponent>();
 		SourceTextureViewComponent->SetSprite(DummySprite);
 		UpdateSourceTextureSpriteFromSprite(GetSpriteBeingEdited());
 
@@ -566,7 +566,7 @@ void FSpriteEditorViewportClient::Tick(float DeltaSeconds)
 		FIntPoint Size = Viewport->GetSizeXY();
 		if (bDeferZoomToSprite && Size.X > 0 && Size.Y > 0)
 		{
-			UPaperRenderComponent* ComponentToFocusOn = SourceTextureViewComponent->IsVisible() ? SourceTextureViewComponent : RenderSpriteComponent;
+			UPaperSpriteComponent* ComponentToFocusOn = SourceTextureViewComponent->IsVisible() ? SourceTextureViewComponent : RenderSpriteComponent;
 			FocusViewportOnBox(ComponentToFocusOn->Bounds.GetBox(), true);
 			bDeferZoomToSprite = false;
 		}		
diff --git a/Engine/Plugins/Experimental/Paper2D/Source/Paper2DEditor/Private/SpriteEditor/SpriteEditorViewportClient.h b/Engine/Plugins/Experimental/Paper2D/Source/Paper2DEditor/Private/SpriteEditor/SpriteEditorViewportClient.h
--- a/Engine/Plugins/Experimental/Paper2D/Source/Paper2DEditor/Private/SpriteEditor/SpriteEditorViewportClient.h
+++ b/Engine/Plugins/Experimental/Paper2D/Source/Paper2DEditor/Private/SpriteEditor/SpriteEditorViewportClient.h
@@ -110,10 +110,10 @@ private:
 	TWeakPtr<FSpriteEditor> SpriteEditorPtr;
 
 	// Render component for the source texture view
-	UPaperRenderComponent* SourceTextureViewComponent;
+	UPaperSpriteComponent* SourceTextureViewComponent;
 
 	// Render component for the sprite being edited
-	UPaperRenderComponent* RenderSpriteComponent;
+	UPaperSpriteComponent* RenderSpriteComponent;
 
 	// Widget mode
 	FWidget::EWidgetMode WidgetMode;
diff --git a/Engine/Source/Editor/EditorStyle/Private/SlateEditorStyle.cpp b/Engine/Source/Editor/EditorStyle/Private/SlateEditorStyle.cpp
--- a/Engine/Source/Editor/EditorStyle/Private/SlateEditorStyle.cpp
+++ b/Engine/Source/Editor/EditorStyle/Private/SlateEditorStyle.cpp
@@ -4216,13 +4216,13 @@ void FSlateEditorStyle::FStyle::SetupClassIconsAndThumbnails()
 
 			Set("ClassIcon.PaperTileMapActor", new IMAGE_BRUSH("Icons/ActorIcons/Paper2DTileMapActor_16x", Icon16x16));
 
-			Set("ClassIcon.PaperRenderActor", new IMAGE_BRUSH("Icons/AssetIcons/Paper2DSprite_16x", Icon16x16));
-			Set("ClassIcon.PaperRenderComponent", new IMAGE_BRUSH("Icons/AssetIcons/Paper2DSprite_16x", Icon16x16));
+			Set("ClassIcon.PaperSpriteActor", new IMAGE_BRUSH("Icons/AssetIcons/Paper2DSprite_16x", Icon16x16));
+			Set("ClassIcon.PaperSpriteComponent", new IMAGE_BRUSH("Icons/AssetIcons/Paper2DSprite_16x", Icon16x16));
 			Set("ClassIcon.PaperSprite", new IMAGE_BRUSH("Icons/AssetIcons/Paper2DSprite_16x", Icon16x16));
 			Set("ClassThumbnail.PaperSprite", new IMAGE_BRUSH("Icons/AssetIcons/Paper2DSprite_64x", Icon64x64));
 
 			Set("ClassIcon.PaperFlipbookActor", new IMAGE_BRUSH("Icons/AssetIcons/Paper2DSpriteFlipbook16x", Icon16x16));
-			Set("ClassIcon.PaperAnimatedRenderComponent", new IMAGE_BRUSH("Icons/AssetIcons/Paper2DSpriteFlipbook16x", Icon16x16));
+			Set("ClassIcon.PaperFlipbookComponent", new IMAGE_BRUSH("Icons/AssetIcons/Paper2DSpriteFlipbook16x", Icon16x16));
 			Set("ClassIcon.PaperFlipbook", new IMAGE_BRUSH("Icons/AssetIcons/Paper2DSpriteFlipbook16x", Icon16x16));
 			Set("ClassThumbnail.PaperFlipbook", new IMAGE_BRUSH("Icons/AssetIcons/Paper2DSpriteFlipbook64x", Icon64x64));
 
