Paper2D: Add ability to drag-drop sprite and flipbook assets into a flipbook timeline, appending them as keyframes

diff --git a/Engine/Plugins/Experimental/Paper2D/Source/Paper2DEditor/Private/FlipbookEditor/SFlipbookTimeline.cpp b/Engine/Plugins/Experimental/Paper2D/Source/Paper2DEditor/Private/FlipbookEditor/SFlipbookTimeline.cpp
--- a/Engine/Plugins/Experimental/Paper2D/Source/Paper2DEditor/Private/FlipbookEditor/SFlipbookTimeline.cpp
+++ b/Engine/Plugins/Experimental/Paper2D/Source/Paper2DEditor/Private/FlipbookEditor/SFlipbookTimeline.cpp
@@ -2,6 +2,7 @@
 
 #include "Paper2DEditorPrivatePCH.h"
 #include "SFlipbookTimeline.h"
+#include "Editor/UnrealEd/Public/DragAndDrop/AssetDragDropOp.h"
 
 #define LOCTEXT_NAMESPACE "FlipbookEditor"
 
@@ -56,6 +57,61 @@ void SFlipbookTimeline::Construct(const FArguments& InArgs, TSharedPtr<const FUI
 	];
 }
 
+FReply SFlipbookTimeline::OnDrop(const FGeometry& MyGeometry, const FDragDropEvent& DragDropEvent)
+{
+	bool bWasDropHandled = false;
+
+	TSharedPtr<FDragDropOperation> Operation = DragDropEvent.GetOperation();
+	if (!Operation.IsValid())
+	{
+	}
+	else if (Operation->IsOfType<FAssetDragDropOp>())
+	{
+		const auto& DragDropOp = StaticCastSharedPtr<FAssetDragDropOp>(Operation);
+
+		OnAssetsDropped(*DragDropOp);
+
+		bWasDropHandled = true;
+	}
+
+	return bWasDropHandled ? FReply::Handled() : FReply::Unhandled();
+}
+
+void SFlipbookTimeline::OnAssetsDropped(const class FAssetDragDropOp& DragDropOp)
+{
+	TArray<FPaperFlipbookKeyFrame> NewFrames;
+	for (const FAssetData& AssetData : DragDropOp.AssetData)
+	{
+		if (UObject* Object = AssetData.GetAsset())
+		{
+			if (UPaperSprite* SpriteAsset = Cast<UPaperSprite>(Object))
+			{
+				// Insert this sprite as a keyframe
+				FPaperFlipbookKeyFrame& NewFrame = *new (NewFrames) FPaperFlipbookKeyFrame();
+				NewFrame.Sprite = SpriteAsset;
+			}
+			else if (UPaperFlipbook* FlipbookAsset = Cast<UPaperFlipbook>(Object))
+			{
+				// Insert all of the keyframes from the other flipbook into this one
+				for (const FPaperFlipbookKeyFrame& OtherFlipbookFrame : FlipbookAsset->KeyFrames)
+				{
+					FPaperFlipbookKeyFrame& NewFrame = *new (NewFrames) FPaperFlipbookKeyFrame();
+					NewFrame = OtherFlipbookFrame;
+				}
+			}
+		}
+	}
+
+	UPaperFlipbook* ThisFlipbook = FlipbookBeingEdited.Get();
+	if (NewFrames.Num() && (ThisFlipbook != nullptr))
+	{
+		const FScopedTransaction Transaction(LOCTEXT("DroppedAssetOntoTimeline", "Insert assets as frames"));
+
+		ThisFlipbook->Modify();
+		ThisFlipbook->KeyFrames.Append(NewFrames);
+	}
+}
+
 //////////////////////////////////////////////////////////////////////////
 
 #undef LOCTEXT_NAMESPACE
\ No newline at end of file
diff --git a/Engine/Plugins/Experimental/Paper2D/Source/Paper2DEditor/Private/FlipbookEditor/SFlipbookTimeline.h b/Engine/Plugins/Experimental/Paper2D/Source/Paper2DEditor/Private/FlipbookEditor/SFlipbookTimeline.h
--- a/Engine/Plugins/Experimental/Paper2D/Source/Paper2DEditor/Private/FlipbookEditor/SFlipbookTimeline.h
+++ b/Engine/Plugins/Experimental/Paper2D/Source/Paper2DEditor/Private/FlipbookEditor/SFlipbookTimeline.h
@@ -22,6 +22,13 @@ public:
 
 	void Construct(const FArguments& InArgs, TSharedPtr<const FUICommandList> InCommandList);
 
+	// SWidget interface
+	virtual FReply OnDrop(const FGeometry& MyGeometry, const FDragDropEvent& DragDropEvent) override;
+	// End of SWidget interface
+
+private:
+	void OnAssetsDropped(const class FAssetDragDropOp& DragDropOp);
+
 private:
 	TAttribute<class UPaperFlipbook*> FlipbookBeingEdited;
 	TAttribute<float> PlayTime;
