Paper2D: Make PixelsPerUnrealUnit available in non-editor builds

diff --git a/Engine/Plugins/Experimental/Paper2D/Source/Paper2D/Classes/PaperSprite.h b/Engine/Plugins/Experimental/Paper2D/Source/Paper2D/Classes/PaperSprite.h
--- a/Engine/Plugins/Experimental/Paper2D/Source/Paper2D/Classes/PaperSprite.h
+++ b/Engine/Plugins/Experimental/Paper2D/Source/Paper2D/Classes/PaperSprite.h
@@ -60,6 +60,7 @@ protected:
 	FVector2D BakedSourceUV;
 #endif
 
+	// The source texture that the sprite comes from
 	UPROPERTY(Category=Sprite, EditAnywhere, AssetRegistrySearchable)
 	UTexture2D* SourceTexture;
 
@@ -78,6 +79,10 @@ protected:
 	UPROPERTY(Category=Collision, EditAnywhere)
 	TEnumAsByte<ESpriteCollisionMode::Type> SpriteCollisionDomain;
 
+	// The scaling factor between pixels and Unreal units (cm) (e.g., 0.64 would make a 64 pixel wide sprite take up 100 cm)
+	UPROPERTY(Category = Sprite, EditAnywhere)
+	float PixelsPerUnrealUnit;
+
 public:
 	// Baked physics data.
 	UPROPERTY() //@TODO: Is anything in this worth exposing to edit? Category=Physics, EditAnywhere, EditInline
@@ -102,10 +107,6 @@ protected:
 	UPROPERTY(Category=Collision, EditAnywhere)
 	float CollisionThickness;
 
-	// The scaling factor between pixels and Unreal units (cm) (e.g., 0.64 would make a 64 pixel wide sprite take up 100 cm)
-	UPROPERTY(Category=Sprite, EditAnywhere)
-	float PixelsPerUnrealUnit;
-
 	// Custom render geometry polygons (in texture space)
 	UPROPERTY(Category=Rendering, EditAnywhere)
 	FSpritePolygonCollection RenderGeometry;
@@ -173,9 +174,10 @@ public:
 	FVector2D GetSourceUV() const { return SourceUV; }
 	FVector2D GetSourceSize() const { return SourceDimension; }
 	UTexture2D* GetSourceTexture() const { return SourceTexture; }
+#endif
+
 	float GetPixelsPerUnrealUnit() const { return PixelsPerUnrealUnit; }
 	float GetUnrealUnitsPerPixel() const { return 1.0f / PixelsPerUnrealUnit; }
-#endif
 
 	// Returns the texture this should be rendered with
 	UTexture2D* GetBakedTexture() const;
