Paper2D: Add in-viewport toolbars to sprite and flipbook editors

diff --git a/Engine/Plugins/Experimental/Paper2D/Source/Paper2DEditor/Private/FlipbookEditor/FlipbookEditor.cpp b/Engine/Plugins/Experimental/Paper2D/Source/Paper2DEditor/Private/FlipbookEditor/FlipbookEditor.cpp
--- a/Engine/Plugins/Experimental/Paper2D/Source/Paper2DEditor/Private/FlipbookEditor/FlipbookEditor.cpp
+++ b/Engine/Plugins/Experimental/Paper2D/Source/Paper2DEditor/Private/FlipbookEditor/FlipbookEditor.cpp
@@ -12,6 +12,7 @@
 #include "SEditorViewport.h"
 #include "WorkspaceMenuStructureModule.h"
 #include "Paper2DEditorModule.h"
+#include "SFlipbookEditorViewportToolbar.h"
 
 #include "SFlipbookTimeline.h"
 
@@ -38,7 +39,7 @@ const FName FFlipbookEditorTabs::ViewportID(TEXT("Viewport"));
 //////////////////////////////////////////////////////////////////////////
 // SFlipbookEditorViewport
 
-class SFlipbookEditorViewport : public SEditorViewport
+class SFlipbookEditorViewport : public SEditorViewport, public ICommonEditorViewportToolbarInfoProvider
 {
 public:
 	SLATE_BEGIN_ARGS(SFlipbookEditorViewport)
@@ -57,8 +58,15 @@ public:
 	virtual void BindCommands() override;
 	virtual TSharedRef<FEditorViewportClient> MakeEditorViewportClient() override;
 	virtual TSharedPtr<SWidget> MakeViewportToolbar() override;
+	virtual EVisibility GetTransformToolbarVisibility() const override;
 	// End of SEditorViewport interface
 
+	// ICommonEditorViewportToolbarInfoProvider interface
+	virtual TSharedRef<class SEditorViewport> GetViewportWidget() override;
+	virtual TSharedPtr<FExtender> GetExtenders() const override;
+	virtual void OnFloatingButtonClicked() override;
+	// End of ICommonEditorViewportToolbarInfoProvider interface
+
 private:
 	TAttribute<UPaperFlipbook*> FlipbookBeingEdited;
 	TAttribute<float> PlayTime;
@@ -128,7 +136,28 @@ bool SFlipbookEditorViewport::IsVisible() const
 
 TSharedPtr<SWidget> SFlipbookEditorViewport::MakeViewportToolbar()
 {
-	return NULL;
+	return SNew(SFlipbookEditorViewportToolbar, SharedThis(this));
+}
+
+EVisibility SFlipbookEditorViewport::GetTransformToolbarVisibility() const
+{
+	// Nothing to transform in the flipbook editor, yet...
+	return EVisibility::Hidden;
+}
+
+TSharedRef<class SEditorViewport> SFlipbookEditorViewport::GetViewportWidget()
+{
+	return SharedThis(this);
+}
+
+TSharedPtr<FExtender> SFlipbookEditorViewport::GetExtenders() const
+{
+	TSharedPtr<FExtender> Result(MakeShareable(new FExtender));
+	return Result;
+}
+
+void SFlipbookEditorViewport::OnFloatingButtonClicked()
+{
 }
 
 /////////////////////////////////////////////////////
diff --git a/Engine/Plugins/Experimental/Paper2D/Source/Paper2DEditor/Private/FlipbookEditor/FlipbookEditorViewportClient.cpp b/Engine/Plugins/Experimental/Paper2D/Source/Paper2DEditor/Private/FlipbookEditor/FlipbookEditorViewportClient.cpp
--- a/Engine/Plugins/Experimental/Paper2D/Source/Paper2DEditor/Private/FlipbookEditor/FlipbookEditorViewportClient.cpp
+++ b/Engine/Plugins/Experimental/Paper2D/Source/Paper2DEditor/Private/FlipbookEditor/FlipbookEditorViewportClient.cpp
@@ -44,6 +44,18 @@ void FFlipbookEditorViewportClient::DrawCanvas(FViewport& Viewport, FSceneView&
 	{
 		Canvas.SetHitProxy(NULL);
 	}
+
+	int32 YPos = 42;
+
+	static const FText FlipbookHelpStr = LOCTEXT("FlipbookEditHelp", "Flipbook editor (currently playing)");
+
+	// Display tool help
+	{
+		FCanvasTextItem TextItem(FVector2D(6, YPos), FlipbookHelpStr, GEngine->GetSmallFont(), FLinearColor::White);
+		TextItem.EnableShadow(FLinearColor::Black);
+		TextItem.Draw(&Canvas);
+		YPos += 36;
+	}
 }
 
 void FFlipbookEditorViewportClient::Draw(FViewport* Viewport, FCanvas* Canvas)
diff --git a/Engine/Plugins/Experimental/Paper2D/Source/Paper2DEditor/Private/FlipbookEditor/SFlipbookEditorViewportToolbar.cpp b/Engine/Plugins/Experimental/Paper2D/Source/Paper2DEditor/Private/FlipbookEditor/SFlipbookEditorViewportToolbar.cpp
new file mode 100644
--- /dev/null
+++ b/Engine/Plugins/Experimental/Paper2D/Source/Paper2DEditor/Private/FlipbookEditor/SFlipbookEditorViewportToolbar.cpp
@@ -0,0 +1,16 @@
+// Copyright 1998-2014 Epic Games, Inc. All Rights Reserved.
+
+#include "Paper2DEditorPrivatePCH.h"
+#include "SFlipbookEditorViewportToolbar.h"
+
+#define LOCTEXT_NAMESPACE "SFlipbookEditorViewportToolbar"
+
+///////////////////////////////////////////////////////////
+// SFlipbookEditorViewportToolbar
+
+void SFlipbookEditorViewportToolbar::Construct(const FArguments& InArgs, TSharedPtr<class ICommonEditorViewportToolbarInfoProvider> InInfoProvider)
+{
+	SCommonEditorViewportToolbarBase::Construct(SCommonEditorViewportToolbarBase::FArguments(), InInfoProvider);
+}
+
+#undef LOCTEXT_NAMESPACE
diff --git a/Engine/Plugins/Experimental/Paper2D/Source/Paper2DEditor/Private/FlipbookEditor/SFlipbookEditorViewportToolbar.h b/Engine/Plugins/Experimental/Paper2D/Source/Paper2DEditor/Private/FlipbookEditor/SFlipbookEditorViewportToolbar.h
new file mode 100644
--- /dev/null
+++ b/Engine/Plugins/Experimental/Paper2D/Source/Paper2DEditor/Private/FlipbookEditor/SFlipbookEditorViewportToolbar.h
@@ -0,0 +1,15 @@
+// Copyright 1998-2014 Epic Games, Inc. All Rights Reserved.
+
+#pragma once
+
+#include "Editor/UnrealEd/Public/SCommonEditorViewportToolbarBase.h"
+
+// In-viewport toolbar widget used in the flipbook editor
+class SFlipbookEditorViewportToolbar : public SCommonEditorViewportToolbarBase
+{
+public:
+	SLATE_BEGIN_ARGS(SFlipbookEditorViewportToolbar) {}
+	SLATE_END_ARGS()
+
+	void Construct(const FArguments& InArgs, TSharedPtr<class ICommonEditorViewportToolbarInfoProvider> InInfoProvider);
+};
diff --git a/Engine/Plugins/Experimental/Paper2D/Source/Paper2DEditor/Private/SpriteEditor/SSpriteEditorViewportToolbar.cpp b/Engine/Plugins/Experimental/Paper2D/Source/Paper2DEditor/Private/SpriteEditor/SSpriteEditorViewportToolbar.cpp
new file mode 100644
--- /dev/null
+++ b/Engine/Plugins/Experimental/Paper2D/Source/Paper2DEditor/Private/SpriteEditor/SSpriteEditorViewportToolbar.cpp
@@ -0,0 +1,16 @@
+// Copyright 1998-2014 Epic Games, Inc. All Rights Reserved.
+
+#include "Paper2DEditorPrivatePCH.h"
+#include "SSpriteEditorViewportToolbar.h"
+
+#define LOCTEXT_NAMESPACE "SSpriteEditorViewportToolbar"
+
+///////////////////////////////////////////////////////////
+// SSpriteEditorViewportToolbar
+
+void SSpriteEditorViewportToolbar::Construct(const FArguments& InArgs, TSharedPtr<class ICommonEditorViewportToolbarInfoProvider> InInfoProvider)
+{
+	SCommonEditorViewportToolbarBase::Construct(SCommonEditorViewportToolbarBase::FArguments(), InInfoProvider);
+}
+
+#undef LOCTEXT_NAMESPACE
diff --git a/Engine/Plugins/Experimental/Paper2D/Source/Paper2DEditor/Private/SpriteEditor/SSpriteEditorViewportToolbar.h b/Engine/Plugins/Experimental/Paper2D/Source/Paper2DEditor/Private/SpriteEditor/SSpriteEditorViewportToolbar.h
new file mode 100644
--- /dev/null
+++ b/Engine/Plugins/Experimental/Paper2D/Source/Paper2DEditor/Private/SpriteEditor/SSpriteEditorViewportToolbar.h
@@ -0,0 +1,15 @@
+// Copyright 1998-2014 Epic Games, Inc. All Rights Reserved.
+
+#pragma once
+
+#include "Editor/UnrealEd/Public/SCommonEditorViewportToolbarBase.h"
+
+// In-viewport toolbar widget used in the sprite editor
+class SSpriteEditorViewportToolbar : public SCommonEditorViewportToolbarBase
+{
+public:
+	SLATE_BEGIN_ARGS(SSpriteEditorViewportToolbar) {}
+	SLATE_END_ARGS()
+
+	void Construct(const FArguments& InArgs, TSharedPtr<class ICommonEditorViewportToolbarInfoProvider> InInfoProvider);
+};
diff --git a/Engine/Plugins/Experimental/Paper2D/Source/Paper2DEditor/Private/SpriteEditor/SpriteEditor.cpp b/Engine/Plugins/Experimental/Paper2D/Source/Paper2DEditor/Private/SpriteEditor/SpriteEditor.cpp
--- a/Engine/Plugins/Experimental/Paper2D/Source/Paper2DEditor/Private/SpriteEditor/SpriteEditor.cpp
+++ b/Engine/Plugins/Experimental/Paper2D/Source/Paper2DEditor/Private/SpriteEditor/SpriteEditor.cpp
@@ -12,6 +12,7 @@
 #include "SEditorViewport.h"
 #include "WorkspaceMenuStructureModule.h"
 #include "Paper2DEditorModule.h"
+#include "SSpriteEditorViewportToolbar.h"
 
 #include "SSpriteList.h"
 
@@ -40,7 +41,7 @@ const FName FSpriteEditorTabs::SpriteListID(TEXT("SpriteList"));
 //////////////////////////////////////////////////////////////////////////
 // SSpriteEditorViewport
 
-class SSpriteEditorViewport : public SEditorViewport
+class SSpriteEditorViewport : public SEditorViewport, public ICommonEditorViewportToolbarInfoProvider
 {
 public:
 	SLATE_BEGIN_ARGS(SSpriteEditorViewport) {}
@@ -52,8 +53,15 @@ public:
 	virtual void BindCommands() override;
 	virtual TSharedRef<FEditorViewportClient> MakeEditorViewportClient() override;
 	virtual TSharedPtr<SWidget> MakeViewportToolbar() override;
+	virtual EVisibility GetTransformToolbarVisibility() const override;
 	// End of SEditorViewport interface
 
+	// ICommonEditorViewportToolbarInfoProvider interface
+	virtual TSharedRef<class SEditorViewport> GetViewportWidget() override;
+	virtual TSharedPtr<FExtender> GetExtenders() const override;
+	virtual void OnFloatingButtonClicked() override;
+	// End of ICommonEditorViewportToolbarInfoProvider interface
+
 	// Invalidate any references to the sprite being edited; it has changed
 	void NotifySpriteBeingEditedHasChanged()
 	{
@@ -191,12 +199,32 @@ TSharedRef<FEditorViewportClient> SSpriteEditorViewport::MakeEditorViewportClien
 
 bool SSpriteEditorViewport::IsVisible() const
 {
-	return true;//@TODO: Determine this better so viewport ticking optimizaitons can take place
+	return true;//@TODO: Determine this better so viewport ticking optimizations can take place
 }
 
 TSharedPtr<SWidget> SSpriteEditorViewport::MakeViewportToolbar()
 {
-	return NULL;
+	return SNew(SSpriteEditorViewportToolbar, SharedThis(this));
+}
+
+EVisibility SSpriteEditorViewport::GetTransformToolbarVisibility() const
+{
+	return EVisibility::Visible;
+}
+
+TSharedRef<class SEditorViewport> SSpriteEditorViewport::GetViewportWidget()
+{
+	return SharedThis(this);
+}
+
+TSharedPtr<FExtender> SSpriteEditorViewport::GetExtenders() const
+{
+	TSharedPtr<FExtender> Result(MakeShareable(new FExtender));
+	return Result;
+}
+
+void SSpriteEditorViewport::OnFloatingButtonClicked()
+{
 }
 
 /////////////////////////////////////////////////////
diff --git a/Engine/Plugins/Experimental/Paper2D/Source/Paper2DEditor/Private/SpriteEditor/SpriteEditorViewportClient.cpp b/Engine/Plugins/Experimental/Paper2D/Source/Paper2DEditor/Private/SpriteEditor/SpriteEditorViewportClient.cpp
--- a/Engine/Plugins/Experimental/Paper2D/Source/Paper2DEditor/Private/SpriteEditor/SpriteEditorViewportClient.cpp
+++ b/Engine/Plugins/Experimental/Paper2D/Source/Paper2DEditor/Private/SpriteEditor/SpriteEditorViewportClient.cpp
@@ -437,7 +437,7 @@ void FSpriteEditorViewportClient::DrawCanvas(FViewport& Viewport, FSceneView& Vi
 
 	UPaperSprite* Sprite = GetSpriteBeingEdited();
 
-	int32 YPos = 6;
+	int32 YPos = 42;
 
 	static const FText GeomHelpStr = LOCTEXT("GeomEditHelp", "Select an edge and press Insert to add a vertex.\nSelect an edge or vertex and press Delete to remove it.\n");
 	static const FText SourceRegionHelpStr = LOCTEXT("SourceRegionHelp", "Drag handles to adjust source region\nDouble-click on an image region to select all connected pixels");
diff --git a/Engine/Plugins/Experimental/Paper2D/Source/Paper2DEditor/Private/SpriteEditor/SpriteEditorViewportClient.h b/Engine/Plugins/Experimental/Paper2D/Source/Paper2DEditor/Private/SpriteEditor/SpriteEditorViewportClient.h
--- a/Engine/Plugins/Experimental/Paper2D/Source/Paper2DEditor/Private/SpriteEditor/SpriteEditorViewportClient.h
+++ b/Engine/Plugins/Experimental/Paper2D/Source/Paper2DEditor/Private/SpriteEditor/SpriteEditorViewportClient.h
@@ -99,6 +99,11 @@ public:
 
 	// Invalidate any references to the sprite being edited; it has changed
 	void NotifySpriteBeingEditedHasChanged();
+
+	ESpriteEditorMode::Type GetCurrentMode() const
+	{
+		return CurrentMode;
+	}
 private:
 	// Editor mode
 	ESpriteEditorMode::Type CurrentMode;
