Paper2D: Fix non-unity compilation error (missing ScopedTransaction.h)

diff --git a/Engine/Plugins/Experimental/Paper2D/Source/Paper2DEditor/Private/FlipbookEditor/SFlipbookTimeline.cpp b/Engine/Plugins/Experimental/Paper2D/Source/Paper2DEditor/Private/FlipbookEditor/SFlipbookTimeline.cpp
--- a/Engine/Plugins/Experimental/Paper2D/Source/Paper2DEditor/Private/FlipbookEditor/SFlipbookTimeline.cpp
+++ b/Engine/Plugins/Experimental/Paper2D/Source/Paper2DEditor/Private/FlipbookEditor/SFlipbookTimeline.cpp
@@ -3,6 +3,7 @@
 #include "Paper2DEditorPrivatePCH.h"
 #include "SFlipbookTimeline.h"
 #include "Editor/UnrealEd/Public/DragAndDrop/AssetDragDropOp.h"
+#include "Editor/UnrealEd/Public/ScopedTransaction.h"
 
 #define LOCTEXT_NAMESPACE "FlipbookEditor"
 
