Paper2D: Add accurate component bounds for flipbooks

- Add UPaperFlipbook::GetRenderBounds() and update UPaperFlipbookComponent::CalcBounds to use it

diff --git a/Engine/Plugins/Experimental/Paper2D/Source/Paper2D/Classes/PaperFlipbook.h b/Engine/Plugins/Experimental/Paper2D/Source/Paper2D/Classes/PaperFlipbook.h
--- a/Engine/Plugins/Experimental/Paper2D/Source/Paper2D/Classes/PaperFlipbook.h
+++ b/Engine/Plugins/Experimental/Paper2D/Source/Paper2D/Classes/PaperFlipbook.h
@@ -32,9 +32,11 @@ class UPaperFlipbook : public UObject
 {
 	GENERATED_UCLASS_BODY()
 
+	// The nominal frame rate to play this flipbook animation back at
 	UPROPERTY(Category=Sprite, EditAnywhere)
 	float FramesPerSecond;
 
+	// The set of key frames for this flipbook animation (each one has a duration and a sprite to display)
 	UPROPERTY(Category=Sprite, EditAnywhere)
 	TArray<FPaperFlipbookKeyFrame> KeyFrames;
 
@@ -43,6 +45,9 @@ class UPaperFlipbook : public UObject
 	PAPER2D_API UPaperSprite* GetSpriteAtTime(float Time) const;
 	PAPER2D_API UPaperSprite* GetSpriteAtFrame(int32 FrameIndex) const;
 
+	// Returns the render bounds of this sprite
+	FBoxSphereBounds GetRenderBounds() const;
+
 	// UObject interface
 #if WITH_EDITORONLY_DATA
 	virtual void Serialize(FArchive& Ar) override;
diff --git a/Engine/Plugins/Experimental/Paper2D/Source/Paper2D/Private/PaperFlipbook.cpp b/Engine/Plugins/Experimental/Paper2D/Source/Paper2D/Private/PaperFlipbook.cpp
--- a/Engine/Plugins/Experimental/Paper2D/Source/Paper2D/Private/PaperFlipbook.cpp
+++ b/Engine/Plugins/Experimental/Paper2D/Source/Paper2D/Private/PaperFlipbook.cpp
@@ -26,7 +26,7 @@ int32 UPaperFlipbook::GetNumFrames() const
 
 float UPaperFlipbook::GetDuration() const
 {
-	if ( FramesPerSecond != 0 )
+	if (FramesPerSecond != 0)
 	{
 		return GetNumFrames() / FramesPerSecond;
 	}
@@ -36,7 +36,7 @@ float UPaperFlipbook::GetDuration() const
 
 UPaperSprite* UPaperFlipbook::GetSpriteAtTime(float Time) const
 {
-	if ( FramesPerSecond != 0 )
+	if (FramesPerSecond != 0)
 	{
 		float SumTime = 0.0f;
 
@@ -77,4 +77,30 @@ void UPaperFlipbook::PostLoad()
 		SetFlags(RF_Transactional);
 	}
 }
-#endif
\ No newline at end of file
+#endif
+
+FBoxSphereBounds UPaperFlipbook::GetRenderBounds() const
+{
+	FBoxSphereBounds MergedBoundingBox(ForceInit);
+	bool bHasValidBounds = false;
+
+	for (const FPaperFlipbookKeyFrame& KeyFrame : KeyFrames)
+	{
+		if (KeyFrame.Sprite != nullptr)
+		{
+			const FBoxSphereBounds FrameBounds = KeyFrame.Sprite->GetRenderBounds();
+
+			if (bHasValidBounds)
+			{
+				MergedBoundingBox = Union(MergedBoundingBox, FrameBounds);
+			}
+			else
+			{
+				MergedBoundingBox = FrameBounds;
+				bHasValidBounds = true;
+			}
+		}
+	}
+
+	return MergedBoundingBox;
+}
diff --git a/Engine/Plugins/Experimental/Paper2D/Source/Paper2D/Private/PaperFlipbookComponent.cpp b/Engine/Plugins/Experimental/Paper2D/Source/Paper2D/Private/PaperFlipbookComponent.cpp
--- a/Engine/Plugins/Experimental/Paper2D/Source/Paper2D/Private/PaperFlipbookComponent.cpp
+++ b/Engine/Plugins/Experimental/Paper2D/Source/Paper2D/Private/PaperFlipbookComponent.cpp
@@ -48,8 +48,34 @@ FPrimitiveSceneProxy* UPaperFlipbookComponent::CreateSceneProxy()
 
 FBoxSphereBounds UPaperFlipbookComponent::CalcBounds(const FTransform & LocalToWorld) const
 {
-	const float NewScale = LocalToWorld.GetScale3D().GetMax() * 300.0f;//(Sprite ? (float)FMath::Max(Sprite->GetSizeX(),Sprite->GetSizeY()) : 1.0f);  //@TODO: Compute a realistic bound
-	return FBoxSphereBounds(LocalToWorld.GetTranslation(), FVector(NewScale,NewScale,NewScale), FMath::Sqrt(3.0f * FMath::Square(NewScale)));
+	if (SourceFlipbook != nullptr)
+	{
+		// Graphics bounds.
+		FBoxSphereBounds NewBounds = SourceFlipbook->GetRenderBounds().TransformBy(LocalToWorld);
+
+		//@TODO: PAPER2D: Add collision support to flipbooks
+#if 0
+		// Add bounds of collision geometry (if present).
+		if (UBodySetup* BodySetup = SourceFlipbook->BodySetup)
+		{
+			const FBox AggGeomBox = BodySetup->AggGeom.CalcAABB(LocalToWorld);
+			if (AggGeomBox.IsValid)
+			{
+				NewBounds = Union(NewBounds, FBoxSphereBounds(AggGeomBox));
+			}
+		}
+#endif
+
+		// Apply bounds scale
+		NewBounds.BoxExtent *= BoundsScale;
+		NewBounds.SphereRadius *= BoundsScale;
+
+		return NewBounds;
+	}
+	else
+	{
+		return FBoxSphereBounds(LocalToWorld.GetLocation(), FVector::ZeroVector, 0.f);
+	}
 }
 
 void UPaperFlipbookComponent::CalculateCurrentFrame()
