Paper2D: Refactored the public interface of UPaperFlipbook to allow for maintaining cached values in the future

- Hid KeyFrames and FramesPerSecond, adding const accessors
- Added a scoped editor lock called FScopedFlipbookMutator that ensures the cache will be invalidated
- Added clamping of FramesPerSecond to >= 0

- Fixed a rendering issue when editing brand new flipbooks where the time accumulator would go to infinity due to a zero sum time
- Thumbnail previews of flipbooks with no frames now show the background and a warning, instead of blackness

diff --git a/Engine/Plugins/Experimental/Paper2D/Source/Paper2D/Classes/PaperFlipbook.h b/Engine/Plugins/Experimental/Paper2D/Source/Paper2D/Classes/PaperFlipbook.h
--- a/Engine/Plugins/Experimental/Paper2D/Source/Paper2D/Classes/PaperFlipbook.h
+++ b/Engine/Plugins/Experimental/Paper2D/Source/Paper2D/Classes/PaperFlipbook.h
@@ -27,31 +27,105 @@ public:
 /**
  * Contains an animation sequence of sprite frames
  */
-UCLASS(MinimalAPI, BlueprintType, meta = (DisplayThumbnail = "true"))
-class UPaperFlipbook : public UObject
+UCLASS(BlueprintType, meta = (DisplayThumbnail = "true"))
+class PAPER2D_API UPaperFlipbook : public UObject
 {
 	GENERATED_UCLASS_BODY()
 
+protected:
 	// The nominal frame rate to play this flipbook animation back at
-	UPROPERTY(Category=Sprite, EditAnywhere)
+	UPROPERTY(Category=Sprite, EditAnywhere, BlueprintReadOnly)
 	float FramesPerSecond;
 
 	// The set of key frames for this flipbook animation (each one has a duration and a sprite to display)
 	UPROPERTY(Category=Sprite, EditAnywhere)
 	TArray<FPaperFlipbookKeyFrame> KeyFrames;
 
-	PAPER2D_API int32 GetNumFrames() const;
-	PAPER2D_API float GetDuration() const;
-	PAPER2D_API UPaperSprite* GetSpriteAtTime(float Time) const;
-	PAPER2D_API UPaperSprite* GetSpriteAtFrame(int32 FrameIndex) const;
+public:
+	// Returns the total number of frames
+	int32 GetNumFrames() const;
+
+	// Returns the total duration in seconds
+	float GetTotalDuration() const;
+
+	// Returns the keyframe index that covers the specified time (in seconds), or INDEX_NONE if none exists.
+	// When bClampToEnds is true, it will choose the first or last keyframe if the time is out of range.
+	int32 GetKeyFrameIndexAtTime(float Time, bool bClampToEnds = false) const;
+
+	// Returns the sprite at the specified time (in seconds), or nullptr if none exists.
+	// When bClampToEnds is true, it will choose the first or last sprite if the time is out of range.
+	UPaperSprite* GetSpriteAtTime(float Time, bool bClampToEnds = false) const;
+
+	// Returns the sprite at the specified frame index, or nullptr if none exists
+	UPaperSprite* GetSpriteAtFrame(int32 FrameIndex) const;
 
 	// Returns the render bounds of this sprite
 	FBoxSphereBounds GetRenderBounds() const;
 
+	// Returns the number of key frames
+	int32 GetNumKeyFrames() const
+	{
+		return KeyFrames.Num();
+	}
+
+	// Is the specified Index within the valid range of key frames?
+	bool IsValidKeyFrameIndex(int32 Index) const
+	{
+		return KeyFrames.IsValidIndex(Index);
+	}
+
+	// Returns the key frame at the specified index, make sure the index is valid before use
+	const FPaperFlipbookKeyFrame& GetKeyFrameChecked(int32 Index) const
+	{
+		return KeyFrames[Index];
+	}
+
+	// Rebuilds cached data about the animation (such as total number of frames that the keyframes span, etc...)
+	void InvalidateCachedData();
+
 	// UObject interface
 #if WITH_EDITORONLY_DATA
 	virtual void Serialize(FArchive& Ar) override;
 	virtual void PostLoad() override;
+	virtual void PostEditChangeProperty(FPropertyChangedEvent& PropertyChangedEvent) override;
 #endif
 	// End of UObject interface
+
+private:
+	friend class FScopedFlipbookMutator;
+};
+
+
+// Helper class to edit properties of a UPaperFlipbook while ensuring that cached data remains up to date
+class FScopedFlipbookMutator
+{
+public:
+	float& FramesPerSecond;
+	TArray<FPaperFlipbookKeyFrame>& KeyFrames;
+
+private:
+	UPaperFlipbook* SourceFlipbook;
+
+public:
+	FScopedFlipbookMutator(UPaperFlipbook* InFlipbook)
+		: SourceFlipbook(InFlipbook)
+		, FramesPerSecond(InFlipbook->FramesPerSecond)
+		, KeyFrames(InFlipbook->KeyFrames)
+	{
+	}
+
+	~FScopedFlipbookMutator()
+	{
+		InvalidateCachedData();
+	}
+
+	void InvalidateCachedData()
+	{
+		SourceFlipbook->InvalidateCachedData();
+	}
+
+	UPaperFlipbook* GetSourceFlipbook() const
+	{
+		return SourceFlipbook;
+	}
 };
diff --git a/Engine/Plugins/Experimental/Paper2D/Source/Paper2D/Classes/PaperFlipbookComponent.h b/Engine/Plugins/Experimental/Paper2D/Source/Paper2D/Classes/PaperFlipbookComponent.h
--- a/Engine/Plugins/Experimental/Paper2D/Source/Paper2D/Classes/PaperFlipbookComponent.h
+++ b/Engine/Plugins/Experimental/Paper2D/Source/Paper2D/Classes/PaperFlipbookComponent.h
@@ -49,6 +49,8 @@ protected:
 
 protected:
 	void CalculateCurrentFrame();
+	UPaperSprite* GetSpriteAtCachedIndex() const;
+
 public:
 	// UActorComponent interface
 	virtual void TickComponent(float DeltaTime, enum ELevelTick TickType, FActorComponentTickFunction* ThisTickFunction) override;
diff --git a/Engine/Plugins/Experimental/Paper2D/Source/Paper2D/Private/PaperFlipbook.cpp b/Engine/Plugins/Experimental/Paper2D/Source/Paper2D/Private/PaperFlipbook.cpp
--- a/Engine/Plugins/Experimental/Paper2D/Source/Paper2D/Private/PaperFlipbook.cpp
+++ b/Engine/Plugins/Experimental/Paper2D/Source/Paper2D/Private/PaperFlipbook.cpp
@@ -4,6 +4,10 @@
 #include "PaperFlipbook.h"
 #include "PaperCustomVersion.h"
 
+#if WITH_EDITORONLY_DATA
+#include "Runtime/Engine/Public/ComponentReregisterContext.h"
+#endif
+
 //////////////////////////////////////////////////////////////////////////
 // UPaperFlipbook
 
@@ -24,19 +28,24 @@ int32 UPaperFlipbook::GetNumFrames() const
 	return SumFrames;
 }
 
-float UPaperFlipbook::GetDuration() const
+float UPaperFlipbook::GetTotalDuration() const
 {
 	if (FramesPerSecond != 0)
 	{
 		return GetNumFrames() / FramesPerSecond;
 	}
 
-	return 0;
+	return 0.0f;
 }
 
-UPaperSprite* UPaperFlipbook::GetSpriteAtTime(float Time) const
+int32 UPaperFlipbook::GetKeyFrameIndexAtTime(float Time, bool bClampToEnds) const
 {
-	if (FramesPerSecond != 0)
+	if ((Time < 0.0f) && !bClampToEnds)
+	{
+		return INDEX_NONE;
+	}
+
+	if (FramesPerSecond > 0.0f)
 	{
 		float SumTime = 0.0f;
 
@@ -46,12 +55,23 @@ UPaperSprite* UPaperFlipbook::GetSpriteAtTime(float Time) const
 
 			if (Time <= SumTime)
 			{
-				return KeyFrames[KeyFrameIndex].Sprite;
+				return KeyFrameIndex;
 			}
 		}
+
+		// Return the last frame (note: relies on INDEX_NONE = -1 if there are no key frames)
+		return KeyFrames.Num() - 1;
 	}
+	else
+	{
+		return (KeyFrames.Num() > 0) ? 0 : INDEX_NONE;
+	}
+}
 
-	return NULL;
+UPaperSprite* UPaperFlipbook::GetSpriteAtTime(float Time, bool bClampToEnds) const
+{
+	const int32 KeyFrameIndex = GetKeyFrameIndexAtTime(Time, bClampToEnds);
+	return KeyFrames.IsValidIndex(KeyFrameIndex) ? KeyFrames[KeyFrameIndex].Sprite : nullptr;
 }
 
 UPaperSprite* UPaperFlipbook::GetSpriteAtFrame(int32 FrameIndex) const
@@ -77,6 +97,19 @@ void UPaperFlipbook::PostLoad()
 		SetFlags(RF_Transactional);
 	}
 }
+
+void UPaperFlipbook::PostEditChangeProperty(FPropertyChangedEvent& PropertyChangedEvent)
+{
+	if (FramesPerSecond < 0.0f)
+	{
+		FramesPerSecond = 0.0f;
+	}
+
+	//@TODO: Determine when this is really needed, as it is seriously expensive!
+	TComponentReregisterContext<UPaperFlipbookComponent> ReregisterAnimatedComponents;
+
+	Super::PostEditChangeProperty(PropertyChangedEvent);
+}
 #endif
 
 FBoxSphereBounds UPaperFlipbook::GetRenderBounds() const
@@ -104,3 +137,8 @@ FBoxSphereBounds UPaperFlipbook::GetRenderBounds() const
 
 	return MergedBoundingBox;
 }
+
+void UPaperFlipbook::InvalidateCachedData()
+{
+	// No cached data yet, but the functions that currently have to iterate over all frames can use cached data in the future
+}
diff --git a/Engine/Plugins/Experimental/Paper2D/Source/Paper2D/Private/PaperFlipbookComponent.cpp b/Engine/Plugins/Experimental/Paper2D/Source/Paper2D/Private/PaperFlipbookComponent.cpp
--- a/Engine/Plugins/Experimental/Paper2D/Source/Paper2D/Private/PaperFlipbookComponent.cpp
+++ b/Engine/Plugins/Experimental/Paper2D/Source/Paper2D/Private/PaperFlipbookComponent.cpp
@@ -28,16 +28,22 @@ UPaperFlipbookComponent::UPaperFlipbookComponent(const FPostConstructInitializeP
 	PlayRate = 1.0f;
 }
 
+UPaperSprite* UPaperFlipbookComponent::GetSpriteAtCachedIndex() const
+{
+	UPaperSprite* SpriteToSend = nullptr;
+	if ((SourceFlipbook != nullptr) && SourceFlipbook->IsValidKeyFrameIndex(CachedFrameIndex))
+	{
+		SpriteToSend = SourceFlipbook->GetKeyFrameChecked(CachedFrameIndex).Sprite;
+	}
+	return SpriteToSend;
+}
+
 FPrimitiveSceneProxy* UPaperFlipbookComponent::CreateSceneProxy()
 {
 	FPaperFlipbookSceneProxy* NewProxy = new FPaperFlipbookSceneProxy(this);
 
 	CalculateCurrentFrame();
-	UPaperSprite* SpriteToSend = NULL;
-	if ((SourceFlipbook != NULL) && (CachedFrameIndex >= 0) && (CachedFrameIndex < SourceFlipbook->KeyFrames.Num()))
-	{
-		SpriteToSend = SourceFlipbook->KeyFrames[CachedFrameIndex].Sprite;
-	}
+	UPaperSprite* SpriteToSend = GetSpriteAtCachedIndex();
 
 	FSpriteDrawCallRecord DrawCall;
 	DrawCall.BuildFromSprite(SpriteToSend);
@@ -80,61 +86,38 @@ FBoxSphereBounds UPaperFlipbookComponent::CalcBounds(const FTransform & LocalToW
 
 void UPaperFlipbookComponent::CalculateCurrentFrame()
 {
-	if (SourceFlipbook != NULL && SourceFlipbook->FramesPerSecond > 0)
-	{
-		int32 LastCachedFrame = CachedFrameIndex;
-
-		float SumTime = 0.0f;
-		int32 FrameIndex;
-		for (FrameIndex = 0; FrameIndex < SourceFlipbook->KeyFrames.Num(); ++FrameIndex)
-		{
-			SumTime += SourceFlipbook->KeyFrames[FrameIndex].FrameRun / SourceFlipbook->FramesPerSecond;
-
-			if (AccumulatedTime < SumTime)
-			{
-				break;
-			}
-		}
-
-		if (AccumulatedTime >= SumTime)
-		{
-			AccumulatedTime = FMath::Fmod(AccumulatedTime, SumTime);
-			//@TODO: Could have gone farther than this!!!
-			CachedFrameIndex = 0;
-		}
-		else
-		{
-			CachedFrameIndex = FrameIndex;
-		}
+	const int32 LastCachedFrame = CachedFrameIndex;
+	CachedFrameIndex = (SourceFlipbook != nullptr) ? SourceFlipbook->GetKeyFrameIndexAtTime(AccumulatedTime) : INDEX_NONE;
 
-		if (CachedFrameIndex != LastCachedFrame)
-		{
-			// Indicate we need to send new dynamic data.
-			MarkRenderDynamicDataDirty();
-		}
+	if (CachedFrameIndex != LastCachedFrame)
+	{
+		// Indicate we need to send new dynamic data.
+		MarkRenderDynamicDataDirty();
 	}
 }
 
 void UPaperFlipbookComponent::TickComponent(float DeltaTime, enum ELevelTick TickType, FActorComponentTickFunction* ThisTickFunction)
 {
-	// Indicate we need to send new dynamic data.
-	MarkRenderDynamicDataDirty();
-	
-	AccumulatedTime += DeltaTime * PlayRate;
-	CalculateCurrentFrame();
+	// Advance time
+	const float TotalTime = (SourceFlipbook != nullptr) ? SourceFlipbook->GetTotalDuration() : 0.0f;
+	if (TotalTime > 0.0f)
+	{
+		AccumulatedTime += DeltaTime * PlayRate;
+		AccumulatedTime = FMath::Fmod(AccumulatedTime, TotalTime);
+	}
+	else
+	{
+		AccumulatedTime = 0.0f;
+	}
 
-	//Super::TickComponent(DeltaTime, TickType, ThisTickFunction);
+	CalculateCurrentFrame();
 }
 
 void UPaperFlipbookComponent::SendRenderDynamicData_Concurrent()
 {
 	if (SceneProxy != NULL)
 	{
-		UPaperSprite* SpriteToSend = NULL;
-		if ((SourceFlipbook != NULL) && (CachedFrameIndex >= 0) && (CachedFrameIndex < SourceFlipbook->KeyFrames.Num()))
-		{
-			SpriteToSend = SourceFlipbook->KeyFrames[CachedFrameIndex].Sprite;
-		}
+		UPaperSprite* SpriteToSend = GetSpriteAtCachedIndex();
 
 		FSpriteDrawCallRecord DrawCall;
 		DrawCall.BuildFromSprite(SpriteToSend);
@@ -204,8 +187,7 @@ void UPaperFlipbookComponent::SetSpriteColor(FLinearColor NewColor)
 
 void UPaperFlipbookComponent::SetCurrentTime(float NewTime)
 {
-	// Can't set color on a static component
-	if (!(IsRegistered() && (Mobility == EComponentMobility::Static)) && (NewTime != AccumulatedTime))
+	if (NewTime != AccumulatedTime)
 	{
 		AccumulatedTime = NewTime;
 		CalculateCurrentFrame();
diff --git a/Engine/Plugins/Experimental/Paper2D/Source/Paper2DEditor/Classes/PaperSpriteThumbnailRenderer.h b/Engine/Plugins/Experimental/Paper2D/Source/Paper2DEditor/Classes/PaperSpriteThumbnailRenderer.h
--- a/Engine/Plugins/Experimental/Paper2D/Source/Paper2DEditor/Classes/PaperSpriteThumbnailRenderer.h
+++ b/Engine/Plugins/Experimental/Paper2D/Source/Paper2DEditor/Classes/PaperSpriteThumbnailRenderer.h
@@ -13,4 +13,6 @@ class UPaperSpriteThumbnailRenderer : public UDefaultSizedThumbnailRenderer
 
 protected:
 	void DrawFrame(class UPaperSprite* Sprite, int32 X, int32 Y, uint32 Width, uint32 Height, FRenderTarget*, FCanvas* Canvas);
+
+	void DrawGrid(int32 X, int32 Y, uint32 Width, uint32 Height, FCanvas* Canvas);
 };
diff --git a/Engine/Plugins/Experimental/Paper2D/Source/Paper2DEditor/Private/FlipbookEditor/FlipbookEditor.cpp b/Engine/Plugins/Experimental/Paper2D/Source/Paper2DEditor/Private/FlipbookEditor/FlipbookEditor.cpp
--- a/Engine/Plugins/Experimental/Paper2D/Source/Paper2DEditor/Private/FlipbookEditor/FlipbookEditor.cpp
+++ b/Engine/Plugins/Experimental/Paper2D/Source/Paper2DEditor/Private/FlipbookEditor/FlipbookEditor.cpp
@@ -437,12 +437,14 @@ void FFlipbookEditor::ExtendToolbar()
 
 void FFlipbookEditor::DeleteSelection()
 {
-	if (FlipbookBeingEdited->KeyFrames.IsValidIndex(CurrentSelectedKeyframe))
+
+	if (FlipbookBeingEdited->IsValidKeyFrameIndex(CurrentSelectedKeyframe))
 	{
 		const FScopedTransaction Transaction(LOCTEXT("DeleteKeyframe", "Delete Keyframe"));
 		FlipbookBeingEdited->Modify();
 
-		FlipbookBeingEdited->KeyFrames.RemoveAt(CurrentSelectedKeyframe);
+		FScopedFlipbookMutator EditLock(FlipbookBeingEdited);
+		EditLock.KeyFrames.RemoveAt(CurrentSelectedKeyframe);
 
 		CurrentSelectedKeyframe = INDEX_NONE;
 	}
@@ -450,13 +452,15 @@ void FFlipbookEditor::DeleteSelection()
 
 void FFlipbookEditor::DuplicateSelection()
 {
-	if (FlipbookBeingEdited->KeyFrames.IsValidIndex(CurrentSelectedKeyframe))
+	if (FlipbookBeingEdited->IsValidKeyFrameIndex(CurrentSelectedKeyframe))
 	{
 		const FScopedTransaction Transaction(LOCTEXT("DuplicateKeyframe", "Duplicate Keyframe"));
 		FlipbookBeingEdited->Modify();
 
-		FPaperFlipbookKeyFrame NewFrame = FlipbookBeingEdited->KeyFrames[CurrentSelectedKeyframe];
-		FlipbookBeingEdited->KeyFrames.Insert(NewFrame, CurrentSelectedKeyframe);
+		FScopedFlipbookMutator EditLock(FlipbookBeingEdited);
+
+		FPaperFlipbookKeyFrame NewFrame = EditLock.KeyFrames[CurrentSelectedKeyframe];
+		EditLock.KeyFrames.Insert(NewFrame, CurrentSelectedKeyframe);
 
 		CurrentSelectedKeyframe = INDEX_NONE;
 	}
@@ -469,7 +473,7 @@ void FFlipbookEditor::SetSelection(int32 NewSelection)
 
 bool FFlipbookEditor::HasValidSelection() const
 {
-	return FlipbookBeingEdited->KeyFrames.IsValidIndex(CurrentSelectedKeyframe);
+	return FlipbookBeingEdited->IsValidKeyFrameIndex(CurrentSelectedKeyframe);
 }
 
 //////////////////////////////////////////////////////////////////////////
diff --git a/Engine/Plugins/Experimental/Paper2D/Source/Paper2DEditor/Private/FlipbookEditor/SFlipbookTimeline.cpp b/Engine/Plugins/Experimental/Paper2D/Source/Paper2DEditor/Private/FlipbookEditor/SFlipbookTimeline.cpp
--- a/Engine/Plugins/Experimental/Paper2D/Source/Paper2DEditor/Private/FlipbookEditor/SFlipbookTimeline.cpp
+++ b/Engine/Plugins/Experimental/Paper2D/Source/Paper2DEditor/Private/FlipbookEditor/SFlipbookTimeline.cpp
@@ -94,8 +94,9 @@ void SFlipbookTimeline::OnAssetsDropped(const class FAssetDragDropOp& DragDropOp
 			else if (UPaperFlipbook* FlipbookAsset = Cast<UPaperFlipbook>(Object))
 			{
 				// Insert all of the keyframes from the other flipbook into this one
-				for (const FPaperFlipbookKeyFrame& OtherFlipbookFrame : FlipbookAsset->KeyFrames)
+				for (int32 KeyIndex = 0; KeyIndex < FlipbookAsset->GetNumKeyFrames(); ++KeyIndex)
 				{
+					const FPaperFlipbookKeyFrame& OtherFlipbookFrame = FlipbookAsset->GetKeyFrameChecked(KeyIndex);
 					FPaperFlipbookKeyFrame& NewFrame = *new (NewFrames) FPaperFlipbookKeyFrame();
 					NewFrame = OtherFlipbookFrame;
 				}
@@ -107,9 +108,10 @@ void SFlipbookTimeline::OnAssetsDropped(const class FAssetDragDropOp& DragDropOp
 	if (NewFrames.Num() && (ThisFlipbook != nullptr))
 	{
 		const FScopedTransaction Transaction(LOCTEXT("DroppedAssetOntoTimeline", "Insert assets as frames"));
-
 		ThisFlipbook->Modify();
-		ThisFlipbook->KeyFrames.Append(NewFrames);
+
+		FScopedFlipbookMutator EditLock(ThisFlipbook);
+		EditLock.KeyFrames.Append(NewFrames);
 	}
 }
 
diff --git a/Engine/Plugins/Experimental/Paper2D/Source/Paper2DEditor/Private/FlipbookEditor/SFlipbookTrackHandle.h b/Engine/Plugins/Experimental/Paper2D/Source/Paper2DEditor/Private/FlipbookEditor/SFlipbookTrackHandle.h
--- a/Engine/Plugins/Experimental/Paper2D/Source/Paper2DEditor/Private/FlipbookEditor/SFlipbookTrackHandle.h
+++ b/Engine/Plugins/Experimental/Paper2D/Source/Paper2DEditor/Private/FlipbookEditor/SFlipbookTrackHandle.h
@@ -56,9 +56,9 @@ public:
 			if (bDragging && (StartingFrameRun != INDEX_NONE))
 			{
 				UPaperFlipbook* Flipbook = FlipbookBeingEdited.Get();
-				if (Flipbook && Flipbook->KeyFrames.IsValidIndex(KeyFrameIdx))
+				if (Flipbook && Flipbook->IsValidKeyFrameIndex(KeyFrameIdx))
 				{
-					FPaperFlipbookKeyFrame& KeyFrame = Flipbook->KeyFrames[KeyFrameIdx];
+					const FPaperFlipbookKeyFrame& KeyFrame = Flipbook->GetKeyFrameChecked(KeyFrameIdx);
 
 					if (KeyFrame.FrameRun != StartingFrameRun)
 					{
@@ -89,16 +89,15 @@ public:
 		if (this->HasMouseCapture())
 		{
 			UPaperFlipbook* Flipbook = FlipbookBeingEdited.Get();
-			if (Flipbook && Flipbook->KeyFrames.IsValidIndex(KeyFrameIdx))
+			if (Flipbook && Flipbook->IsValidKeyFrameIndex(KeyFrameIdx))
 			{
-				FPaperFlipbookKeyFrame& KeyFrame = Flipbook->KeyFrames[KeyFrameIdx];
-
 				DistanceDragged += MouseEvent.GetCursorDelta().X;
 
 				if (!bDragging)
 				{
 					if (FMath::Abs(DistanceDragged) > SlateDragStartDistance)
 					{
+						const FPaperFlipbookKeyFrame& KeyFrame = Flipbook->GetKeyFrameChecked(KeyFrameIdx);
 						StartingFrameRun = KeyFrame.FrameRun;
 						bDragging = true;
 					}
@@ -108,6 +107,8 @@ public:
 					float LocalSlateUnitsPerFrame = SlateUnitsPerFrame.Get();
 					if (LocalSlateUnitsPerFrame != 0)
 					{
+						FScopedFlipbookMutator EditLock(Flipbook);
+						FPaperFlipbookKeyFrame& KeyFrame = EditLock.KeyFrames[KeyFrameIdx];
 						KeyFrame.FrameRun = StartingFrameRun + (DistanceDragged / LocalSlateUnitsPerFrame);
 						KeyFrame.FrameRun = FMath::Max<int32>(1, KeyFrame.FrameRun);
 					}
diff --git a/Engine/Plugins/Experimental/Paper2D/Source/Paper2DEditor/Private/FlipbookEditor/STimelineTrack.h b/Engine/Plugins/Experimental/Paper2D/Source/Paper2DEditor/Private/FlipbookEditor/STimelineTrack.h
--- a/Engine/Plugins/Experimental/Paper2D/Source/Paper2DEditor/Private/FlipbookEditor/STimelineTrack.h
+++ b/Engine/Plugins/Experimental/Paper2D/Source/Paper2DEditor/Private/FlipbookEditor/STimelineTrack.h
@@ -41,7 +41,7 @@ public:
 		SCompoundWidget::Tick(AllottedGeometry, InCurrentTime, InDeltaTime);
 
 		UPaperFlipbook* Flipbook = FlipbookBeingEdited.Get();
-		int32 NewNumKeyframes = (Flipbook != NULL ? Flipbook->KeyFrames.Num() : 0);
+		int32 NewNumKeyframes = (Flipbook != nullptr) ? Flipbook->GetNumKeyFrames() : 0;
 		if (NewNumKeyframes != NumKeyframesFromLastRebuild)
 		{
 			Rebuild();
@@ -77,7 +77,7 @@ private:
 		// Create the sections for each keyframe
 		if (UPaperFlipbook* Flipbook = FlipbookBeingEdited.Get())
 		{
-			for (int32 KeyFrameIdx = 0; KeyFrameIdx < Flipbook->KeyFrames.Num(); ++KeyFrameIdx)
+			for (int32 KeyFrameIdx = 0; KeyFrameIdx < Flipbook->GetNumKeyFrames(); ++KeyFrameIdx)
 			{
 				MainBoxPtr->AddSlot()
 				.AutoWidth()
@@ -109,7 +109,7 @@ private:
 				];
 			}
 
-			NumKeyframesFromLastRebuild = Flipbook->KeyFrames.Num();
+			NumKeyframesFromLastRebuild = Flipbook->GetNumKeyFrames();
 		}
 		else
 		{
@@ -120,9 +120,9 @@ private:
 	FOptionalSize GetFrameWidth(int32 KeyFrameIdx) const
 	{
 		UPaperFlipbook* Flipbook = FlipbookBeingEdited.Get();
-		if (Flipbook && Flipbook->KeyFrames.IsValidIndex(KeyFrameIdx))
+		if (Flipbook && Flipbook->IsValidKeyFrameIndex(KeyFrameIdx))
 		{
-			const FPaperFlipbookKeyFrame& KeyFrame = Flipbook->KeyFrames[KeyFrameIdx];
+			const FPaperFlipbookKeyFrame& KeyFrame = Flipbook->GetKeyFrameChecked(KeyFrameIdx);
 			return FMath::Max<float>(0, KeyFrame.FrameRun * SlateUnitsPerFrame.Get() - HandleWidth);
 		}
 		else
diff --git a/Engine/Plugins/Experimental/Paper2D/Source/Paper2DEditor/Private/PaperAnimatedSpriteThumbnailRenderer.cpp b/Engine/Plugins/Experimental/Paper2D/Source/Paper2DEditor/Private/PaperAnimatedSpriteThumbnailRenderer.cpp
--- a/Engine/Plugins/Experimental/Paper2D/Source/Paper2DEditor/Private/PaperAnimatedSpriteThumbnailRenderer.cpp
+++ b/Engine/Plugins/Experimental/Paper2D/Source/Paper2DEditor/Private/PaperAnimatedSpriteThumbnailRenderer.cpp
@@ -14,16 +14,29 @@ void UPaperFlipbookThumbnailRenderer::Draw(UObject* Object, int32 X, int32 Y, ui
 {
 	if (UPaperFlipbook* Flipbook = Cast<UPaperFlipbook>(Object))
 	{
-		if (Flipbook->KeyFrames.Num() > 0)
+		const double DeltaTime = FApp::GetCurrentTime() - GStartTime;
+		const float TotalDuration = Flipbook->GetTotalDuration();
+		const float PlayTime = (TotalDuration > 0.0f) ? FMath::Fmod(DeltaTime, TotalDuration) : 0.0f;
+
+		if (UPaperSprite* Sprite = Flipbook->GetSpriteAtTime(PlayTime))
+		{
+			DrawFrame(Sprite, X, Y, Width, Height, RenderTarget, Canvas);
+			return;
+		}
+		else
 		{
-			const double DeltaTime = FApp::GetCurrentTime() - GStartTime;
-			const float PlayTime = FMath::Fmod(DeltaTime, Flipbook->GetDuration());
+			// Fallback for empty frames or newly created flipbooks
+			DrawGrid(X, Y, Width, Height, Canvas);
+		}
 
-			UPaperSprite* Sprite = Flipbook->GetSpriteAtTime(PlayTime);
-			if ( Sprite )
-			{
-				DrawFrame(Sprite, X, Y, Width, Height, RenderTarget, Canvas);
-			}
+		if (TotalDuration == 0.0f)
+		{
+			// Warning text for no frames
+			const FText ErrorText = NSLOCTEXT("FlipbookEditorApp", "ThumbnailWarningNoFrames", "No frames");
+			FCanvasTextItem TextItem(FVector2D(5.0f, 5.0f), ErrorText, GEngine->GetLargeFont(), FLinearColor::Red);
+			TextItem.EnableShadow(FLinearColor::Black);
+			TextItem.Scale = FVector2D(Width / 128.0f, Height / 128.0f);
+			TextItem.Draw(Canvas);
 		}
 	}
 }
\ No newline at end of file
diff --git a/Engine/Plugins/Experimental/Paper2D/Source/Paper2DEditor/Private/PaperSpriteThumbnailRenderer.cpp b/Engine/Plugins/Experimental/Paper2D/Source/Paper2DEditor/Private/PaperSpriteThumbnailRenderer.cpp
--- a/Engine/Plugins/Experimental/Paper2D/Source/Paper2DEditor/Private/PaperSpriteThumbnailRenderer.cpp
+++ b/Engine/Plugins/Experimental/Paper2D/Source/Paper2DEditor/Private/PaperSpriteThumbnailRenderer.cpp
@@ -16,6 +16,30 @@ void UPaperSpriteThumbnailRenderer::Draw(UObject* Object, int32 X, int32 Y, uint
 	DrawFrame(Sprite, X, Y, Width, Height, RenderTarget, Canvas);
 }
 
+void UPaperSpriteThumbnailRenderer::DrawGrid(int32 X, int32 Y, uint32 Width, uint32 Height, FCanvas* Canvas)
+{
+	static UTexture2D* GridTexture = NULL;
+	if (GridTexture == NULL)
+	{
+		GridTexture = LoadObject<UTexture2D>(NULL, TEXT("/Engine/EngineMaterials/DefaultWhiteGrid.DefaultWhiteGrid"), NULL, LOAD_None, NULL);
+	}
+
+	const bool bAlphaBlend = false;
+
+	Canvas->DrawTile(
+		(float)X,
+		(float)Y,
+		(float)Width,
+		(float)Height,
+		0.0f,
+		0.0f,
+		4.0f,
+		4.0f,
+		FLinearColor::White,
+		GridTexture->Resource,
+		bAlphaBlend);
+}
+
 void UPaperSpriteThumbnailRenderer::DrawFrame(class UPaperSprite* Sprite, int32 X, int32 Y, uint32 Width, uint32 Height, FRenderTarget*, FCanvas* Canvas)
 {
 	if (const UTexture2D* SourceTexture = (Sprite != NULL) ? Sprite->GetSourceTexture() : NULL)
@@ -25,26 +49,7 @@ void UPaperSpriteThumbnailRenderer::DrawFrame(class UPaperSprite* Sprite, int32
 		// Draw the grid behind the sprite
 		if (bUseTranslucentBlend)
 		{
-			static UTexture2D* GridTexture = NULL;
-			if (GridTexture == NULL)
-			{
-				GridTexture = LoadObject<UTexture2D>(NULL, TEXT("/Engine/EngineMaterials/DefaultWhiteGrid.DefaultWhiteGrid"), NULL, LOAD_None, NULL);
-			}
-
-			const bool bAlphaBlend = false;
-
-			Canvas->DrawTile(
-				(float)X,
-				(float)Y,
-				(float)Width,
-				(float)Height,
-				0.0f,
-				0.0f,
-				4.0f,
-				4.0f,
-				FLinearColor::White,
-				GridTexture->Resource,
-				bAlphaBlend);
+			DrawGrid(X, Y, Width, Height, Canvas);
 		}
 
 		// Draw the sprite itself
@@ -86,4 +91,9 @@ void UPaperSpriteThumbnailRenderer::DrawFrame(class UPaperSprite* Sprite, int32
 			SourceTexture->Resource,
 			bUseTranslucentBlend);
 	}
+	else
+	{
+		// Fallback for a bogus sprite
+		DrawGrid(X, Y, Width, Height, Canvas);
+	}
 }
\ No newline at end of file
diff --git a/Engine/Plugins/Experimental/Paper2D/Source/PaperJsonImporter/Private/PaperJsonImporterFactory.cpp b/Engine/Plugins/Experimental/Paper2D/Source/PaperJsonImporter/Private/PaperJsonImporterFactory.cpp
--- a/Engine/Plugins/Experimental/Paper2D/Source/PaperJsonImporter/Private/PaperJsonImporterFactory.cpp
+++ b/Engine/Plugins/Experimental/Paper2D/Source/PaperJsonImporter/Private/PaperJsonImporterFactory.cpp
@@ -248,60 +248,61 @@ UObject* UPaperJsonImporterFactory::FactoryCreateText(UClass* InClass, UObject*
 				{
 					Flipbook = NewNamedObject<UPaperFlipbook>(InParent, InName, Flags);
 					Result = Flipbook;
-				}
-
-				GWarn->BeginSlowTask(NSLOCTEXT("Paper2D", "PaperJsonImporterFactory_ImportingSprites", "Importing Sprite Frame"), true, true);
 
-				// Create objects for each successfully parsed frame
-				const int32 FrameRun = 1; //@TODO: Don't make a keyframe for every single item if we can help it
-				for (int32 FrameIndex = 0; FrameIndex < ParsedFrames.Num(); ++FrameIndex)
-				{
-					GWarn->StatusUpdate(FrameIndex, ParsedFrames.Num(), NSLOCTEXT("Paper2D", "PaperJsonImporterFactory_ImportingSprites", "Importing Sprite Frames"));
+					GWarn->BeginSlowTask(NSLOCTEXT("Paper2D", "PaperJsonImporterFactory_ImportingSprites", "Importing Sprite Frame"), true, true);
+					FScopedFlipbookMutator EditLock(Flipbook);
 
-					// Check for user canceling the import
-					if (GWarn->ReceivedUserCancel())
+					// Create objects for each successfully parsed frame
+					const int32 FrameRun = 1; //@TODO: Don't make a keyframe for every single item if we can help it
+					for (int32 FrameIndex = 0; FrameIndex < ParsedFrames.Num(); ++FrameIndex)
 					{
-						break;
-					}
+						GWarn->StatusUpdate(FrameIndex, ParsedFrames.Num(), NSLOCTEXT("Paper2D", "PaperJsonImporterFactory_ImportingSprites", "Importing Sprite Frames"));
 
-					const FSpriteFrame& Frame = ParsedFrames[FrameIndex];
+						// Check for user canceling the import
+						if (GWarn->ReceivedUserCancel())
+						{
+							break;
+						}
 
-					// Create a package for the frame
-					const FString TargetSubPath = LongPackagePath + TEXT("/Frames");
+						const FSpriteFrame& Frame = ParsedFrames[FrameIndex];
 
-					UObject* OuterForFrame = NULL; // @TODO: Use this if we don't want them to be individual assets - Flipbook;
+						// Create a package for the frame
+						const FString TargetSubPath = LongPackagePath + TEXT("/Frames");
 
-					// Create a unique package name and asset name for the frame
-					const FString TentativePackagePath = PackageTools::SanitizePackageName(TargetSubPath + TEXT("/") + Frame.FrameName.ToString());
-					FString DefaultSuffix;
-					FString AssetName;
-					FString PackageName;
-					AssetToolsModule.Get().CreateUniqueAssetName(TentativePackagePath, /*out*/ DefaultSuffix, /*out*/ PackageName, /*out*/ AssetName);
+						UObject* OuterForFrame = NULL; // @TODO: Use this if we don't want them to be individual assets - Flipbook;
 
-					// Create a package for the frame
-					OuterForFrame = CreatePackage(NULL, *PackageName);
+						// Create a unique package name and asset name for the frame
+						const FString TentativePackagePath = PackageTools::SanitizePackageName(TargetSubPath + TEXT("/") + Frame.FrameName.ToString());
+						FString DefaultSuffix;
+						FString AssetName;
+						FString PackageName;
+						AssetToolsModule.Get().CreateUniqueAssetName(TentativePackagePath, /*out*/ DefaultSuffix, /*out*/ PackageName, /*out*/ AssetName);
 
-					// Create a frame in the package
-					UPaperSprite* TargetSprite = NewNamedObject<UPaperSprite>(OuterForFrame, *AssetName, Flags);
-					FAssetRegistryModule::AssetCreated(TargetSprite);
+						// Create a package for the frame
+						OuterForFrame = CreatePackage(NULL, *PackageName);
 
-					TargetSprite->Modify();
+						// Create a frame in the package
+						UPaperSprite* TargetSprite = NewNamedObject<UPaperSprite>(OuterForFrame, *AssetName, Flags);
+						FAssetRegistryModule::AssetCreated(TargetSprite);
 
-					TargetSprite->InitializeSprite(ImportedTexture, Frame.SpritePosInSheet, Frame.SpriteSizeInSheet);
+						TargetSprite->Modify();
 
-					//@TODO: Need to support pivot behavior - Total guess at pivot behavior
-					//const FVector2D SizeDifference = Frame.SpriteSourceSize - Frame.SpriteSizeInSheet;
-					//TargetSprite->SpriteData.Destination = FVector(SizeDifference.X * -0.5f, 0.0f, SizeDifference.Y * -0.5f);
+						TargetSprite->InitializeSprite(ImportedTexture, Frame.SpritePosInSheet, Frame.SpriteSizeInSheet);
 
-					// Create the entry in the animation
-					FPaperFlipbookKeyFrame* DestFrame = new (Flipbook->KeyFrames) FPaperFlipbookKeyFrame();
-					DestFrame->Sprite = TargetSprite;
-					DestFrame->FrameRun = FrameRun;
+						//@TODO: Need to support pivot behavior - Total guess at pivot behavior
+						//const FVector2D SizeDifference = Frame.SpriteSourceSize - Frame.SpriteSizeInSheet;
+						//TargetSprite->SpriteData.Destination = FVector(SizeDifference.X * -0.5f, 0.0f, SizeDifference.Y * -0.5f);
 
-					TargetSprite->PostEditChange();
-				}
+						// Create the entry in the animation
+						FPaperFlipbookKeyFrame* DestFrame = new (EditLock.KeyFrames) FPaperFlipbookKeyFrame();
+						DestFrame->Sprite = TargetSprite;
+						DestFrame->FrameRun = FrameRun;
 
-				GWarn->EndSlowTask();
+						TargetSprite->PostEditChange();
+					}
+
+					GWarn->EndSlowTask();
+				}
 			}
 			else
 			{
