Paper2D: Fixed order-of-initialization warning on SourceFlipbook

diff --git a/Engine/Plugins/Experimental/Paper2D/Source/Paper2D/Classes/PaperFlipbook.h b/Engine/Plugins/Experimental/Paper2D/Source/Paper2D/Classes/PaperFlipbook.h
--- a/Engine/Plugins/Experimental/Paper2D/Source/Paper2D/Classes/PaperFlipbook.h
+++ b/Engine/Plugins/Experimental/Paper2D/Source/Paper2D/Classes/PaperFlipbook.h
@@ -108,9 +108,9 @@ private:
 
 public:
 	FScopedFlipbookMutator(UPaperFlipbook* InFlipbook)
-		: SourceFlipbook(InFlipbook)
-		, FramesPerSecond(InFlipbook->FramesPerSecond)
+		: FramesPerSecond(InFlipbook->FramesPerSecond)
 		, KeyFrames(InFlipbook->KeyFrames)
+		, SourceFlipbook(InFlipbook)
 	{
 	}
 
