Paper2D: Add experimental designation warning when setting the collision domain on a sprite to 2D

diff --git a/Engine/Plugins/Experimental/Paper2D/Source/Paper2DEditor/Private/SpriteEditor/SpriteDetailsCustomization.cpp b/Engine/Plugins/Experimental/Paper2D/Source/Paper2DEditor/Private/SpriteEditor/SpriteDetailsCustomization.cpp
--- a/Engine/Plugins/Experimental/Paper2D/Source/Paper2DEditor/Private/SpriteEditor/SpriteDetailsCustomization.cpp
+++ b/Engine/Plugins/Experimental/Paper2D/Source/Paper2DEditor/Private/SpriteEditor/SpriteDetailsCustomization.cpp
@@ -4,6 +4,9 @@
 #include "SpriteDetailsCustomization.h"
 
 #include "PhysicsEngine/BodySetup.h"
+#include "Editor/Documentation/Public/IDocumentation.h"
+
+#define LOCTEXT_NAMESPACE "SpriteEditor"
 
 //////////////////////////////////////////////////////////////////////////
 // FSpriteDetailsCustomization
@@ -13,38 +16,57 @@ TSharedRef<IDetailCustomization> FSpriteDetailsCustomization::MakeInstance()
 	return MakeShareable(new FSpriteDetailsCustomization);
 }
 
+FDetailWidgetRow& FSpriteDetailsCustomization::GenerateWarningRow(IDetailCategoryBuilder& WarningCategory, bool bExperimental, const FText& WarningText, const FText& Tooltip, const FString& ExcerptLink, const FString& ExcerptName)
+{
+	const FString SearchString = WarningText.ToString();
+	const FSlateBrush* WarningIcon = FEditorStyle::GetBrush(bExperimental ? "PropertyEditor.ExperimentalClass" : "PropertyEditor.EarlyAccessClass");
+
+	FDetailWidgetRow& WarningRow = WarningCategory.AddCustomRow(SearchString)
+		.WholeRowContent()
+		[
+			SNew(SHorizontalBox)
+			.ToolTip(IDocumentation::Get()->CreateToolTip(Tooltip, nullptr, ExcerptLink, ExcerptName))
+			.Visibility(EVisibility::Visible)
+
+			+SHorizontalBox::Slot()
+			.VAlign(VAlign_Center)
+			.AutoWidth()
+			.Padding(4.0f, 0.0f, 0.0f, 0.0f)
+			[	SNew(SImage)
+				.Image(WarningIcon)
+			]
+
+			+SHorizontalBox::Slot()
+			.VAlign(VAlign_Center)
+			.AutoWidth()
+			.Padding(4.0f, 0.0f, 0.0f, 0.0f)
+			[
+				SNew(STextBlock)
+				.Text(WarningText)
+				.Font(IDetailLayoutBuilder::GetDetailFont())
+			]
+		];
+
+	return WarningRow;
+}
+
 void FSpriteDetailsCustomization::CustomizeDetails(IDetailLayoutBuilder& DetailLayout)
 {
 	// Make sure sprite properties are near the top
 	IDetailCategoryBuilder& SpriteCategory = DetailLayout.EditCategory("Sprite", TEXT(""), ECategoryPriority::TypeSpecific);
+	BuildSpriteSection(SpriteCategory, DetailLayout);
 
-	TSharedRef<IPropertyHandle> RenderGeometry = DetailLayout.GetProperty(GET_MEMBER_NAME_CHECKED(UPaperSprite, RenderGeometry));
-
+	// Build the rendering category
 	IDetailCategoryBuilder& RenderingCategory = DetailLayout.EditCategory("Rendering");
-	IDetailPropertyRow& RenderGeometryProperty = RenderingCategory.AddProperty(RenderGeometry);
+	BuildRenderingSection(RenderingCategory, DetailLayout);
 
 	// Build the collision category
 	IDetailCategoryBuilder& CollisionCategory = DetailLayout.EditCategory("Collision");
-	
-	TSharedPtr<IPropertyHandle> SpriteCollisionDomainProperty = DetailLayout.GetProperty(GET_MEMBER_NAME_CHECKED(UPaperSprite, SpriteCollisionDomain));
-	TAttribute<EVisibility> ParticipatesInPhysics = TAttribute<EVisibility>::Create( TAttribute<EVisibility>::FGetter::CreateSP( this, &FSpriteDetailsCustomization::AnyPhysicsMode, SpriteCollisionDomainProperty) ) ;
-	TAttribute<EVisibility> ParticipatesInPhysics3D = TAttribute<EVisibility>::Create( TAttribute<EVisibility>::FGetter::CreateSP( this, &FSpriteDetailsCustomization::PhysicsModeMatches, SpriteCollisionDomainProperty, ESpriteCollisionMode::Use3DPhysics) ) ;
-
-	CollisionCategory.AddProperty(SpriteCollisionDomainProperty);
-	CollisionCategory.AddProperty( DetailLayout.GetProperty(GET_MEMBER_NAME_CHECKED(UPaperSprite, CollisionGeometry)) )
-		.Visibility(ParticipatesInPhysics);
-	CollisionCategory.AddProperty( DetailLayout.GetProperty(GET_MEMBER_NAME_CHECKED(UPaperSprite, CollisionThickness)) )
-		.Visibility(ParticipatesInPhysics3D);
-
-	// Hide some useless body setup properties
-// 	CollisionCategory.AddProperty( DetailLayout.GetProperty(GET_MEMBER_NAME_CHECKED(UPaperSprite, BodySetup3D)) )
-// 		.Visibility(ParticipatesInPhysics3D);
-// 	TSharedRef<IPropertyHandle> BodySetup3D = DetailLayout.GetProperty(GET_MEMBER_NAME_CHECKED(UPaperSprite, BodySetup3D));
-// 	DetailLayout.HideProperty(BodySetup3D->GetChildHandle(GET_MEMBER_NAME_CHECKED(UBodySetup, DefaultInstance)));
-// 	DetailLayout.HideProperty(BodySetup3D->GetChildHandle(GET_MEMBER_NAME_CHECKED(UBodySetup, BoneName)));
-// 	DetailLayout.HideProperty(BodySetup3D->GetChildHandle(GET_MEMBER_NAME_CHECKED(UBodySetup, PhysicsType)));
-// 	DetailLayout.HideProperty(BodySetup3D->GetChildHandle(GET_MEMBER_NAME_CHECKED(UBodySetup, bConsiderForBounds)));
+	BuildCollisionSection(CollisionCategory, DetailLayout);
+}
 
+void FSpriteDetailsCustomization::BuildSpriteSection(IDetailCategoryBuilder& SpriteCategory, IDetailLayoutBuilder& DetailLayout)
+{
 	// Show other normal properties in the sprite category so that desired ordering doesn't get messed up
 	SpriteCategory.AddProperty(GET_MEMBER_NAME_CHECKED(UPaperSprite, SourceUV));
 	SpriteCategory.AddProperty(GET_MEMBER_NAME_CHECKED(UPaperSprite, SourceDimension));
@@ -55,7 +77,7 @@ void FSpriteDetailsCustomization::CustomizeDetails(IDetailLayoutBuilder& DetailL
 
 	// Show/hide the experimental atlas group support based on whether or not it is enabled
 	TSharedPtr<IPropertyHandle> AtlasGroupProperty = DetailLayout.GetProperty(GET_MEMBER_NAME_CHECKED(UPaperSprite, AtlasGroup));
-	TAttribute<EVisibility> AtlasGroupPropertyVisibility = TAttribute<EVisibility>::Create(TAttribute<EVisibility>::FGetter::CreateSP(this, &FSpriteDetailsCustomization::GetAtlasGroupVisibility, AtlasGroupProperty));
+	TAttribute<EVisibility> AtlasGroupPropertyVisibility = TAttribute<EVisibility>::Create(TAttribute<EVisibility>::FGetter::CreateStatic(&FSpriteDetailsCustomization::GetAtlasGroupVisibility));
 	SpriteCategory.AddProperty(AtlasGroupProperty, EPropertyLocation::Advanced).Visibility(AtlasGroupPropertyVisibility);
 
 	// Show/hide the custom pivot point based on the pivot mode
@@ -66,6 +88,43 @@ void FSpriteDetailsCustomization::CustomizeDetails(IDetailLayoutBuilder& DetailL
 	SpriteCategory.AddProperty(CustomPivotPointProperty).Visibility(CustomPivotPointVisibility);
 }
 
+void FSpriteDetailsCustomization::BuildRenderingSection(IDetailCategoryBuilder& RenderingCategory, IDetailLayoutBuilder& DetailLayout)
+{
+	TSharedRef<IPropertyHandle> RenderGeometry = DetailLayout.GetProperty(GET_MEMBER_NAME_CHECKED(UPaperSprite, RenderGeometry));
+	IDetailPropertyRow& RenderGeometryProperty = RenderingCategory.AddProperty(RenderGeometry);
+}
+
+void FSpriteDetailsCustomization::BuildCollisionSection(IDetailCategoryBuilder& CollisionCategory, IDetailLayoutBuilder& DetailLayout)
+{
+	TSharedPtr<IPropertyHandle> SpriteCollisionDomainProperty = DetailLayout.GetProperty(GET_MEMBER_NAME_CHECKED(UPaperSprite, SpriteCollisionDomain));
+	TAttribute<EVisibility> ParticipatesInPhysics = TAttribute<EVisibility>::Create( TAttribute<EVisibility>::FGetter::CreateSP( this, &FSpriteDetailsCustomization::AnyPhysicsMode, SpriteCollisionDomainProperty) ) ;
+	TAttribute<EVisibility> ParticipatesInPhysics3D = TAttribute<EVisibility>::Create(TAttribute<EVisibility>::FGetter::CreateSP(this, &FSpriteDetailsCustomization::PhysicsModeMatches, SpriteCollisionDomainProperty, ESpriteCollisionMode::Use3DPhysics));
+	TAttribute<EVisibility> ParticipatesInPhysics2D = TAttribute<EVisibility>::Create(TAttribute<EVisibility>::FGetter::CreateSP(this, &FSpriteDetailsCustomization::PhysicsModeMatches, SpriteCollisionDomainProperty, ESpriteCollisionMode::Use2DPhysics));
+
+	CollisionCategory.AddProperty(SpriteCollisionDomainProperty);
+
+	// Add a warning bar about 2D collision being experimental
+	FText WarningFor2D = LOCTEXT("Experimental2DPhysicsWarning", "2D collision support is *experimental*");
+	FText TooltipFor2D = LOCTEXT("Experimental2DPhysicsWarningTooltip", "2D collision support is *experimental* and should not be relied on yet.\n\nRigid body collision detection and response works, but there are only precompiled libraries for Windows currently.\n\nRaycasts are partially supported (and need to be enabled in project settings), but queries, sweeps, or overlap tests are not implemented yet.");
+	GenerateWarningRow(CollisionCategory, /*bExperimental=*/ true, WarningFor2D, TooltipFor2D, TEXT("Shared/Editors/SpriteEditor"), TEXT("CollisionDomain2DWarning"))
+		.Visibility(ParticipatesInPhysics2D);
+
+	// Add a warning bar if 2D collision queries aren't enabled
+	TAttribute<EVisibility> WarnAbout2DQueriesBeingDisabledVisibility = TAttribute<EVisibility>::Create(TAttribute<EVisibility>::FGetter::CreateStatic(&FSpriteDetailsCustomization::Get2DPhysicsNotEnabledWarningVisibility));
+	FText QueryWarningFor2D = LOCTEXT("Query2DPhysicsWarning", "2D collision queries are disabled");
+	FText QueryTooltipFor2D = LOCTEXT("Query2DPhysicsWarningTooltip", "You can enable 2D queries in Project Settings..Physics by setting bEnable2DPhysics to true, otherwise only collision detection and response will work.\n\nNote: Only raycasts are partially supported; other queries, sweeps, and overlap tests are not implemented yet.");
+	GenerateWarningRow(CollisionCategory, /*bExperimental=*/ false, QueryWarningFor2D, QueryTooltipFor2D, TEXT("Shared/Editors/SpriteEditor"), TEXT("Disabled2DCollisionQueriesWarning"))
+		.Visibility(WarnAbout2DQueriesBeingDisabledVisibility);
+
+	// Show the collision geometry when not None
+	CollisionCategory.AddProperty( DetailLayout.GetProperty(GET_MEMBER_NAME_CHECKED(UPaperSprite, CollisionGeometry)) )
+		.Visibility(ParticipatesInPhysics);
+
+	// Show the collision thickness only in 3D mode
+	CollisionCategory.AddProperty( DetailLayout.GetProperty(GET_MEMBER_NAME_CHECKED(UPaperSprite, CollisionThickness)) )
+		.Visibility(ParticipatesInPhysics3D);
+}
+
 EVisibility FSpriteDetailsCustomization::PhysicsModeMatches(TSharedPtr<IPropertyHandle> Property, ESpriteCollisionMode::Type DesiredMode) const
 {
 	if (Property.IsValid())
@@ -100,7 +159,7 @@ EVisibility FSpriteDetailsCustomization::AnyPhysicsMode(TSharedPtr<IPropertyHand
 	return EVisibility::Visible;
 }
 
-EVisibility FSpriteDetailsCustomization::GetAtlasGroupVisibility(TSharedPtr<IPropertyHandle> Property) const
+EVisibility FSpriteDetailsCustomization::GetAtlasGroupVisibility()
 {
 	return GetDefault<UPaperRuntimeSettings>()->bEnableSpriteAtlasGroups ? EVisibility::Visible : EVisibility::Collapsed;
 }
@@ -121,3 +180,10 @@ EVisibility FSpriteDetailsCustomization::GetCustomPivotVisibility(TSharedPtr<IPr
 	// If there are multiple values, show all properties
 	return EVisibility::Visible;
 }
+
+EVisibility FSpriteDetailsCustomization::Get2DPhysicsNotEnabledWarningVisibility()
+{
+	return GetDefault<UPhysicsSettings>()->bEnable2DPhysics ? EVisibility::Collapsed : EVisibility::Visible;
+}
+
+#undef LOCTEXT_NAMESPACE
diff --git a/Engine/Plugins/Experimental/Paper2D/Source/Paper2DEditor/Private/SpriteEditor/SpriteDetailsCustomization.h b/Engine/Plugins/Experimental/Paper2D/Source/Paper2DEditor/Private/SpriteEditor/SpriteDetailsCustomization.h
--- a/Engine/Plugins/Experimental/Paper2D/Source/Paper2DEditor/Private/SpriteEditor/SpriteDetailsCustomization.h
+++ b/Engine/Plugins/Experimental/Paper2D/Source/Paper2DEditor/Private/SpriteEditor/SpriteDetailsCustomization.h
@@ -20,6 +20,14 @@ public:
 protected:
 	EVisibility PhysicsModeMatches(TSharedPtr<IPropertyHandle> Property, ESpriteCollisionMode::Type DesiredMode) const;
 	EVisibility AnyPhysicsMode(TSharedPtr<IPropertyHandle> Property) const;
-	EVisibility GetAtlasGroupVisibility(TSharedPtr<IPropertyHandle> Property) const;
 	EVisibility GetCustomPivotVisibility(TSharedPtr<IPropertyHandle> Property) const;
+
+	static EVisibility GetAtlasGroupVisibility();
+	static EVisibility Get2DPhysicsNotEnabledWarningVisibility();
+
+	static FDetailWidgetRow& GenerateWarningRow(IDetailCategoryBuilder& WarningCategory, bool bExperimental, const FText& WarningText, const FText& Tooltip, const FString& ExcerptLink, const FString& ExcerptName);
+
+	void BuildSpriteSection(IDetailCategoryBuilder& SpriteCategory, IDetailLayoutBuilder& DetailLayout);
+	void BuildCollisionSection(IDetailCategoryBuilder& CollisionCategory, IDetailLayoutBuilder& DetailLayout);
+	void BuildRenderingSection(IDetailCategoryBuilder& RenderingCategory, IDetailLayoutBuilder& DetailLayout);
 };
