Paper2D: Add display of sprite render bounds extents to viewport in sprite editor

diff --git a/Engine/Plugins/Experimental/Paper2D/Source/Paper2DEditor/Private/SpriteEditor/SpriteEditorViewportClient.cpp b/Engine/Plugins/Experimental/Paper2D/Source/Paper2DEditor/Private/SpriteEditor/SpriteEditorViewportClient.cpp
--- a/Engine/Plugins/Experimental/Paper2D/Source/Paper2DEditor/Private/SpriteEditor/SpriteEditorViewportClient.cpp
+++ b/Engine/Plugins/Experimental/Paper2D/Source/Paper2DEditor/Private/SpriteEditor/SpriteEditorViewportClient.cpp
@@ -347,6 +347,28 @@ void FSpriteEditorViewportClient::DrawGeometryStats(FViewport& InViewport, FScen
 	YPos = (int32)TextItem.Position.Y;
 }
 
+void FSpriteEditorViewportClient::DrawBoundsAsText(FViewport& InViewport, FSceneView& View, FCanvas& Canvas, int32& YPos)
+{
+	FNumberFormattingOptions NoDigitGroupingFormat;
+	NoDigitGroupingFormat.UseGrouping = false;
+
+	UPaperSprite* Sprite = GetSpriteBeingEdited();
+	FBoxSphereBounds Bounds = Sprite->GetRenderBounds();
+
+	const FText DisplaySizeText = FText::Format(LOCTEXT("BoundsSize", "Approx. Size: {0}x{1}x{2}"),
+		FText::AsNumber((int32)(Bounds.BoxExtent.X * 2.0f), &NoDigitGroupingFormat),
+		FText::AsNumber((int32)(Bounds.BoxExtent.Y * 2.0f), &NoDigitGroupingFormat),
+		FText::AsNumber((int32)(Bounds.BoxExtent.Z * 2.0f), &NoDigitGroupingFormat));
+
+	Canvas.DrawShadowedString(
+		6,
+		YPos,
+		*DisplaySizeText.ToString(),
+		GEngine->GetSmallFont(),
+		FLinearColor::White);
+	YPos += 18;
+}
+
 void FSpriteEditorViewportClient::DrawSockets(const FSceneView* View, FPrimitiveDrawInterface* PDI)
 {
 	UPaperSprite* Sprite = GetSpriteBeingEdited();
@@ -461,6 +483,10 @@ void FSpriteEditorViewportClient::DrawCanvas(FViewport& Viewport, FSceneView& Vi
 		// As well as all of the geometry
 		DrawGeometryStats(Viewport, View, Canvas, Sprite->CollisionGeometry, false, /*inout*/ YPos);
 		DrawGeometryStats(Viewport, View, Canvas, Sprite->RenderGeometry, true, /*inout*/ YPos);
+
+		// And bounds
+		DrawBoundsAsText(Viewport, View, Canvas, /*inout*/ YPos);
+
 		break;
 	case ESpriteEditorMode::EditCollisionMode:
 		{
@@ -492,6 +518,9 @@ void FSpriteEditorViewportClient::DrawCanvas(FViewport& Viewport, FSceneView& Vi
 			const FLinearColor RenderGeomColor(1.0f, 0.2f, 0.0f, 1.0f);
 			DrawGeometry(Viewport, View, Canvas, Sprite->RenderGeometry, RenderGeomColor, true);
 			DrawGeometryStats(Viewport, View, Canvas, Sprite->RenderGeometry, true, /*inout*/ YPos);
+
+			// And bounds
+			DrawBoundsAsText(Viewport, View, Canvas, /*inout*/ YPos);
 		}
 		break;
 	case ESpriteEditorMode::EditSourceRegionMode:
diff --git a/Engine/Plugins/Experimental/Paper2D/Source/Paper2DEditor/Private/SpriteEditor/SpriteEditorViewportClient.h b/Engine/Plugins/Experimental/Paper2D/Source/Paper2DEditor/Private/SpriteEditor/SpriteEditorViewportClient.h
--- a/Engine/Plugins/Experimental/Paper2D/Source/Paper2DEditor/Private/SpriteEditor/SpriteEditorViewportClient.h
+++ b/Engine/Plugins/Experimental/Paper2D/Source/Paper2DEditor/Private/SpriteEditor/SpriteEditorViewportClient.h
@@ -162,6 +162,7 @@ private:
 	FVector TextureSpaceToWorldSpace(const FVector2D& SourcePoint) const;
 
 	void DrawTriangleList(FViewport& InViewport, FSceneView& View, FCanvas& Canvas, const TArray<FVector2D>& Triangles);
+	void DrawBoundsAsText(FViewport& InViewport, FSceneView& View, FCanvas& Canvas, int32& YPos);
 	void DrawGeometry(FViewport& InViewport, FSceneView& View, FCanvas& Canvas, const FSpritePolygonCollection& Geometry, const FLinearColor& GeometryVertexColor, bool bIsRenderGeometry);
 	void DrawGeometryStats(FViewport& InViewport, FSceneView& View, FCanvas& Canvas, const FSpritePolygonCollection& Geometry, bool bIsRenderGeometry, int32& YPos);
 	void DrawSockets(const FSceneView* View, FPrimitiveDrawInterface* PDI);
