Paper2D: Fix sprite render bounds bug that could result in a box with Min > Max

diff --git a/Engine/Plugins/Experimental/Paper2D/Source/Paper2D/Private/PaperSprite.cpp b/Engine/Plugins/Experimental/Paper2D/Source/Paper2D/Private/PaperSprite.cpp
--- a/Engine/Plugins/Experimental/Paper2D/Source/Paper2D/Private/PaperSprite.cpp
+++ b/Engine/Plugins/Experimental/Paper2D/Source/Paper2D/Private/PaperSprite.cpp
@@ -1071,8 +1071,8 @@ FBoxSphereBounds UPaperSprite::GetRenderBounds() const
 	
 	// Make the whole thing a single unit 'deep'
 	const FVector HalfThicknessVector = 0.5f * PaperAxisZ;
-	BoundingBox.Min -= HalfThicknessVector;
-	BoundingBox.Max += HalfThicknessVector;
+	BoundingBox += -HalfThicknessVector;
+	BoundingBox += HalfThicknessVector;
 
 	return FBoxSphereBounds(BoundingBox);
 }
