Paper2D: Abbreviated PixelsPerUnrealUnit to "Pixels / unit" in details panel

diff --git a/Engine/Plugins/Experimental/Paper2D/Source/Paper2D/Classes/PaperSprite.h b/Engine/Plugins/Experimental/Paper2D/Source/Paper2D/Classes/PaperSprite.h
--- a/Engine/Plugins/Experimental/Paper2D/Source/Paper2D/Classes/PaperSprite.h
+++ b/Engine/Plugins/Experimental/Paper2D/Source/Paper2D/Classes/PaperSprite.h
@@ -80,7 +80,7 @@ protected:
 	TEnumAsByte<ESpriteCollisionMode::Type> SpriteCollisionDomain;
 
 	// The scaling factor between pixels and Unreal units (cm) (e.g., 0.64 would make a 64 pixel wide sprite take up 100 cm)
-	UPROPERTY(Category = Sprite, EditAnywhere)
+	UPROPERTY(Category=Sprite, EditAnywhere, meta = (DisplayName = "Pixels per unit"))
 	float PixelsPerUnrealUnit;
 
 public:
