Paper2D: Added a number of playback controls to UPaperFlipbookComponent (cribbing from UTimelineComponent)

diff --git a/Engine/Plugins/Experimental/Paper2D/Source/Paper2D/Classes/PaperFlipbookComponent.h b/Engine/Plugins/Experimental/Paper2D/Source/Paper2D/Classes/PaperFlipbookComponent.h
--- a/Engine/Plugins/Experimental/Paper2D/Source/Paper2D/Classes/PaperFlipbookComponent.h
+++ b/Engine/Plugins/Experimental/Paper2D/Source/Paper2D/Classes/PaperFlipbookComponent.h
@@ -19,6 +19,19 @@ protected:
 
 	UPROPERTY(Category=Sprite, EditAnywhere, BlueprintReadWrite)
 	float PlayRate;
+
+	/** Whether the flipbook should loop when it reaches the end, or stop */
+	UPROPERTY()
+	uint32 bLooping:1;
+
+	/** If playback should move the current position backwards instead of forwards */
+	UPROPERTY()
+	uint32 bReversePlayback:1;
+
+	/** Are we currently playing (moving Position) */
+	UPROPERTY()
+	uint32 bPlaying:1;
+
 public:
 	/** Change the flipbook used by this instance. */
 	UFUNCTION(BlueprintCallable, Category="Sprite")
@@ -36,14 +49,76 @@ public:
 	UFUNCTION(BlueprintCallable, Category="Sprite")
 	void SetSpriteColor(FLinearColor NewColor);
 
-	//@TODO: Hackery for demo
-	UFUNCTION(BlueprintCallable, Category="Sprite")
-	void SetCurrentTime(float NewTime);
+	/** Start playback of flipbook */
+	UFUNCTION(BlueprintCallable, Category="Components|Flipbook")
+	void Play();
+
+	/** Start playback of flipbook from the start */
+	UFUNCTION(BlueprintCallable, Category="Components|Flipbook")
+	void PlayFromStart();
+
+	/** Start playback of flipbook in reverse */
+	UFUNCTION(BlueprintCallable, Category="Components|Flipbook")
+	void Reverse();
+
+	/** Start playback of flipbook in reverse from the end */
+	UFUNCTION(BlueprintCallable, Category="Components|Flipbook")
+	void ReverseFromEnd();
+
+	/** Stop playback of flipbook */
+	UFUNCTION(BlueprintCallable, Category="Components|Flipbook")
+	void Stop();
+
+	/** Get whether this flipbook is playing or not. */
+	UFUNCTION(BlueprintCallable, Category="Components|Flipbook")
+	bool IsPlaying() const;
+
+	/** Get whether we are reversing or not */
+	UFUNCTION(BlueprintCallable, Category="Components|Flipbook")
+	bool IsReversing() const;
+
+	/** Jump to a position in the flipbook. If bFireEvents is true, event functions will fire, otherwise will not. */
+	UFUNCTION(BlueprintCallable, Category="Components|Flipbook")
+	void SetPlaybackPosition(float NewPosition, bool bFireEvents);
+
+	/** Get the current playback position of the flipbook */
+	UFUNCTION(BlueprintCallable, Category="Components|Flipbook")
+	float GetPlaybackPosition() const;
+
+	/** true means we should loop, false means we should not. */
+	UFUNCTION(BlueprintCallable, Category="Components|Flipbook")
+	void SetLooping(bool bNewLooping);
+
+	/** Get whether we are looping or not */
+	UFUNCTION(BlueprintCallable, Category="Components|Flipbook")
+	bool IsLooping() const;
+
+	/** Sets the new play rate for this flipbook */
+	UFUNCTION(BlueprintCallable, Category="Components|Flipbook")
+	void SetPlayRate(float NewRate);
+
+	/** Get the current play rate for this flipbook */
+	UFUNCTION(BlueprintCallable, Category="Components|Flipbook")
+	float GetPlayRate() const;
+
+	/** Set the new playback position time to use */
+	UFUNCTION(BlueprintCallable, Category="Components|Flipbook")
+	void SetNewTime(float NewTime);
+
+	/** Get length of the flipbook */
+	UFUNCTION(BlueprintCallable, Category="Components|Flipbook")
+	float GetFlipbookLength() const;
 
 protected:
+	/** Current position in the timeline */
+	UPROPERTY()
 	float AccumulatedTime;
+
+	/** Last frame index calculated */
+	UPROPERTY()
 	int32 CachedFrameIndex;
 
+	/* Vertex color to apply to the frames */
 	UPROPERTY(BlueprintReadOnly, Interp, Category=Sprite)
 	FLinearColor SpriteColor;
 
@@ -51,6 +126,8 @@ protected:
 	void CalculateCurrentFrame();
 	UPaperSprite* GetSpriteAtCachedIndex() const;
 
+	void TickFlipbook(float DeltaTime);
+
 public:
 	// UActorComponent interface
 	virtual void TickComponent(float DeltaTime, enum ELevelTick TickType, FActorComponentTickFunction* ThisTickFunction) override;
diff --git a/Engine/Plugins/Experimental/Paper2D/Source/Paper2D/Private/PaperFlipbookComponent.cpp b/Engine/Plugins/Experimental/Paper2D/Source/Paper2D/Private/PaperFlipbookComponent.cpp
--- a/Engine/Plugins/Experimental/Paper2D/Source/Paper2D/Private/PaperFlipbookComponent.cpp
+++ b/Engine/Plugins/Experimental/Paper2D/Source/Paper2D/Private/PaperFlipbookComponent.cpp
@@ -26,6 +26,10 @@ UPaperFlipbookComponent::UPaperFlipbookComponent(const FPostConstructInitializeP
 	CachedFrameIndex = INDEX_NONE;
 	AccumulatedTime = 0.0f;
 	PlayRate = 1.0f;
+
+	bLooping = true;
+	bReversePlayback = false;
+	bPlaying = true;
 }
 
 UPaperSprite* UPaperFlipbookComponent::GetSpriteAtCachedIndex() const
@@ -96,20 +100,101 @@ void UPaperFlipbookComponent::CalculateCurrentFrame()
 	}
 }
 
-void UPaperFlipbookComponent::TickComponent(float DeltaTime, enum ELevelTick TickType, FActorComponentTickFunction* ThisTickFunction)
+void UPaperFlipbookComponent::TickFlipbook(float DeltaTime)
 {
-	// Advance time
-	const float TotalTime = (SourceFlipbook != nullptr) ? SourceFlipbook->GetTotalDuration() : 0.0f;
-	if (TotalTime > 0.0f)
+	bool bIsFinished = false;
+
+	if (bPlaying)
 	{
-		AccumulatedTime += DeltaTime * PlayRate;
-		AccumulatedTime = FMath::Fmod(AccumulatedTime, TotalTime);
+		float NewPosition = AccumulatedTime;
+		const float TimelineLength = GetFlipbookLength();
+
+		if (!bReversePlayback)
+		{
+			NewPosition = AccumulatedTime + (DeltaTime * PlayRate);
+
+			if (NewPosition > TimelineLength)
+			{
+				if (bLooping)
+				{
+					// If looping, play to end, jump to start, and set target to somewhere near the beginning.
+					SetPlaybackPosition(TimelineLength, true);
+					SetPlaybackPosition(0.0f, false);
+
+					if (TimelineLength > 0.0f)
+					{
+						while (NewPosition > TimelineLength)
+						{
+							NewPosition -= TimelineLength;
+						}
+					}
+					else
+					{
+						NewPosition = 0.0f;
+					}
+				}
+				else
+				{
+					// If not looping, snap to end and stop playing.
+					NewPosition = TimelineLength;
+					Stop();
+					bIsFinished = true;
+				}
+			}
+		}
+		else
+		{
+			NewPosition = AccumulatedTime - (DeltaTime * PlayRate);
+
+			if (NewPosition < 0.0f)
+			{
+				if (bLooping)
+				{
+					// If looping, play to start, jump to end, and set target to somewhere near the end.
+					SetPlaybackPosition(0.0f, true);
+					SetPlaybackPosition(TimelineLength, false);
+
+					if (TimelineLength > 0.0f)
+					{
+						while (NewPosition < 0.0f)
+						{
+							NewPosition += TimelineLength;
+						}
+					}
+					else
+					{
+						NewPosition = 0.0f;
+					}
+				}
+				else
+				{
+					// If not looping, snap to start and stop playing.
+					NewPosition = 0.0f;
+					Stop();
+					bIsFinished = true;
+				}
+			}
+		}
+
+		SetPlaybackPosition(NewPosition, true);
 	}
-	else
+
+#if 0
+	// Notify user that timeline finished
+	if (bIsFinished)
 	{
-		AccumulatedTime = 0.0f;
+		TimelineFinishedFunc.ExecuteIfBound();
+		TimelineFinishFuncStatic.ExecuteIfBound();
 	}
+#endif
+}
+
+void UPaperFlipbookComponent::TickComponent(float DeltaTime, enum ELevelTick TickType, FActorComponentTickFunction* ThisTickFunction)
+{
+	// Advance time
+	TickFlipbook(DeltaTime);
 
+	// Update the frame and push it to the renderer if necessary
 	CalculateCurrentFrame();
 }
 
@@ -185,16 +270,164 @@ void UPaperFlipbookComponent::SetSpriteColor(FLinearColor NewColor)
 	}
 }
 
-void UPaperFlipbookComponent::SetCurrentTime(float NewTime)
+const UObject* UPaperFlipbookComponent::AdditionalStatObject() const
+{
+	return SourceFlipbook;
+}
+
+
+void UPaperFlipbookComponent::Play()
+{
+	Activate();
+	bReversePlayback = false;
+	bPlaying = true;
+}
+
+void UPaperFlipbookComponent::PlayFromStart()
+{
+	SetPlaybackPosition(0.0f, /*bFireEvents=*/ false);
+	Play();
+}
+
+void UPaperFlipbookComponent::Reverse()
+{
+	Activate();
+	bReversePlayback = true;
+	bPlaying = true;
+}
+
+void UPaperFlipbookComponent::ReverseFromEnd()
+{
+	SetPlaybackPosition(GetFlipbookLength(), /*bFireEvents=*/ false);
+	Reverse();
+}
+
+void UPaperFlipbookComponent::Stop()
+{
+	bPlaying = false;
+}
+
+bool UPaperFlipbookComponent::IsPlaying() const
+{
+	return bPlaying;
+}
+
+bool UPaperFlipbookComponent::IsReversing() const
+{
+	return bPlaying && bReversePlayback;
+}
+
+void UPaperFlipbookComponent::SetPlaybackPosition(float NewPosition, bool bFireEvents)
 {
-	if (NewTime != AccumulatedTime)
+	float OldPosition = AccumulatedTime;
+	AccumulatedTime = NewPosition;
+
+	// If we should be firing events for this track...
+	if (bFireEvents)
+	{
+		float MinTime;
+		float MaxTime;
+		if (!bReversePlayback)
+		{
+			// If playing sequence forwards.
+			MinTime = OldPosition;
+			MaxTime = AccumulatedTime;
+
+			// Slight hack here.. if playing forwards and reaching the end of the sequence, force it over a little to ensure we fire events actually on the end of the sequence.
+			if (MaxTime == GetFlipbookLength())
+			{
+				MaxTime += (float)KINDA_SMALL_NUMBER;
+			}
+		}
+		else
+		{
+			// If playing sequence backwards.
+			MinTime = AccumulatedTime;
+			MaxTime = OldPosition;
+
+			// Same small hack as above for backwards case.
+			if (MinTime == 0.0f)
+			{
+				MinTime -= (float)KINDA_SMALL_NUMBER;
+			}
+		}
+
+#if 0
+		// See which events fall into traversed region.
+		for (int32 i = 0; i < Events.Num(); i++)
+		{
+			float EventTime = Events[i].Time;
+
+			// Need to be slightly careful here and make behavior for firing events symmetric when playing forwards of backwards.
+			bool bFireThisEvent = false;
+			if (!bReversePlayback)
+			{
+				if (EventTime >= MinTime && EventTime < MaxTime)
+				{
+					bFireThisEvent = true;
+				}
+			}
+			else
+			{
+				if (EventTime > MinTime && EventTime <= MaxTime)
+				{
+					bFireThisEvent = true;
+				}
+			}
+
+			if (bFireThisEvent)
+			{
+				Events[i].EventFunc.ExecuteIfBound();
+			}
+		}
+#endif
+	}
+
+#if 0
+	// Execute the delegate to say that all properties are updated
+	TimelinePostUpdateFunc.ExecuteIfBound();
+#endif
+
+	if (OldPosition != AccumulatedTime)
 	{
-		AccumulatedTime = NewTime;
 		CalculateCurrentFrame();
 	}
 }
 
-const UObject* UPaperFlipbookComponent::AdditionalStatObject() const
+float UPaperFlipbookComponent::GetPlaybackPosition() const
 {
-	return SourceFlipbook;
+	return AccumulatedTime;
+}
+
+void UPaperFlipbookComponent::SetLooping(bool bNewLooping)
+{
+	bLooping = bNewLooping;
+}
+
+bool UPaperFlipbookComponent::IsLooping() const
+{
+	return bLooping;
+}
+
+void UPaperFlipbookComponent::SetPlayRate(float NewRate)
+{
+	PlayRate = NewRate;
+}
+
+float UPaperFlipbookComponent::GetPlayRate() const
+{
+	return PlayRate;
+}
+
+void UPaperFlipbookComponent::SetNewTime(float NewTime)
+{
+	// Ensure value is sensible
+	//@TODO: PAPER2D: NewTime = FMath::Clamp<float>(NewTime, 0.0f, Length);
+
+	SetPlaybackPosition(NewTime, /*bFireEvents=*/ false);
+}
+
+float UPaperFlipbookComponent::GetFlipbookLength() const
+{
+	return (SourceFlipbook != nullptr) ? SourceFlipbook->GetTotalDuration() : 0.0f;
 }
