Paper2D: Expose a few more accessors on UPaperFlipbook and UPaperFlipbookComponent

diff --git a/Engine/Plugins/Experimental/Paper2D/Source/Paper2D/Classes/PaperFlipbook.h b/Engine/Plugins/Experimental/Paper2D/Source/Paper2D/Classes/PaperFlipbook.h
--- a/Engine/Plugins/Experimental/Paper2D/Source/Paper2D/Classes/PaperFlipbook.h
+++ b/Engine/Plugins/Experimental/Paper2D/Source/Paper2D/Classes/PaperFlipbook.h
@@ -42,6 +42,9 @@ protected:
 	TArray<FPaperFlipbookKeyFrame> KeyFrames;
 
 public:
+	// Returns the nominal frame rate to play this flipbook animation back at
+	int32 GetFramesPerSecond() const;
+
 	// Returns the total number of frames
 	int32 GetNumFrames() const;
 
diff --git a/Engine/Plugins/Experimental/Paper2D/Source/Paper2D/Classes/PaperFlipbookComponent.h b/Engine/Plugins/Experimental/Paper2D/Source/Paper2D/Classes/PaperFlipbookComponent.h
--- a/Engine/Plugins/Experimental/Paper2D/Source/Paper2D/Classes/PaperFlipbookComponent.h
+++ b/Engine/Plugins/Experimental/Paper2D/Source/Paper2D/Classes/PaperFlipbookComponent.h
@@ -5,8 +5,8 @@
 #include "PaperFlipbook.h"
 #include "PaperFlipbookComponent.generated.h"
 
-UCLASS(MinimalAPI, ShowCategories=(Mobility), EarlyAccessPreview, meta=(BlueprintSpawnableComponent))
-class UPaperFlipbookComponent : public UPrimitiveComponent
+UCLASS(ShowCategories=(Mobility), EarlyAccessPreview, meta=(BlueprintSpawnableComponent))
+class PAPER2D_API UPaperFlipbookComponent : public UPrimitiveComponent
 {
 	GENERATED_UCLASS_BODY()
 
@@ -105,10 +105,18 @@ public:
 	UFUNCTION(BlueprintCallable, Category="Components|Flipbook")
 	void SetNewTime(float NewTime);
 
-	/** Get length of the flipbook */
+	/** Get length of the flipbook (in seconds) */
 	UFUNCTION(BlueprintCallable, Category="Components|Flipbook")
 	float GetFlipbookLength() const;
 
+	/** Get length of the flipbook (in frames) */
+	UFUNCTION(BlueprintCallable, Category = "Components|Flipbook")
+	int32 GetFlipbookLengthInFrames() const;
+
+	/** Get the nominal framerate that the flipbook will be played back at (ignoring PlayRate), in frames per second */
+	UFUNCTION(BlueprintCallable, Category = "Components|Flipbook")
+	float GetFlipbookFramerate() const;
+
 protected:
 	/** Current position in the timeline */
 	UPROPERTY()
diff --git a/Engine/Plugins/Experimental/Paper2D/Source/Paper2D/Private/PaperFlipbook.cpp b/Engine/Plugins/Experimental/Paper2D/Source/Paper2D/Private/PaperFlipbook.cpp
--- a/Engine/Plugins/Experimental/Paper2D/Source/Paper2D/Private/PaperFlipbook.cpp
+++ b/Engine/Plugins/Experimental/Paper2D/Source/Paper2D/Private/PaperFlipbook.cpp
@@ -17,6 +17,11 @@ UPaperFlipbook::UPaperFlipbook(const FPostConstructInitializeProperties& PCIP)
 	FramesPerSecond = 15;
 }
 
+int32 UPaperFlipbook::GetFramesPerSecond() const
+{
+	return FramesPerSecond;
+}
+
 int32 UPaperFlipbook::GetNumFrames() const
 {
 	int32 SumFrames = 0;
diff --git a/Engine/Plugins/Experimental/Paper2D/Source/Paper2D/Private/PaperFlipbookComponent.cpp b/Engine/Plugins/Experimental/Paper2D/Source/Paper2D/Private/PaperFlipbookComponent.cpp
--- a/Engine/Plugins/Experimental/Paper2D/Source/Paper2D/Private/PaperFlipbookComponent.cpp
+++ b/Engine/Plugins/Experimental/Paper2D/Source/Paper2D/Private/PaperFlipbookComponent.cpp
@@ -431,3 +431,13 @@ float UPaperFlipbookComponent::GetFlipbookLength() const
 {
 	return (SourceFlipbook != nullptr) ? SourceFlipbook->GetTotalDuration() : 0.0f;
 }
+
+int32 UPaperFlipbookComponent::GetFlipbookLengthInFrames() const
+{
+	return (SourceFlipbook != nullptr) ? SourceFlipbook->GetNumFrames() : 0;
+}
+
+float UPaperFlipbookComponent::GetFlipbookFramerate() const
+{
+	return (SourceFlipbook != nullptr) ? SourceFlipbook->GetFramesPerSecond() : 15.0f;
+}
