Paper2D: Sprite editor style tweak

diff --git a/Engine/Plugins/Experimental/Paper2D/Source/Paper2DEditor/Private/SpriteEditor/SpriteEditorViewportClient.cpp b/Engine/Plugins/Experimental/Paper2D/Source/Paper2DEditor/Private/SpriteEditor/SpriteEditorViewportClient.cpp
--- a/Engine/Plugins/Experimental/Paper2D/Source/Paper2DEditor/Private/SpriteEditor/SpriteEditorViewportClient.cpp
+++ b/Engine/Plugins/Experimental/Paper2D/Source/Paper2DEditor/Private/SpriteEditor/SpriteEditorViewportClient.cpp
@@ -593,7 +593,7 @@ void FSpriteEditorViewportClient::Tick(float DeltaSeconds)
 		// Zoom in on the sprite
 		//@TODO: Fix this properly so it doesn't need to be deferred, or wait for the viewport to initialize
 		FIntPoint Size = Viewport->GetSizeXY();
-		if (bDeferZoomToSprite && Size.X > 0 && Size.Y > 0)
+		if (bDeferZoomToSprite && (Size.X > 0) && (Size.Y > 0))
 		{
 			UPaperSpriteComponent* ComponentToFocusOn = SourceTextureViewComponent->IsVisible() ? SourceTextureViewComponent : RenderSpriteComponent;
 			FocusViewportOnBox(ComponentToFocusOn->Bounds.GetBox(), true);
