Paper2D: Add scrub bar and playback controls to flipbook editor

diff --git a/Engine/Plugins/Experimental/Paper2D/Source/Paper2DEditor/Private/FlipbookEditor/FlipbookEditor.cpp b/Engine/Plugins/Experimental/Paper2D/Source/Paper2DEditor/Private/FlipbookEditor/FlipbookEditor.cpp
--- a/Engine/Plugins/Experimental/Paper2D/Source/Paper2DEditor/Private/FlipbookEditor/FlipbookEditor.cpp
+++ b/Engine/Plugins/Experimental/Paper2D/Source/Paper2DEditor/Private/FlipbookEditor/FlipbookEditor.cpp
@@ -13,7 +13,7 @@
 #include "WorkspaceMenuStructureModule.h"
 #include "Paper2DEditorModule.h"
 #include "SFlipbookEditorViewportToolbar.h"
-
+#include "Editor/KismetWidgets/Public/SScrubControlPanel.h"
 #include "SFlipbookTimeline.h"
 
 #define LOCTEXT_NAMESPACE "FlipbookEditor"
@@ -44,11 +44,9 @@ class SFlipbookEditorViewport : public SEditorViewport, public ICommonEditorView
 public:
 	SLATE_BEGIN_ARGS(SFlipbookEditorViewport)
 		: _FlipbookBeingEdited((UPaperFlipbook*)NULL)
-		, _PlayTime(0)
 	{}
 
 		SLATE_ATTRIBUTE( UPaperFlipbook*, FlipbookBeingEdited )
-		SLATE_ATTRIBUTE( float, PlayTime )
 
 	SLATE_END_ARGS()
 
@@ -67,9 +65,13 @@ public:
 	virtual void OnFloatingButtonClicked() override;
 	// End of ICommonEditorViewportToolbarInfoProvider interface
 
+	UPaperFlipbookComponent* GetPreviewComponent() const
+	{
+		return EditorViewportClient->GetPreviewComponent();
+	}
+
 private:
 	TAttribute<UPaperFlipbook*> FlipbookBeingEdited;
-	TAttribute<float> PlayTime;
 
 	// Viewport client
 	TSharedPtr<FFlipbookEditorViewportClient> EditorViewportClient;
@@ -82,7 +84,6 @@ private:
 void SFlipbookEditorViewport::Construct(const FArguments& InArgs)
 {
 	FlipbookBeingEdited = InArgs._FlipbookBeingEdited;
-	PlayTime = InArgs._PlayTime;
 
 	SEditorViewport::Construct(SEditorViewport::FArguments());
 }
@@ -122,7 +123,7 @@ void SFlipbookEditorViewport::BindCommands()
 
 TSharedRef<FEditorViewportClient> SFlipbookEditorViewport::MakeEditorViewportClient()
 {
-	EditorViewportClient = MakeShareable(new FFlipbookEditorViewportClient(FlipbookBeingEdited, PlayTime));
+	EditorViewportClient = MakeShareable(new FFlipbookEditorViewportClient(FlipbookBeingEdited));
 
 	EditorViewportClient->VisibilityDelegate.BindSP(this, &SFlipbookEditorViewport::IsVisible);
 
@@ -201,8 +202,42 @@ public:
 //////////////////////////////////////////////////////////////////////////
 // FFlipbookEditor
 
+FFlipbookEditor::FFlipbookEditor()
+{
+}
+
 TSharedRef<SDockTab> FFlipbookEditor::SpawnTab_Viewport(const FSpawnTabArgs& Args)
 {
+	ViewInputMin = 0.0f;
+	ViewInputMax = GetTotalSequenceLength();
+	LastObservedSequenceLength = ViewInputMax;
+
+	TSharedRef<SWidget> ScrubControl = SNew(SScrubControlPanel)
+		.IsEnabled(true)
+		.Value(this, &FFlipbookEditor::GetPlaybackPosition)
+		.NumOfKeys(this, &FFlipbookEditor::GetTotalFrameCount)
+		.SequenceLength(this, &FFlipbookEditor::GetTotalSequenceLength)
+		.OnValueChanged(this, &FFlipbookEditor::SetPlaybackPosition)
+//		.OnBeginSliderMovement(this, &SAnimationScrubPanel::OnBeginSliderMovement)
+//		.OnEndSliderMovement(this, &SAnimationScrubPanel::OnEndSliderMovement)
+		.OnClickedForwardPlay(this, &FFlipbookEditor::OnClick_Forward)
+		.OnClickedForwardStep(this, &FFlipbookEditor::OnClick_Forward_Step)
+		.OnClickedForwardEnd(this, &FFlipbookEditor::OnClick_Forward_End)
+		.OnClickedBackwardPlay(this, &FFlipbookEditor::OnClick_Backward)
+		.OnClickedBackwardStep(this, &FFlipbookEditor::OnClick_Backward_Step)
+		.OnClickedBackwardEnd(this, &FFlipbookEditor::OnClick_Backward_End)
+		.OnClickedToggleLoop(this, &FFlipbookEditor::OnClick_ToggleLoop)
+		.OnGetLooping(this, &FFlipbookEditor::IsLooping)
+		.OnGetPlaybackMode(this, &FFlipbookEditor::GetPlaybackMode)
+		.ViewInputMin(this, &FFlipbookEditor::GetViewRangeMin)
+		.ViewInputMax(this, &FFlipbookEditor::GetViewRangeMax)
+		.OnSetInputViewRange(this, &FFlipbookEditor::SetViewRange)
+//		.OnCropAnimSequence(this, &SAnimationScrubPanel::OnCropAnimSequence)
+//		.OnReZeroAnimSequence(this, &SAnimationScrubPanel::OnReZeroAnimSequence)
+		.bAllowZoom(true)
+		.IsRealtimeStreamingMode(false)
+		.bLastFrameIsFirstFrame(false);
+
 	return SNew(SDockTab)
 		.Label(LOCTEXT("ViewportTab_Title", "Viewport"))
 		[
@@ -230,7 +265,7 @@ TSharedRef<SDockTab> FFlipbookEditor::SpawnTab_Viewport(const FSpawnTabArgs& Arg
 					.Text(LOCTEXT("FlipbookEditorViewportExperimentalWarning", "Early access preview"))
 				]
 			]
-			
+
 			+SVerticalBox::Slot()
 			.Padding(0, 8, 0, 0)
 			.AutoHeight()
@@ -239,6 +274,13 @@ TSharedRef<SDockTab> FFlipbookEditor::SpawnTab_Viewport(const FSpawnTabArgs& Arg
 				.FlipbookBeingEdited(this, &FFlipbookEditor::GetFlipbookBeingEdited)
 				.OnSelectionChanged(this, &FFlipbookEditor::SetSelection)
 			]
+
+			+SVerticalBox::Slot()
+			.Padding(0, 8, 0, 0)
+			.AutoHeight()
+			[
+				ScrubControl
+			]
 		];
 }
 
@@ -283,7 +325,6 @@ void FFlipbookEditor::InitFlipbookEditor(const EToolkitMode::Type Mode, const TS
 {
 	FAssetEditorManager::Get().CloseOtherEditors(InitFlipbook, this);
 	FlipbookBeingEdited = InitFlipbook;
-	PlayTime = 0;
 	CurrentSelectedKeyframe = INDEX_NONE;
 
 	FFlipbookEditorCommands::Register();
@@ -291,8 +332,7 @@ void FFlipbookEditor::InitFlipbookEditor(const EToolkitMode::Type Mode, const TS
 	BindCommands();
 
 	ViewportPtr = SNew(SFlipbookEditorViewport)
-		.FlipbookBeingEdited(this, &FFlipbookEditor::GetFlipbookBeingEdited)
-		.PlayTime(this, &FFlipbookEditor::GetPlayTime);
+		.FlipbookBeingEdited(this, &FFlipbookEditor::GetFlipbookBeingEdited);
 	
 	// Default layout
 	const TSharedRef<FTabManager::FLayout> StandaloneDefaultLayout = FTabManager::NewLayout("Standalone_FlipbookEditor_Layout_v1")
@@ -337,6 +377,13 @@ void FFlipbookEditor::InitFlipbookEditor(const EToolkitMode::Type Mode, const TS
 	RegenerateMenusAndToolbars();
 }
 
+UPaperFlipbookComponent* FFlipbookEditor::GetPreviewComponent() const
+{
+	UPaperFlipbookComponent* PreviewComponent = ViewportPtr->GetPreviewComponent();
+	check(PreviewComponent);
+	return PreviewComponent;
+}
+
 void FFlipbookEditor::BindCommands()
 {
 	const FFlipbookEditorCommands& Commands = FFlipbookEditorCommands::Get();
@@ -402,12 +449,6 @@ void FFlipbookEditor::ExtendToolbar()
 	{
 		static void FillToolbar(FToolBarBuilder& ToolbarBuilder)
 		{
-// 			ToolbarBuilder.BeginSection("Realtime");
-// 			{
-// 				ToolbarBuilder.AddToolBarButton(FEditorViewportCommands::Get().ToggleRealTime);
-// 			}
-// 			ToolbarBuilder.EndSection();
-
 			ToolbarBuilder.BeginSection("Command");
 			{
 				ToolbarBuilder.AddToolBarButton(FFlipbookEditorCommands::Get().SetShowGrid);
@@ -476,6 +517,161 @@ bool FFlipbookEditor::HasValidSelection() const
 	return FlipbookBeingEdited->IsValidKeyFrameIndex(CurrentSelectedKeyframe);
 }
 
+FReply FFlipbookEditor::OnClick_Forward()
+{
+	UPaperFlipbookComponent* PreviewComponent = GetPreviewComponent();
+
+	const bool bIsReverse = PreviewComponent->IsReversing();
+	const bool bIsPlaying = PreviewComponent->IsPlaying();
+		
+	if (bIsReverse && bIsPlaying)
+	{
+		// Play forwards instead of backwards
+		PreviewComponent->Play();
+	}
+	else if (bIsPlaying)
+	{
+		// Was already playing forwards, so pause
+		PreviewComponent->Stop();
+	}
+	else
+	{
+		// Was paused, start playing
+		PreviewComponent->Play();
+	}
+
+	return FReply::Handled();
+}
+
+FReply FFlipbookEditor::OnClick_Forward_Step()
+{
+	GetPreviewComponent()->Stop();
+	SetCurrentFrame(GetCurrentFrame() + 1);
+	return FReply::Handled();
+}
+
+FReply FFlipbookEditor::OnClick_Forward_End()
+{
+	UPaperFlipbookComponent* PreviewComponent = GetPreviewComponent();
+	PreviewComponent->Stop();
+	PreviewComponent->SetPlaybackPosition(PreviewComponent->GetFlipbookLength(), /*bFireEvents=*/ false);
+	return FReply::Handled();
+}
+
+FReply FFlipbookEditor::OnClick_Backward()
+{
+	UPaperFlipbookComponent* PreviewComponent = GetPreviewComponent();
+
+	const bool bIsReverse = PreviewComponent->IsReversing();
+	const bool bIsPlaying = PreviewComponent->IsPlaying();
+
+	if (bIsReverse && bIsPlaying)
+	{
+		// Was already playing backwards, so pause
+		PreviewComponent->Stop();
+	}
+	else if (bIsPlaying)
+	{
+		// Play backwards instead of forwards
+		PreviewComponent->Reverse();
+	}
+	else
+	{
+		// Was paused, start reversing
+		PreviewComponent->Reverse();
+	}
+
+	return FReply::Handled();
+}
+
+FReply FFlipbookEditor::OnClick_Backward_Step()
+{
+	GetPreviewComponent()->Stop();
+	SetCurrentFrame(GetCurrentFrame() - 1);
+	return FReply::Handled();
+}
+
+FReply FFlipbookEditor::OnClick_Backward_End()
+{
+	UPaperFlipbookComponent* PreviewComponent = GetPreviewComponent();
+	PreviewComponent->Stop();
+	PreviewComponent->SetPlaybackPosition(0.0f, /*bFireEvents=*/ false);
+	return FReply::Handled();
+}
+
+FReply FFlipbookEditor::OnClick_ToggleLoop()
+{
+	UPaperFlipbookComponent* PreviewComponent = GetPreviewComponent();
+	PreviewComponent->SetLooping(!PreviewComponent->IsLooping());
+	return FReply::Handled();
+}
+
+EPlaybackMode::Type FFlipbookEditor::GetPlaybackMode() const
+{
+	UPaperFlipbookComponent* PreviewComponent = GetPreviewComponent();
+	if (PreviewComponent->IsPlaying())
+	{
+		return PreviewComponent->IsReversing() ? EPlaybackMode::PlayingReverse : EPlaybackMode::PlayingForward;
+	}
+	else
+	{
+		return EPlaybackMode::Stopped;
+	}
+}
+
+uint32 FFlipbookEditor::GetTotalFrameCount() const
+{
+	return FlipbookBeingEdited->GetNumFrames();
+}
+
+float FFlipbookEditor::GetTotalSequenceLength() const
+{
+	return FlipbookBeingEdited->GetTotalDuration();
+}
+
+float FFlipbookEditor::GetPlaybackPosition() const
+{
+	return GetPreviewComponent()->GetPlaybackPosition();
+}
+
+void FFlipbookEditor::SetPlaybackPosition(float NewTime)
+{
+	NewTime = FMath::Clamp<float>(NewTime, 0.0f, GetTotalSequenceLength());
+
+	GetPreviewComponent()->SetPlaybackPosition(NewTime, /*bFireEvents=*/ false);
+}
+
+bool FFlipbookEditor::IsLooping() const
+{
+	return GetPreviewComponent()->IsLooping();
+}
+
+float FFlipbookEditor::GetViewRangeMin() const
+{
+	return ViewInputMin;
+}
+
+float FFlipbookEditor::GetViewRangeMax() const
+{
+	// See if the flipbook changed length, and if so reframe the scrub bar to include the full length
+	//@TODO: This is a pretty odd place to put it, but there's no callback for a modified timeline at the moment, so...
+	const float SequenceLength = GetTotalSequenceLength();
+	if (SequenceLength != LastObservedSequenceLength)
+	{
+		LastObservedSequenceLength = SequenceLength;
+		ViewInputMin = 0.0f;
+		ViewInputMax = SequenceLength;
+	}
+
+	return ViewInputMax;
+}
+
+void FFlipbookEditor::SetViewRange(float NewMin, float NewMax)
+{
+	ViewInputMin = FMath::Max<float>(NewMin, 0.0f);
+	ViewInputMax = FMath::Min<float>(NewMax, GetTotalSequenceLength());
+}
+
 //////////////////////////////////////////////////////////////////////////
 
 #undef LOCTEXT_NAMESPACE
\ No newline at end of file
diff --git a/Engine/Plugins/Experimental/Paper2D/Source/Paper2DEditor/Private/FlipbookEditor/FlipbookEditor.h b/Engine/Plugins/Experimental/Paper2D/Source/Paper2DEditor/Private/FlipbookEditor/FlipbookEditor.h
--- a/Engine/Plugins/Experimental/Paper2D/Source/Paper2DEditor/Private/FlipbookEditor/FlipbookEditor.h
+++ b/Engine/Plugins/Experimental/Paper2D/Source/Paper2DEditor/Private/FlipbookEditor/FlipbookEditor.h
@@ -4,6 +4,7 @@
 
 #include "Toolkits/AssetEditorToolkit.h"
 #include "Toolkits/AssetEditorManager.h"
+#include "Editor/EditorWidgets/Public/ITransportControl.h"
 
 //////////////////////////////////////////////////////////////////////////
 // FFlipbookEditor
@@ -11,6 +12,8 @@
 class FFlipbookEditor : public FAssetEditorToolkit, public FGCObject
 {
 public:
+	FFlipbookEditor();
+
 	// IToolkit interface
 	virtual void RegisterTabSpawners(const TSharedRef<class FTabManager>& TabManager) override;
 	virtual void UnregisterTabSpawners(const TSharedRef<class FTabManager>& TabManager) override;
@@ -33,19 +36,59 @@ public:
 	void InitFlipbookEditor(const EToolkitMode::Type Mode, const TSharedPtr< class IToolkitHost >& InitToolkitHost, class UPaperFlipbook* InitSprite);
 
 	UPaperFlipbook* GetFlipbookBeingEdited() const { return FlipbookBeingEdited; }
+
+	UPaperFlipbookComponent* GetPreviewComponent() const;
 protected:
 	UPaperFlipbook* FlipbookBeingEdited;
+
 	TSharedPtr<class SFlipbookEditorViewport> ViewportPtr;
-	float PlayTime;
+
+	// Selection set
 	int32 CurrentSelectedKeyframe;
 
+	// Range of times currently being viewed
+	mutable float ViewInputMin;
+	mutable float ViewInputMax;
+	mutable float LastObservedSequenceLength;
+
+protected:
+	float GetFramesPerSecond() const
+	{
+		return FlipbookBeingEdited->GetFramesPerSecond();
+	}
+
+	int32 GetCurrentFrame() const
+	{
+		const int32 TotalLengthInFrames = GetTotalFrameCount();
+
+		if (TotalLengthInFrames == 0)
+		{
+			return INDEX_NONE;
+		}
+		else
+		{
+			return FMath::Clamp<int32>((int32)(GetPlaybackPosition() * GetFramesPerSecond()), 0, TotalLengthInFrames - 1);
+		}
+	}
+
+	void SetCurrentFrame(int32 NewIndex)
+	{
+		const int32 TotalLengthInFrames = GetTotalFrameCount();
+		if (TotalLengthInFrames > 0)
+		{
+			int32 ClampedIndex = FMath::Clamp<int32>(NewIndex, 0, TotalLengthInFrames - 1);
+			SetPlaybackPosition(ClampedIndex / GetFramesPerSecond());
+		}
+		else
+		{
+			SetPlaybackPosition(0.0f);
+		}
+	}
 protected:
 	void BindCommands();
 	void ExtendMenu();
 	void ExtendToolbar();
 
-	float GetPlayTime() const { return PlayTime; }
-
 	TSharedRef<SDockTab> SpawnTab_Viewport(const FSpawnTabArgs& Args);
 	TSharedRef<SDockTab> SpawnTab_Details(const FSpawnTabArgs& Args);
 
@@ -53,4 +96,24 @@ protected:
 	void DuplicateSelection();
 	void SetSelection(int32 NewSelection);
 	bool HasValidSelection() const;
+
+
+	FReply OnClick_Forward();
+	FReply OnClick_Forward_Step();
+	FReply OnClick_Forward_End();
+	FReply OnClick_Backward();
+	FReply OnClick_Backward_Step();
+	FReply OnClick_Backward_End();
+	FReply OnClick_ToggleLoop();
+
+	uint32 GetTotalFrameCount() const;
+	float GetTotalSequenceLength() const;
+	float GetPlaybackPosition() const;
+	void SetPlaybackPosition(float NewTime);
+	bool IsLooping() const;
+	EPlaybackMode::Type GetPlaybackMode() const;
+
+	float GetViewRangeMin() const;
+	float GetViewRangeMax() const;
+	void SetViewRange(float NewMin, float NewMax);
 };
\ No newline at end of file
diff --git a/Engine/Plugins/Experimental/Paper2D/Source/Paper2DEditor/Private/FlipbookEditor/FlipbookEditorViewportClient.cpp b/Engine/Plugins/Experimental/Paper2D/Source/Paper2DEditor/Private/FlipbookEditor/FlipbookEditorViewportClient.cpp
--- a/Engine/Plugins/Experimental/Paper2D/Source/Paper2DEditor/Private/FlipbookEditor/FlipbookEditorViewportClient.cpp
+++ b/Engine/Plugins/Experimental/Paper2D/Source/Paper2DEditor/Private/FlipbookEditor/FlipbookEditorViewportClient.cpp
@@ -13,11 +13,10 @@
 //////////////////////////////////////////////////////////////////////////
 // FFlipbookEditorViewportClient
 
-FFlipbookEditorViewportClient::FFlipbookEditorViewportClient(const TAttribute<UPaperFlipbook*>& InFlipbookBeingEdited, const TAttribute<float>& InPlayTime)
+FFlipbookEditorViewportClient::FFlipbookEditorViewportClient(const TAttribute<UPaperFlipbook*>& InFlipbookBeingEdited)
 {
 	FlipbookBeingEdited = InFlipbookBeingEdited;
 	FlipbookBeingEditedLastFrame = FlipbookBeingEdited.Get();
-	PlayTime = InPlayTime;
 	PreviewScene = &OwnedPreviewScene;
 
 	SetRealtime(true);
@@ -75,10 +74,10 @@ void FFlipbookEditorViewportClient::Draw(const FSceneView* View, FPrimitiveDrawI
 
 void FFlipbookEditorViewportClient::Tick(float DeltaSeconds)
 {
-	if ( AnimatedRenderComponent.IsValid() )
+	if (AnimatedRenderComponent.IsValid())
 	{
 		UPaperFlipbook* Flipbook = FlipbookBeingEdited.Get();
-		if ( Flipbook != FlipbookBeingEditedLastFrame.Get() )
+		if (Flipbook != FlipbookBeingEditedLastFrame.Get())
 		{
 			AnimatedRenderComponent->SetFlipbook(Flipbook);
 			AnimatedRenderComponent->UpdateBounds();
@@ -87,7 +86,9 @@ void FFlipbookEditorViewportClient::Tick(float DeltaSeconds)
 
 		// Zoom in on the sprite
 		//@TODO: This doesn't work correctly, only partially zooming in or something
-		if (bDeferZoomToSprite)
+		//@TODO: Fix this properly so it doesn't need to be deferred, or wait for the viewport to initialize
+		FIntPoint Size = Viewport->GetSizeXY();
+		if (bDeferZoomToSprite && (Size.X > 0) && (Size.Y > 0))
 		{
 			FocusViewportOnBox(AnimatedRenderComponent->Bounds.GetBox(), true);
 			bDeferZoomToSprite = false;
diff --git a/Engine/Plugins/Experimental/Paper2D/Source/Paper2DEditor/Private/FlipbookEditor/FlipbookEditorViewportClient.h b/Engine/Plugins/Experimental/Paper2D/Source/Paper2DEditor/Private/FlipbookEditor/FlipbookEditorViewportClient.h
--- a/Engine/Plugins/Experimental/Paper2D/Source/Paper2DEditor/Private/FlipbookEditor/FlipbookEditorViewportClient.h
+++ b/Engine/Plugins/Experimental/Paper2D/Source/Paper2DEditor/Private/FlipbookEditor/FlipbookEditorViewportClient.h
@@ -17,7 +17,7 @@ class FFlipbookEditorViewportClient : public FPaperEditorViewportClient
 {
 public:
 	/** Constructor */
-	FFlipbookEditorViewportClient(const TAttribute<class UPaperFlipbook*>& InFlipbookBeingEdited, const TAttribute<float>& InPlayTime);
+	FFlipbookEditorViewportClient(const TAttribute<class UPaperFlipbook*>& InFlipbookBeingEdited);
 
 	// FViewportClient interface
 	virtual void Draw(FViewport* Viewport, FCanvas* Canvas) override;
@@ -33,6 +33,10 @@ public:
 	void ToggleShowPivot() { bShowPivot = !bShowPivot; Invalidate(); }
 	bool IsShowPivotChecked() const { return bShowPivot; }
 
+	UPaperFlipbookComponent* GetPreviewComponent() const
+	{
+		return AnimatedRenderComponent.Get();
+	}
 private:
 
 	// The preview scene
@@ -41,8 +45,6 @@ private:
 	// The flipbook being displayed in this client
 	TAttribute<class UPaperFlipbook*> FlipbookBeingEdited;
 
-	TAttribute<float> PlayTime;
-
 	// A cached pointer to the flipbook that was being edited last frame. Used for invalidation reasons.
 	TWeakObjectPtr<class UPaperFlipbook> FlipbookBeingEditedLastFrame;
 
