Paper2D: Switch to filled in brush for flipbook keyframes

diff --git a/Engine/Plugins/Experimental/Paper2D/Source/Paper2DEditor/Private/FlipbookEditor/STimelineTrack.h b/Engine/Plugins/Experimental/Paper2D/Source/Paper2DEditor/Private/FlipbookEditor/STimelineTrack.h
--- a/Engine/Plugins/Experimental/Paper2D/Source/Paper2DEditor/Private/FlipbookEditor/STimelineTrack.h
+++ b/Engine/Plugins/Experimental/Paper2D/Source/Paper2DEditor/Private/FlipbookEditor/STimelineTrack.h
@@ -87,7 +87,7 @@ private:
 					.WidthOverride( TAttribute<FOptionalSize>::Create( TAttribute<FOptionalSize>::FGetter::CreateSP(this, &STimelineTrack::GetFrameWidth, KeyFrameIdx) ) )
 					[
 						SNew(SBorder)
-						.BorderImage(FEditorStyle::GetBrush("FlipbookEditor.RegionBorder"))
+						.BorderImage(FEditorStyle::GetBrush("FlipbookEditor.RegionBody"))
 						.BorderBackgroundColor_Static(BorderColorDelegate, FlipbookBeingEdited, KeyFrameIdx)
 						.OnMouseButtonUp(this, &STimelineTrack::KeyframeOnMouseButtonUp, KeyFrameIdx)
 						[
diff --git a/Engine/Plugins/Experimental/Paper2D/Source/Paper2DEditor/Private/PaperStyle.cpp b/Engine/Plugins/Experimental/Paper2D/Source/Paper2DEditor/Private/PaperStyle.cpp
--- a/Engine/Plugins/Experimental/Paper2D/Source/Paper2DEditor/Private/PaperStyle.cpp
+++ b/Engine/Plugins/Experimental/Paper2D/Source/Paper2DEditor/Private/PaperStyle.cpp
@@ -100,6 +100,7 @@ TSharedRef<FSlateStyleSet> FPaperStyle::Create()
 		Style.Set("FlipbookEditor.SetShowPivot.Small", new IMAGE_BRUSH(TEXT("Icons/icon_StaticMeshEd_ShowPivot_40x"), Icon20x20));
 
 		Style.Set("FlipbookEditor.RegionGrabHandle", new BOX_BRUSH("Sequencer/ScrubHandleWhole", FMargin(6.f / 13.f, 10 / 24.f, 6 / 13.f, 10 / 24.f)));
+		Style.Set("FlipbookEditor.RegionBody", new BOX_BRUSH("Common/Scrollbar_Thumb", FMargin(4.f / 16.f)));
 		Style.Set("FlipbookEditor.RegionBorder", new BOX_BRUSH("Common/CurrentCellBorder", FMargin(4.f / 16.f), FLinearColor(1.0f, 1.0f, 1.0f, 1.0f)));
 	}
 
