Paper2D: Add scrub time visualization to flipbook track editor

diff --git a/Engine/Plugins/Experimental/Paper2D/Source/Paper2DEditor/Private/FlipbookEditor/FlipbookEditor.cpp b/Engine/Plugins/Experimental/Paper2D/Source/Paper2DEditor/Private/FlipbookEditor/FlipbookEditor.cpp
--- a/Engine/Plugins/Experimental/Paper2D/Source/Paper2DEditor/Private/FlipbookEditor/FlipbookEditor.cpp
+++ b/Engine/Plugins/Experimental/Paper2D/Source/Paper2DEditor/Private/FlipbookEditor/FlipbookEditor.cpp
@@ -212,6 +212,8 @@ TSharedRef<SDockTab> FFlipbookEditor::SpawnTab_Viewport(const FSpawnTabArgs& Arg
 	ViewInputMax = GetTotalSequenceLength();
 	LastObservedSequenceLength = ViewInputMax;
 
+	//@TODO: PAPER2D: Implement OnBeginSliderMovement/OnEndSliderMovement so that refreshing works, animation stops, etc...
+
 	TSharedRef<SWidget> ScrubControl = SNew(SScrubControlPanel)
 		.IsEnabled(true)
 		.Value(this, &FFlipbookEditor::GetPlaybackPosition)
@@ -232,8 +234,6 @@ TSharedRef<SDockTab> FFlipbookEditor::SpawnTab_Viewport(const FSpawnTabArgs& Arg
 		.ViewInputMin(this, &FFlipbookEditor::GetViewRangeMin)
 		.ViewInputMax(this, &FFlipbookEditor::GetViewRangeMax)
 		.OnSetInputViewRange(this, &FFlipbookEditor::SetViewRange)
-//		.OnCropAnimSequence(this, &SAnimationScrubPanel::OnCropAnimSequence)
-//		.OnReZeroAnimSequence(this, &SAnimationScrubPanel::OnReZeroAnimSequence)
 		.bAllowZoom(true)
 		.IsRealtimeStreamingMode(false)
 		.bLastFrameIsFirstFrame(false);
@@ -273,6 +273,7 @@ TSharedRef<SDockTab> FFlipbookEditor::SpawnTab_Viewport(const FSpawnTabArgs& Arg
 				SNew(SFlipbookTimeline, GetToolkitCommands())
 				.FlipbookBeingEdited(this, &FFlipbookEditor::GetFlipbookBeingEdited)
 				.OnSelectionChanged(this, &FFlipbookEditor::SetSelection)
+				.PlayTime(this, &FFlipbookEditor::GetPlaybackPosition)
 			]
 
 			+SVerticalBox::Slot()
diff --git a/Engine/Plugins/Experimental/Paper2D/Source/Paper2DEditor/Private/FlipbookEditor/FlipbookEditorViewportClient.cpp b/Engine/Plugins/Experimental/Paper2D/Source/Paper2DEditor/Private/FlipbookEditor/FlipbookEditorViewportClient.cpp
--- a/Engine/Plugins/Experimental/Paper2D/Source/Paper2DEditor/Private/FlipbookEditor/FlipbookEditorViewportClient.cpp
+++ b/Engine/Plugins/Experimental/Paper2D/Source/Paper2DEditor/Private/FlipbookEditor/FlipbookEditorViewportClient.cpp
@@ -46,7 +46,7 @@ void FFlipbookEditorViewportClient::DrawCanvas(FViewport& Viewport, FSceneView&
 
 	int32 YPos = 42;
 
-	static const FText FlipbookHelpStr = LOCTEXT("FlipbookEditHelp", "Flipbook editor (currently playing)");
+	static const FText FlipbookHelpStr = LOCTEXT("FlipbookEditHelp", "Flipbook editor");
 
 	// Display tool help
 	{
diff --git a/Engine/Plugins/Experimental/Paper2D/Source/Paper2DEditor/Private/FlipbookEditor/SFlipbookTimeline.cpp b/Engine/Plugins/Experimental/Paper2D/Source/Paper2DEditor/Private/FlipbookEditor/SFlipbookTimeline.cpp
--- a/Engine/Plugins/Experimental/Paper2D/Source/Paper2DEditor/Private/FlipbookEditor/SFlipbookTimeline.cpp
+++ b/Engine/Plugins/Experimental/Paper2D/Source/Paper2DEditor/Private/FlipbookEditor/SFlipbookTimeline.cpp
@@ -23,7 +23,7 @@ void SFlipbookTimeline::Construct(const FArguments& InArgs, TSharedPtr<const FUI
 	PlayTime = InArgs._PlayTime;
 	OnSelectionChanged = InArgs._OnSelectionChanged;
 
-	const int32 SlateUnitsPerFrame = 32;
+	SlateUnitsPerFrame = 32;
 	const int32 FrameHeight = 48;
 
 	ChildSlot
@@ -115,6 +115,36 @@ void SFlipbookTimeline::OnAssetsDropped(const class FAssetDragDropOp& DragDropOp
 	}
 }
 
+int32 SFlipbookTimeline::OnPaint(const FGeometry& AllottedGeometry, const FSlateRect& MyClippingRect, FSlateWindowElementList& OutDrawElements, int32 LayerId, const FWidgetStyle& InWidgetStyle, bool bParentEnabled) const
+{
+	LayerId = SCompoundWidget::OnPaint(AllottedGeometry, MyClippingRect, OutDrawElements, LayerId, InWidgetStyle, bParentEnabled);
+
+	const float CurrentTimeSecs = PlayTime.Get();
+	UPaperFlipbook* Flipbook = FlipbookBeingEdited.Get();
+	const float TotalTimeSecs = (Flipbook != nullptr) ? Flipbook->GetTotalDuration() : 0.0f;
+	const int32 TotalNumFrames = (Flipbook != nullptr) ? Flipbook->GetNumFrames() : 0;
+
+	const float SlateTotalDistance = SlateUnitsPerFrame * TotalNumFrames;
+	const float CurrentTimeXPos = (CurrentTimeSecs / TotalTimeSecs) * SlateTotalDistance;
+
+	++LayerId;
+	TArray<FVector2D> LinePoints;
+	LinePoints.Add(FVector2D(CurrentTimeXPos, 0.f));
+	LinePoints.Add(FVector2D(CurrentTimeXPos, AllottedGeometry.Size.Y));
+
+	FSlateDrawElement::MakeLines(
+		OutDrawElements,
+		LayerId,
+		AllottedGeometry.ToPaintGeometry(),
+		LinePoints,
+		MyClippingRect,
+		ESlateDrawEffect::None,
+		FLinearColor::Red
+		);
+
+	return LayerId;
+}
+
 //////////////////////////////////////////////////////////////////////////
 
 #undef LOCTEXT_NAMESPACE
\ No newline at end of file
diff --git a/Engine/Plugins/Experimental/Paper2D/Source/Paper2DEditor/Private/FlipbookEditor/SFlipbookTimeline.h b/Engine/Plugins/Experimental/Paper2D/Source/Paper2DEditor/Private/FlipbookEditor/SFlipbookTimeline.h
--- a/Engine/Plugins/Experimental/Paper2D/Source/Paper2DEditor/Private/FlipbookEditor/SFlipbookTimeline.h
+++ b/Engine/Plugins/Experimental/Paper2D/Source/Paper2DEditor/Private/FlipbookEditor/SFlipbookTimeline.h
@@ -24,6 +24,7 @@ public:
 
 	// SWidget interface
 	virtual FReply OnDrop(const FGeometry& MyGeometry, const FDragDropEvent& DragDropEvent) override;
+	virtual int32 OnPaint(const FGeometry& AllottedGeometry, const FSlateRect& MyClippingRect, FSlateWindowElementList& OutDrawElements, int32 LayerId, const FWidgetStyle& InWidgetStyle, bool bParentEnabled) const override;
 	// End of SWidget interface
 
 private:
@@ -33,4 +34,5 @@ private:
 	TAttribute<class UPaperFlipbook*> FlipbookBeingEdited;
 	TAttribute<float> PlayTime;
 	FOnFlipbookKeyframeSelectionChanged OnSelectionChanged;
+	int32 SlateUnitsPerFrame;
 };
diff --git a/Engine/Plugins/Experimental/Paper2D/Source/Paper2DEditor/Private/FlipbookEditor/STimelineHeader.h b/Engine/Plugins/Experimental/Paper2D/Source/Paper2DEditor/Private/FlipbookEditor/STimelineHeader.h
--- a/Engine/Plugins/Experimental/Paper2D/Source/Paper2DEditor/Private/FlipbookEditor/STimelineHeader.h
+++ b/Engine/Plugins/Experimental/Paper2D/Source/Paper2DEditor/Private/FlipbookEditor/STimelineHeader.h
@@ -28,33 +28,18 @@ public:
 
 		ChildSlot
 		[
-			SNew(SOverlay)
-
-			+ SOverlay::Slot()
-			[
-				SAssignNew(MainBoxPtr, SHorizontalBox)
-			]
-
-			+ SOverlay::Slot()
-			.Padding(TAttribute<FMargin>(this, &STimelineHeader::GetPlayTimePadding))
-			[
-				SNew(SBox)
-				.WidthOverride(SlateUnitsPerFrame.Get())
-				[
-					SNew(SSpacer)
-				]
-			]
+			SAssignNew(MainBoxPtr, SHorizontalBox)
 		];
 
 		Rebuild();
 	}
 
-	void Tick(const FGeometry& AllottedGeometry, const double InCurrentTime, const float InDeltaTime) override
+	virtual void Tick(const FGeometry& AllottedGeometry, const double InCurrentTime, const float InDeltaTime) override
 	{
 		SCompoundWidget::Tick(AllottedGeometry, InCurrentTime, InDeltaTime);
 
 		UPaperFlipbook* Flipbook = FlipbookBeingEdited.Get();
-		int32 NewNumFrames = (Flipbook != NULL ? Flipbook->GetNumFrames() : 0);
+		int32 NewNumFrames = (Flipbook != nullptr) ? Flipbook->GetNumFrames() : 0;
 		if (NewNumFrames != NumFramesFromLastRebuild)
 		{
 			Rebuild();
@@ -68,7 +53,7 @@ private:
 
 		UPaperFlipbook* Flipbook = FlipbookBeingEdited.Get();
 		float LocalSlateUnitsPerFrame = SlateUnitsPerFrame.Get();
-		if (Flipbook && LocalSlateUnitsPerFrame > 0)
+		if ((Flipbook != nullptr) && (LocalSlateUnitsPerFrame > 0))
 		{
 			const int32 NumFrames = Flipbook->GetNumFrames();
 			for (int32 FrameIdx = 0; FrameIdx < NumFrames; ++FrameIdx)
@@ -94,13 +79,6 @@ private:
 		}
 	}
 
-	FMargin GetPlayTimePadding() const
-	{
-		float PlayTimePadding = PlayTime.Get() - SlateUnitsPerFrame.Get() / 2;
-		return FMargin(PlayTimePadding, 0, 0, 0);
-	}
-
-
 private:
 	TAttribute<float> SlateUnitsPerFrame;
 	TAttribute<class UPaperFlipbook*> FlipbookBeingEdited;
