Paper2D: Add more ways to create key frames in flipbooks

- Add 'insert before/after' actions to context menu on existing frames
- Add 'add key frame' action to background context menu

diff --git a/Engine/Plugins/Experimental/Paper2D/Source/Paper2DEditor/Private/FlipbookEditor/FlipbookEditor.cpp b/Engine/Plugins/Experimental/Paper2D/Source/Paper2DEditor/Private/FlipbookEditor/FlipbookEditor.cpp
--- a/Engine/Plugins/Experimental/Paper2D/Source/Paper2DEditor/Private/FlipbookEditor/FlipbookEditor.cpp
+++ b/Engine/Plugins/Experimental/Paper2D/Source/Paper2DEditor/Private/FlipbookEditor/FlipbookEditor.cpp
@@ -398,6 +398,16 @@ void FFlipbookEditor::BindCommands()
 	UICommandList->MapAction(FGenericCommands::Get().Duplicate,
 		FExecuteAction::CreateSP(this, &FFlipbookEditor::DuplicateSelection),
 		FCanExecuteAction::CreateSP(this, &FFlipbookEditor::HasValidSelection));
+
+	UICommandList->MapAction(Commands.AddNewFrame,
+		FExecuteAction::CreateSP(this, &FFlipbookEditor::AddNewKeyFrameAtEnd),
+		FCanExecuteAction());
+	UICommandList->MapAction(Commands.AddNewFrameBefore,
+		FExecuteAction::CreateSP(this, &FFlipbookEditor::AddNewKeyFrameBefore),
+		FCanExecuteAction());
+	UICommandList->MapAction(Commands.AddNewFrameAfter,
+		FExecuteAction::CreateSP(this, &FFlipbookEditor::AddNewKeyFrameAfter),
+		FCanExecuteAction());
 }
 
 FName FFlipbookEditor::GetToolkitFName() const
@@ -479,7 +489,6 @@ void FFlipbookEditor::ExtendToolbar()
 
 void FFlipbookEditor::DeleteSelection()
 {
-
 	if (FlipbookBeingEdited->IsValidKeyFrameIndex(CurrentSelectedKeyframe))
 	{
 		const FScopedTransaction Transaction(LOCTEXT("DeleteKeyframe", "Delete Keyframe"));
@@ -508,6 +517,49 @@ void FFlipbookEditor::DuplicateSelection()
 	}
 }
 
+void FFlipbookEditor::AddNewKeyFrameAtEnd()
+{
+	const FScopedTransaction Transaction(LOCTEXT("AddKeyFrame", "Add Key Frame"));
+	FlipbookBeingEdited->Modify();
+
+	FScopedFlipbookMutator EditLock(FlipbookBeingEdited);
+
+	FPaperFlipbookKeyFrame& NewFrame = *new (EditLock.KeyFrames) FPaperFlipbookKeyFrame();
+}
+
+void FFlipbookEditor::AddNewKeyFrameBefore()
+{
+	if (FlipbookBeingEdited->IsValidKeyFrameIndex(CurrentSelectedKeyframe))
+	{
+		const FScopedTransaction Transaction(LOCTEXT("InsertKeyFrame", "Insert Key Frame"));
+		FlipbookBeingEdited->Modify();
+
+		FScopedFlipbookMutator EditLock(FlipbookBeingEdited);
+
+		FPaperFlipbookKeyFrame NewFrame;
+		EditLock.KeyFrames.Insert(NewFrame, CurrentSelectedKeyframe);
+
+		CurrentSelectedKeyframe = INDEX_NONE;
+	}
+}
+
+void FFlipbookEditor::AddNewKeyFrameAfter()
+{
+	if (FlipbookBeingEdited->IsValidKeyFrameIndex(CurrentSelectedKeyframe))
+	{
+		const FScopedTransaction Transaction(LOCTEXT("InsertKeyFrame", "Insert Key Frame"));
+		FlipbookBeingEdited->Modify();
+
+		FScopedFlipbookMutator EditLock(FlipbookBeingEdited);
+
+		FPaperFlipbookKeyFrame NewFrame;
+		EditLock.KeyFrames.Insert(NewFrame, CurrentSelectedKeyframe + 1);
+
+		CurrentSelectedKeyframe = INDEX_NONE;
+	}
+}
+
+
 void FFlipbookEditor::SetSelection(int32 NewSelection)
 {
 	CurrentSelectedKeyframe = NewSelection;
diff --git a/Engine/Plugins/Experimental/Paper2D/Source/Paper2DEditor/Private/FlipbookEditor/FlipbookEditor.h b/Engine/Plugins/Experimental/Paper2D/Source/Paper2DEditor/Private/FlipbookEditor/FlipbookEditor.h
--- a/Engine/Plugins/Experimental/Paper2D/Source/Paper2DEditor/Private/FlipbookEditor/FlipbookEditor.h
+++ b/Engine/Plugins/Experimental/Paper2D/Source/Paper2DEditor/Private/FlipbookEditor/FlipbookEditor.h
@@ -96,7 +96,9 @@ protected:
 	void DuplicateSelection();
 	void SetSelection(int32 NewSelection);
 	bool HasValidSelection() const;
-
+	void AddNewKeyFrameAtEnd();
+	void AddNewKeyFrameBefore();
+	void AddNewKeyFrameAfter();
 
 	FReply OnClick_Forward();
 	FReply OnClick_Forward_Step();
diff --git a/Engine/Plugins/Experimental/Paper2D/Source/Paper2DEditor/Private/FlipbookEditor/FlipbookEditorCommands.cpp b/Engine/Plugins/Experimental/Paper2D/Source/Paper2DEditor/Private/FlipbookEditor/FlipbookEditorCommands.cpp
--- a/Engine/Plugins/Experimental/Paper2D/Source/Paper2DEditor/Private/FlipbookEditor/FlipbookEditorCommands.cpp
+++ b/Engine/Plugins/Experimental/Paper2D/Source/Paper2DEditor/Private/FlipbookEditor/FlipbookEditorCommands.cpp
@@ -13,4 +13,8 @@ void FFlipbookEditorCommands::RegisterCommands()
 	UI_COMMAND(SetShowCollision, "Collision", "Toggles display of the simplified collision mesh of the static mesh, if one has been assigned.", EUserInterfaceActionType::ToggleButton, FInputGesture(EKeys::C, EModifierKey::Alt));
 
 	UI_COMMAND(SetShowPivot, "Show Pivot", "Display the pivot location of the static mesh.", EUserInterfaceActionType::ToggleButton, FInputGesture());
+
+	UI_COMMAND(AddNewFrame, "Add Key Frame", "Adds a new key frame to the flipbook.", EUserInterfaceActionType::Button, FInputGesture());
+	UI_COMMAND(AddNewFrameBefore, "Insert Key Frame Before", "Adds a new key frame to the flipbook before the selection.", EUserInterfaceActionType::Button, FInputGesture());
+	UI_COMMAND(AddNewFrameAfter, "Insert Key Frame After", "Adds a new key frame to the flipbook after the selection.", EUserInterfaceActionType::Button, FInputGesture());
 }
diff --git a/Engine/Plugins/Experimental/Paper2D/Source/Paper2DEditor/Private/FlipbookEditor/FlipbookEditorCommands.h b/Engine/Plugins/Experimental/Paper2D/Source/Paper2DEditor/Private/FlipbookEditor/FlipbookEditorCommands.h
--- a/Engine/Plugins/Experimental/Paper2D/Source/Paper2DEditor/Private/FlipbookEditor/FlipbookEditorCommands.h
+++ b/Engine/Plugins/Experimental/Paper2D/Source/Paper2DEditor/Private/FlipbookEditor/FlipbookEditorCommands.h
@@ -8,7 +8,7 @@ public:
 	FFlipbookEditorCommands()
 		: TCommands<FFlipbookEditorCommands>(
 			TEXT("FlipbookEditor"), // Context name for fast lookup
-			NSLOCTEXT("Contexts", "PaperEditor", "Sprite Editor"), // Localized context name for displaying
+			NSLOCTEXT("Contexts", "FlipbookEditor", "Flipbook Editor"), // Localized context name for displaying
 			NAME_None, // Parent
 			FEditorStyle::GetStyleSetName() // Icon Style Set
 			)
@@ -27,5 +27,8 @@ public:
 	// View Menu Commands
 	TSharedPtr<FUICommandInfo> SetShowPivot;
 
-	// 
+	// Timeline commands
+	TSharedPtr<FUICommandInfo> AddNewFrame;
+	TSharedPtr<FUICommandInfo> AddNewFrameBefore;
+	TSharedPtr<FUICommandInfo> AddNewFrameAfter;
 };
\ No newline at end of file
diff --git a/Engine/Plugins/Experimental/Paper2D/Source/Paper2DEditor/Private/FlipbookEditor/SFlipbookTimeline.cpp b/Engine/Plugins/Experimental/Paper2D/Source/Paper2DEditor/Private/FlipbookEditor/SFlipbookTimeline.cpp
--- a/Engine/Plugins/Experimental/Paper2D/Source/Paper2DEditor/Private/FlipbookEditor/SFlipbookTimeline.cpp
+++ b/Engine/Plugins/Experimental/Paper2D/Source/Paper2DEditor/Private/FlipbookEditor/SFlipbookTimeline.cpp
@@ -4,6 +4,7 @@
 #include "SFlipbookTimeline.h"
 #include "Editor/UnrealEd/Public/DragAndDrop/AssetDragDropOp.h"
 #include "Editor/UnrealEd/Public/ScopedTransaction.h"
+#include "FlipbookEditorCommands.h"
 
 #define LOCTEXT_NAMESPACE "FlipbookEditor"
 
@@ -22,6 +23,7 @@ void SFlipbookTimeline::Construct(const FArguments& InArgs, TSharedPtr<const FUI
 	FlipbookBeingEdited = InArgs._FlipbookBeingEdited;
 	PlayTime = InArgs._PlayTime;
 	OnSelectionChanged = InArgs._OnSelectionChanged;
+	CommandList = InCommandList;
 
 	SlateUnitsPerFrame = 32;
 	const int32 FrameHeight = 48;
@@ -80,6 +82,7 @@ FReply SFlipbookTimeline::OnDrop(const FGeometry& MyGeometry, const FDragDropEve
 
 void SFlipbookTimeline::OnAssetsDropped(const class FAssetDragDropOp& DragDropOp)
 {
+	//@TODO: Support inserting in addition to dropping at the end
 	TArray<FPaperFlipbookKeyFrame> NewFrames;
 	for (const FAssetData& AssetData : DragDropOp.AssetData)
 	{
@@ -127,6 +130,7 @@ int32 SFlipbookTimeline::OnPaint(const FGeometry& AllottedGeometry, const FSlate
 	const float SlateTotalDistance = SlateUnitsPerFrame * TotalNumFrames;
 	const float CurrentTimeXPos = (CurrentTimeSecs / TotalTimeSecs) * SlateTotalDistance;
 
+	// Draw a line for the current scrub cursor
 	++LayerId;
 	TArray<FVector2D> LinePoints;
 	LinePoints.Add(FVector2D(CurrentTimeXPos, 0.f));
@@ -145,6 +149,36 @@ int32 SFlipbookTimeline::OnPaint(const FGeometry& AllottedGeometry, const FSlate
 	return LayerId;
 }
 
+TSharedRef<SWidget> SFlipbookTimeline::GenerateContextMenu()
+{
+	FMenuBuilder MenuBuilder(true, CommandList);
+	MenuBuilder.BeginSection("KeyframeActions", LOCTEXT("KeyframeActionsSectionHeader", "Keyframe Actions"));
+
+	// 		MenuBuilder.AddMenuEntry(FGenericCommands::Get().Cut);
+	// 		MenuBuilder.AddMenuEntry(FGenericCommands::Get().Copy);
+	// 		MenuBuilder.AddMenuEntry(FGenericCommands::Get().Paste);
+	MenuBuilder.AddMenuEntry(FFlipbookEditorCommands::Get().AddNewFrame);
+
+	MenuBuilder.EndSection();
+
+	return MenuBuilder.MakeWidget();
+}
+
+FReply SFlipbookTimeline::OnMouseButtonUp(const FGeometry& MyGeometry, const FPointerEvent& MouseEvent)
+{
+	if (MouseEvent.GetEffectingButton() == EKeys::RightMouseButton)
+	{
+		TSharedRef<SWidget> MenuContents = GenerateContextMenu();
+		FSlateApplication::Get().PushMenu(AsShared(), MenuContents, MouseEvent.GetScreenSpacePosition(), FPopupTransitionEffect(FPopupTransitionEffect::ContextMenu));
+
+		return FReply::Handled();
+	}
+	else
+	{
+		return FReply::Unhandled();
+	}
+}
+
 //////////////////////////////////////////////////////////////////////////
 
 #undef LOCTEXT_NAMESPACE
\ No newline at end of file
diff --git a/Engine/Plugins/Experimental/Paper2D/Source/Paper2DEditor/Private/FlipbookEditor/SFlipbookTimeline.h b/Engine/Plugins/Experimental/Paper2D/Source/Paper2DEditor/Private/FlipbookEditor/SFlipbookTimeline.h
--- a/Engine/Plugins/Experimental/Paper2D/Source/Paper2DEditor/Private/FlipbookEditor/SFlipbookTimeline.h
+++ b/Engine/Plugins/Experimental/Paper2D/Source/Paper2DEditor/Private/FlipbookEditor/SFlipbookTimeline.h
@@ -25,14 +25,18 @@ public:
 	// SWidget interface
 	virtual FReply OnDrop(const FGeometry& MyGeometry, const FDragDropEvent& DragDropEvent) override;
 	virtual int32 OnPaint(const FGeometry& AllottedGeometry, const FSlateRect& MyClippingRect, FSlateWindowElementList& OutDrawElements, int32 LayerId, const FWidgetStyle& InWidgetStyle, bool bParentEnabled) const override;
+	virtual FReply OnMouseButtonUp(const FGeometry& MyGeometry, const FPointerEvent& MouseEvent) override;
 	// End of SWidget interface
 
 private:
 	void OnAssetsDropped(const class FAssetDragDropOp& DragDropOp);
 
+	TSharedRef<SWidget> GenerateContextMenu();
+
 private:
 	TAttribute<class UPaperFlipbook*> FlipbookBeingEdited;
 	TAttribute<float> PlayTime;
+	TSharedPtr<const FUICommandList> CommandList;
 	FOnFlipbookKeyframeSelectionChanged OnSelectionChanged;
 	int32 SlateUnitsPerFrame;
 };
diff --git a/Engine/Plugins/Experimental/Paper2D/Source/Paper2DEditor/Private/FlipbookEditor/STimelineTrack.h b/Engine/Plugins/Experimental/Paper2D/Source/Paper2DEditor/Private/FlipbookEditor/STimelineTrack.h
--- a/Engine/Plugins/Experimental/Paper2D/Source/Paper2DEditor/Private/FlipbookEditor/STimelineTrack.h
+++ b/Engine/Plugins/Experimental/Paper2D/Source/Paper2DEditor/Private/FlipbookEditor/STimelineTrack.h
@@ -145,6 +145,11 @@ private:
 		MenuBuilder.AddMenuEntry(FGenericCommands::Get().Duplicate);
 		MenuBuilder.AddMenuEntry(FGenericCommands::Get().Delete);
 
+		MenuBuilder.AddMenuSeparator();
+
+		MenuBuilder.AddMenuEntry(FFlipbookEditorCommands::Get().AddNewFrameBefore);
+		MenuBuilder.AddMenuEntry(FFlipbookEditorCommands::Get().AddNewFrameAfter);
+
 		MenuBuilder.EndSection();
 
 		return MenuBuilder.MakeWidget();
