Paper2D: Add instructions on how to create key frames when there are no frames in the timeline

diff --git a/Engine/Plugins/Experimental/Paper2D/Source/Paper2DEditor/Private/FlipbookEditor/SFlipbookTimeline.cpp b/Engine/Plugins/Experimental/Paper2D/Source/Paper2DEditor/Private/FlipbookEditor/SFlipbookTimeline.cpp
--- a/Engine/Plugins/Experimental/Paper2D/Source/Paper2DEditor/Private/FlipbookEditor/SFlipbookTimeline.cpp
+++ b/Engine/Plugins/Experimental/Paper2D/Source/Paper2DEditor/Private/FlipbookEditor/SFlipbookTimeline.cpp
@@ -33,27 +33,43 @@ void SFlipbookTimeline::Construct(const FArguments& InArgs, TSharedPtr<const FUI
 		SNew(SBorder)
 		.BorderImage( FEditorStyle::GetBrush("ToolPanel.GroupBorder") )
 		[
-			SNew(SVerticalBox)
-		
-			+SVerticalBox::Slot()
-			.AutoHeight()
-			.Padding(0,0,0,2)
+			SNew(SHorizontalBox)
+
+			// Empty flipbook instructions
+			+SHorizontalBox::Slot()
+			.VAlign(VAlign_Center)
+			.HAlign(HAlign_Center)
 			[
-				SNew(STimelineHeader)
-				.SlateUnitsPerFrame(SlateUnitsPerFrame)
-				.FlipbookBeingEdited(FlipbookBeingEdited)
-				.PlayTime(PlayTime)
+				SNew(STextBlock)
+				.Visibility(this, &SFlipbookTimeline::NoFramesWarningVisibility)
+				.Text(LOCTEXT("EmptyTimelineInstruction", "Right-click here or drop in sprites to add key frames"))
 			]
 
-			+SVerticalBox::Slot()
-			.AutoHeight()
+			// Flipbook header and track
+			+SHorizontalBox::Slot()
 			[
-				SNew(SBox).HeightOverride(FrameHeight)
+				SNew(SVerticalBox)
+		
+				+SVerticalBox::Slot()
+				.AutoHeight()
+				.Padding(0,0,0,2)
 				[
-					SNew(STimelineTrack, InCommandList)
+					SNew(STimelineHeader)
 					.SlateUnitsPerFrame(SlateUnitsPerFrame)
 					.FlipbookBeingEdited(FlipbookBeingEdited)
-					.OnSelectionChanged(OnSelectionChanged)
+					.PlayTime(PlayTime)
+				]
+
+				+SVerticalBox::Slot()
+				.AutoHeight()
+				[
+					SNew(SBox).HeightOverride(FrameHeight)
+					[
+						SNew(STimelineTrack, InCommandList)
+						.SlateUnitsPerFrame(SlateUnitsPerFrame)
+						.FlipbookBeingEdited(FlipbookBeingEdited)
+						.OnSelectionChanged(OnSelectionChanged)
+					]
 				]
 			]
 		]
@@ -179,6 +195,13 @@ FReply SFlipbookTimeline::OnMouseButtonUp(const FGeometry& MyGeometry, const FPo
 	}
 }
 
+EVisibility SFlipbookTimeline::NoFramesWarningVisibility() const
+{
+	UPaperFlipbook* Flipbook = FlipbookBeingEdited.Get();
+	const int32 TotalNumFrames = (Flipbook != nullptr) ? Flipbook->GetNumFrames() : 0;
+	return (TotalNumFrames == 0) ? EVisibility::Visible : EVisibility::Collapsed;
+}
+
 //////////////////////////////////////////////////////////////////////////
 
 #undef LOCTEXT_NAMESPACE
\ No newline at end of file
diff --git a/Engine/Plugins/Experimental/Paper2D/Source/Paper2DEditor/Private/FlipbookEditor/SFlipbookTimeline.h b/Engine/Plugins/Experimental/Paper2D/Source/Paper2DEditor/Private/FlipbookEditor/SFlipbookTimeline.h
--- a/Engine/Plugins/Experimental/Paper2D/Source/Paper2DEditor/Private/FlipbookEditor/SFlipbookTimeline.h
+++ b/Engine/Plugins/Experimental/Paper2D/Source/Paper2DEditor/Private/FlipbookEditor/SFlipbookTimeline.h
@@ -32,7 +32,7 @@ private:
 	void OnAssetsDropped(const class FAssetDragDropOp& DragDropOp);
 
 	TSharedRef<SWidget> GenerateContextMenu();
-
+	EVisibility NoFramesWarningVisibility() const;
 private:
 	TAttribute<class UPaperFlipbook*> FlipbookBeingEdited;
 	TAttribute<float> PlayTime;
