Paper2D: Prevent flipbook keyframe FrameRun from being edited below 1

diff --git a/Engine/Plugins/Experimental/Paper2D/Source/Paper2D/Classes/PaperFlipbook.h b/Engine/Plugins/Experimental/Paper2D/Source/Paper2D/Classes/PaperFlipbook.h
--- a/Engine/Plugins/Experimental/Paper2D/Source/Paper2D/Classes/PaperFlipbook.h
+++ b/Engine/Plugins/Experimental/Paper2D/Source/Paper2D/Classes/PaperFlipbook.h
@@ -14,7 +14,7 @@ public:
 	UPROPERTY(Category=Sprite, EditAnywhere)
 	UPaperSprite* Sprite;
 
-	UPROPERTY(Category=Sprite, EditAnywhere)
+	UPROPERTY(Category=Sprite, EditAnywhere, meta=(ClampMin=1))
 	int32 FrameRun;
 
 	FPaperFlipbookKeyFrame()
