Paper2D: Add tooltip for flipbook keyframes

diff --git a/Engine/Plugins/Experimental/Paper2D/Source/Paper2DEditor/Private/FlipbookEditor/STimelineTrack.h b/Engine/Plugins/Experimental/Paper2D/Source/Paper2DEditor/Private/FlipbookEditor/STimelineTrack.h
--- a/Engine/Plugins/Experimental/Paper2D/Source/Paper2DEditor/Private/FlipbookEditor/STimelineTrack.h
+++ b/Engine/Plugins/Experimental/Paper2D/Source/Paper2DEditor/Private/FlipbookEditor/STimelineTrack.h
@@ -61,6 +61,26 @@ public:
 		return FReply::Unhandled();
 	}
 
+	FText GetKeyframeTooltip(int32 KeyFrameIdx) const
+	{
+		UPaperFlipbook* Flipbook = FlipbookBeingEdited.Get();
+		if ((Flipbook != nullptr) && Flipbook->IsValidKeyFrameIndex(KeyFrameIdx))
+		{
+			const FPaperFlipbookKeyFrame& KeyFrame = Flipbook->GetKeyFrameChecked(KeyFrameIdx);
+
+			FText SpriteLine = (KeyFrame.Sprite != nullptr) ? FText::FromString(KeyFrame.Sprite->GetName()) : LOCTEXT("NoSprite", "(none)");
+
+			return FText::Format(LOCTEXT("KeyFrameTooltip", "Sprite: {0}\nIndex: {1}\nDuration: {2} frame(s)"),
+				SpriteLine,
+				FText::AsNumber(KeyFrameIdx),
+				FText::AsNumber(KeyFrame.FrameRun));
+		}
+		else
+		{
+			return LOCTEXT("KeyFrameTooltip_Invalid", "Invalid key frame index");
+		}
+	}
+
 private:
 	void Rebuild()
 	{
@@ -90,6 +110,7 @@ private:
 						.BorderImage(FEditorStyle::GetBrush("FlipbookEditor.RegionBody"))
 						.BorderBackgroundColor_Static(BorderColorDelegate, FlipbookBeingEdited, KeyFrameIdx)
 						.OnMouseButtonUp(this, &STimelineTrack::KeyframeOnMouseButtonUp, KeyFrameIdx)
+						.ToolTipText(this, &STimelineTrack::GetKeyframeTooltip, KeyFrameIdx)
 						[
 							SNullWidget::NullWidget
 						]
