Paper2D: Add 'show mesh edges' option to sprite editor and move most show commands to in-viewport toolbar Show menu

diff --git a/Engine/Plugins/Experimental/Paper2D/Source/Paper2DEditor/Private/SpriteEditor/SSpriteEditorViewportToolbar.cpp b/Engine/Plugins/Experimental/Paper2D/Source/Paper2DEditor/Private/SpriteEditor/SSpriteEditorViewportToolbar.cpp
--- a/Engine/Plugins/Experimental/Paper2D/Source/Paper2DEditor/Private/SpriteEditor/SSpriteEditorViewportToolbar.cpp
+++ b/Engine/Plugins/Experimental/Paper2D/Source/Paper2DEditor/Private/SpriteEditor/SSpriteEditorViewportToolbar.cpp
@@ -2,6 +2,8 @@
 
 #include "Paper2DEditorPrivatePCH.h"
 #include "SSpriteEditorViewportToolbar.h"
+#include "SpriteEditorCommands.h"
+#include "SEditorViewport.h"
 
 #define LOCTEXT_NAMESPACE "SSpriteEditorViewportToolbar"
 
@@ -13,4 +15,31 @@ void SSpriteEditorViewportToolbar::Construct(const FArguments& InArgs, TSharedPt
 	SCommonEditorViewportToolbarBase::Construct(SCommonEditorViewportToolbarBase::FArguments(), InInfoProvider);
 }
 
+TSharedRef<SWidget> SSpriteEditorViewportToolbar::GenerateShowMenu() const
+{
+	GetInfoProvider().OnFloatingButtonClicked();
+	
+	TSharedRef<SEditorViewport> ViewportRef = GetInfoProvider().GetViewportWidget();
+
+	const bool bInShouldCloseWindowAfterMenuSelection = true;
+	FMenuBuilder ShowMenuBuilder(bInShouldCloseWindowAfterMenuSelection, ViewportRef->GetCommandList());
+	{
+		ShowMenuBuilder.AddMenuEntry(FSpriteEditorCommands::Get().SetShowSockets);
+		ShowMenuBuilder.AddMenuEntry(FSpriteEditorCommands::Get().SetShowPivot);
+
+		ShowMenuBuilder.AddMenuSeparator();
+
+		ShowMenuBuilder.AddMenuEntry(FSpriteEditorCommands::Get().SetShowGrid);
+		ShowMenuBuilder.AddMenuEntry(FSpriteEditorCommands::Get().SetShowBounds);
+		ShowMenuBuilder.AddMenuEntry(FSpriteEditorCommands::Get().SetShowNormals);
+
+		ShowMenuBuilder.AddMenuSeparator();
+
+		ShowMenuBuilder.AddMenuEntry(FSpriteEditorCommands::Get().SetShowCollision);
+		ShowMenuBuilder.AddMenuEntry(FSpriteEditorCommands::Get().SetShowMeshEdges);
+	}
+
+	return ShowMenuBuilder.MakeWidget();
+}
+
 #undef LOCTEXT_NAMESPACE
diff --git a/Engine/Plugins/Experimental/Paper2D/Source/Paper2DEditor/Private/SpriteEditor/SSpriteEditorViewportToolbar.h b/Engine/Plugins/Experimental/Paper2D/Source/Paper2DEditor/Private/SpriteEditor/SSpriteEditorViewportToolbar.h
--- a/Engine/Plugins/Experimental/Paper2D/Source/Paper2DEditor/Private/SpriteEditor/SSpriteEditorViewportToolbar.h
+++ b/Engine/Plugins/Experimental/Paper2D/Source/Paper2DEditor/Private/SpriteEditor/SSpriteEditorViewportToolbar.h
@@ -12,4 +12,9 @@ public:
 	SLATE_END_ARGS()
 
 	void Construct(const FArguments& InArgs, TSharedPtr<class ICommonEditorViewportToolbarInfoProvider> InInfoProvider);
+
+	// SCommonEditorViewportToolbarBase interface
+	virtual TSharedRef<SWidget> GenerateShowMenu() const override;
+	// End of SCommonEditorViewportToolbarBase
+
 };
diff --git a/Engine/Plugins/Experimental/Paper2D/Source/Paper2DEditor/Private/SpriteEditor/SpriteEditor.cpp b/Engine/Plugins/Experimental/Paper2D/Source/Paper2DEditor/Private/SpriteEditor/SpriteEditor.cpp
--- a/Engine/Plugins/Experimental/Paper2D/Source/Paper2DEditor/Private/SpriteEditor/SpriteEditor.cpp
+++ b/Engine/Plugins/Experimental/Paper2D/Source/Paper2DEditor/Private/SpriteEditor/SpriteEditor.cpp
@@ -119,6 +119,12 @@ void SSpriteEditorViewport::BindCommands()
 		FCanExecuteAction(),
 		FIsActionChecked::CreateSP( EditorViewportClientRef, &FEditorViewportClient::IsSetShowCollisionChecked ) );
 
+	CommandList->MapAction(
+		Commands.SetShowMeshEdges,
+		FExecuteAction::CreateSP(EditorViewportClientRef, &FSpriteEditorViewportClient::ToggleShowMeshEdges),
+		FCanExecuteAction(),
+		FIsActionChecked::CreateSP(EditorViewportClientRef, &FSpriteEditorViewportClient::IsShowMeshEdgesChecked));
+
  	CommandList->MapAction(
  		Commands.SetShowSockets,
  		FExecuteAction::CreateSP( EditorViewportClientRef, &FSpriteEditorViewportClient::ToggleShowSockets ),
@@ -500,15 +506,7 @@ void FSpriteEditor::ExtendToolbar()
 
 			ToolbarBuilder.BeginSection("Command");
 			{
-				ToolbarBuilder.AddToolBarButton(FSpriteEditorCommands::Get().SetShowGrid);
 				ToolbarBuilder.AddToolBarButton(FSpriteEditorCommands::Get().SetShowSourceTexture);
-				ToolbarBuilder.AddToolBarButton(FSpriteEditorCommands::Get().SetShowBounds);
-				ToolbarBuilder.AddToolBarButton(FSpriteEditorCommands::Get().SetShowCollision);
-
-				ToolbarBuilder.AddToolBarButton(FSpriteEditorCommands::Get().SetShowSockets);
-
-				ToolbarBuilder.AddToolBarButton(FSpriteEditorCommands::Get().SetShowPivot);
-				ToolbarBuilder.AddToolBarButton(FSpriteEditorCommands::Get().SetShowNormals);
 			}
 			ToolbarBuilder.EndSection();
 
diff --git a/Engine/Plugins/Experimental/Paper2D/Source/Paper2DEditor/Private/SpriteEditor/SpriteEditorCommands.cpp b/Engine/Plugins/Experimental/Paper2D/Source/Paper2DEditor/Private/SpriteEditor/SpriteEditorCommands.cpp
--- a/Engine/Plugins/Experimental/Paper2D/Source/Paper2DEditor/Private/SpriteEditor/SpriteEditorCommands.cpp
+++ b/Engine/Plugins/Experimental/Paper2D/Source/Paper2DEditor/Private/SpriteEditor/SpriteEditorCommands.cpp
@@ -17,7 +17,8 @@ void FSpriteEditorCommands::RegisterCommands()
 	UI_COMMAND(SetShowSockets, "Sockets", "Displays the static mesh sockets.", EUserInterfaceActionType::ToggleButton, FInputGesture());
 
 	UI_COMMAND(SetShowNormals, "Normals", "Toggles display of vertex normals in the Preview Pane.", EUserInterfaceActionType::ToggleButton, FInputGesture());
-	UI_COMMAND(SetShowPivot, "Show Pivot", "Display the pivot location of the static mesh.", EUserInterfaceActionType::ToggleButton, FInputGesture());
+	UI_COMMAND(SetShowPivot, "Pivot", "Display the pivot location of the sprite.", EUserInterfaceActionType::ToggleButton, FInputGesture());
+	UI_COMMAND(SetShowMeshEdges, "Mesh Edges", "Overlays the mesh edges on top of the view.", EUserInterfaceActionType::ToggleButton, FInputGesture());
 
 	// Editing modes
 	UI_COMMAND(EnterViewMode, "View", "View the sprite.", EUserInterfaceActionType::ToggleButton, FInputGesture());
diff --git a/Engine/Plugins/Experimental/Paper2D/Source/Paper2DEditor/Private/SpriteEditor/SpriteEditorCommands.h b/Engine/Plugins/Experimental/Paper2D/Source/Paper2DEditor/Private/SpriteEditor/SpriteEditorCommands.h
--- a/Engine/Plugins/Experimental/Paper2D/Source/Paper2DEditor/Private/SpriteEditor/SpriteEditorCommands.h
+++ b/Engine/Plugins/Experimental/Paper2D/Source/Paper2DEditor/Private/SpriteEditor/SpriteEditorCommands.h
@@ -30,6 +30,7 @@ public:
 
 	TSharedPtr<FUICommandInfo> SetShowNormals;
 	TSharedPtr<FUICommandInfo> SetShowPivot;
+	TSharedPtr<FUICommandInfo> SetShowMeshEdges;
 
 	// Editing modes
 	TSharedPtr<FUICommandInfo> EnterViewMode;
diff --git a/Engine/Plugins/Experimental/Paper2D/Source/Paper2DEditor/Private/SpriteEditor/SpriteEditorViewportClient.cpp b/Engine/Plugins/Experimental/Paper2D/Source/Paper2DEditor/Private/SpriteEditor/SpriteEditorViewportClient.cpp
--- a/Engine/Plugins/Experimental/Paper2D/Source/Paper2DEditor/Private/SpriteEditor/SpriteEditorViewportClient.cpp
+++ b/Engine/Plugins/Experimental/Paper2D/Source/Paper2DEditor/Private/SpriteEditor/SpriteEditorViewportClient.cpp
@@ -617,6 +617,17 @@ void FSpriteEditorViewportClient::ToggleShowSourceTexture()
 	Invalidate();
 }
 
+void FSpriteEditorViewportClient::ToggleShowMeshEdges()
+{
+	EngineShowFlags.MeshEdges = !EngineShowFlags.MeshEdges;
+	Invalidate();
+}
+
+bool FSpriteEditorViewportClient::IsShowMeshEdgesChecked() const
+{
+	return EngineShowFlags.MeshEdges;
+}
+
 void FSpriteEditorViewportClient::UpdateMouseDelta()
 {
 	FPaperEditorViewportClient::UpdateMouseDelta();
diff --git a/Engine/Plugins/Experimental/Paper2D/Source/Paper2DEditor/Private/SpriteEditor/SpriteEditorViewportClient.h b/Engine/Plugins/Experimental/Paper2D/Source/Paper2DEditor/Private/SpriteEditor/SpriteEditorViewportClient.h
--- a/Engine/Plugins/Experimental/Paper2D/Source/Paper2DEditor/Private/SpriteEditor/SpriteEditorViewportClient.h
+++ b/Engine/Plugins/Experimental/Paper2D/Source/Paper2DEditor/Private/SpriteEditor/SpriteEditorViewportClient.h
@@ -67,6 +67,9 @@ public:
 	void ToggleShowNormals() { bShowNormals = !bShowNormals; Invalidate(); }
 	bool IsShowNormalsChecked() const { return bShowNormals; }
 
+	void ToggleShowMeshEdges();
+	bool IsShowMeshEdgesChecked() const;
+
 	void EnterViewMode() { CurrentMode = ESpriteEditorMode::ViewMode; }
 	void EnterSourceRegionEditMode() { CurrentMode = ESpriteEditorMode::EditSourceRegionMode; }
 	void EnterCollisionEditMode() { CurrentMode = ESpriteEditorMode::EditCollisionMode; }
diff --git a/Engine/Source/Editor/UnrealEd/Public/SCommonEditorViewportToolbarBase.h b/Engine/Source/Editor/UnrealEd/Public/SCommonEditorViewportToolbarBase.h
--- a/Engine/Source/Editor/UnrealEd/Public/SCommonEditorViewportToolbarBase.h
+++ b/Engine/Source/Editor/UnrealEd/Public/SCommonEditorViewportToolbarBase.h
@@ -119,7 +119,7 @@ private:
 	 *
 	 * @return The widget containing the show menu content
 	 */
-	TSharedRef<SWidget> GenerateShowMenu() const;
+	virtual TSharedRef<SWidget> GenerateShowMenu() const;
 
 	/**
 	 * @return The widget containing the perspective only FOV window.
