Paper2D: Right-align mode buttons in sprite editor to match style of Blueprint/Persona

diff --git a/Engine/Plugins/Experimental/Paper2D/Source/Paper2DEditor/Private/SpriteEditor/SpriteEditor.cpp b/Engine/Plugins/Experimental/Paper2D/Source/Paper2DEditor/Private/SpriteEditor/SpriteEditor.cpp
--- a/Engine/Plugins/Experimental/Paper2D/Source/Paper2DEditor/Private/SpriteEditor/SpriteEditor.cpp
+++ b/Engine/Plugins/Experimental/Paper2D/Source/Paper2DEditor/Private/SpriteEditor/SpriteEditor.cpp
@@ -517,15 +517,15 @@ void FSpriteEditor::ExtendToolbar()
 			}
 			ToolbarBuilder.EndSection();
 
-			ToolbarBuilder.BeginSection("Mode");
-			{
-				ToolbarBuilder.AddToolBarButton(FSpriteEditorCommands::Get().EnterViewMode);
-				ToolbarBuilder.AddToolBarButton(FSpriteEditorCommands::Get().EnterSourceRegionEditMode);
-				ToolbarBuilder.AddToolBarButton(FSpriteEditorCommands::Get().EnterCollisionEditMode);
-				ToolbarBuilder.AddToolBarButton(FSpriteEditorCommands::Get().EnterRenderingEditMode);
-				//@TODO: PAPER2D: Re-enable once it does something: ToolbarBuilder.AddToolBarButton(FSpriteEditorCommands::Get().EnterAddSpriteMode);
-			}
-			ToolbarBuilder.EndSection();
+// 			ToolbarBuilder.BeginSection("Mode");
+// 			{
+// 				ToolbarBuilder.AddToolBarButton(FSpriteEditorCommands::Get().EnterViewMode);
+// 				ToolbarBuilder.AddToolBarButton(FSpriteEditorCommands::Get().EnterSourceRegionEditMode);
+// 				ToolbarBuilder.AddToolBarButton(FSpriteEditorCommands::Get().EnterCollisionEditMode);
+// 				ToolbarBuilder.AddToolBarButton(FSpriteEditorCommands::Get().EnterRenderingEditMode);
+// 				//@TODO: PAPER2D: Re-enable once it does something: ToolbarBuilder.AddToolBarButton(FSpriteEditorCommands::Get().EnterAddSpriteMode);
+// 			}
+// 			ToolbarBuilder.EndSection();
 		}
 	};
 
@@ -538,10 +538,16 @@ void FSpriteEditor::ExtendToolbar()
 		FToolBarExtensionDelegate::CreateStatic( &Local::FillToolbar )
 		);
 
+	ToolbarExtender->AddToolBarExtension(
+		"Asset",
+		EExtensionHook::After,
+		ViewportPtr->GetCommandList(),
+		FToolBarExtensionDelegate::CreateSP(this, &FSpriteEditor::CreateModeToolbarWidgets));
+
 	AddToolbarExtender(ToolbarExtender);
 
  	IPaper2DEditorModule* Paper2DEditorModule = &FModuleManager::LoadModuleChecked<IPaper2DEditorModule>("Paper2DEditor");
- 	AddToolbarExtender(Paper2DEditorModule->GetSpriteEditorToolBarExtensibilityManager()->GetAllExtenders());
+	AddToolbarExtender(Paper2DEditorModule->GetSpriteEditorToolBarExtensibilityManager()->GetAllExtenders());
 }
 
 void FSpriteEditor::SetSpriteBeingEdited(UPaperSprite* NewSprite)
@@ -560,6 +566,17 @@ void FSpriteEditor::SetSpriteBeingEdited(UPaperSprite* NewSprite)
 	}
 }
 
+
+void FSpriteEditor::CreateModeToolbarWidgets(FToolBarBuilder& IgnoredBuilder)
+{
+	FToolBarBuilder ToolbarBuilder(ViewportPtr->GetCommandList(), FMultiBoxCustomization::None);
+	ToolbarBuilder.AddToolBarButton(FSpriteEditorCommands::Get().EnterViewMode);
+	ToolbarBuilder.AddToolBarButton(FSpriteEditorCommands::Get().EnterSourceRegionEditMode);
+	ToolbarBuilder.AddToolBarButton(FSpriteEditorCommands::Get().EnterCollisionEditMode);
+	ToolbarBuilder.AddToolBarButton(FSpriteEditorCommands::Get().EnterRenderingEditMode);
+	AddToolbarWidget(ToolbarBuilder.MakeWidget());
+}
+
 //////////////////////////////////////////////////////////////////////////
 
 #undef LOCTEXT_NAMESPACE
\ No newline at end of file
diff --git a/Engine/Plugins/Experimental/Paper2D/Source/Paper2DEditor/Private/SpriteEditor/SpriteEditor.h b/Engine/Plugins/Experimental/Paper2D/Source/Paper2DEditor/Private/SpriteEditor/SpriteEditor.h
--- a/Engine/Plugins/Experimental/Paper2D/Source/Paper2DEditor/Private/SpriteEditor/SpriteEditor.h
+++ b/Engine/Plugins/Experimental/Paper2D/Source/Paper2DEditor/Private/SpriteEditor/SpriteEditor.h
@@ -48,4 +48,6 @@ protected:
 	TSharedRef<SDockTab> SpawnTab_Viewport(const FSpawnTabArgs& Args);
 	TSharedRef<SDockTab> SpawnTab_Details(const FSpawnTabArgs& Args);
 	TSharedRef<SDockTab> SpawnTab_SpriteList(const FSpawnTabArgs& Args);
+
+	void CreateModeToolbarWidgets(FToolBarBuilder& ToolbarBuilder);
 };
\ No newline at end of file
