Paper2D: Add material expression node as shorthand for Texture2D Param named 'SpriteTexture'
Integrate CL# 2124084 from 4.3 to main

diff --git a/Engine/Plugins/Experimental/Paper2D/Source/Paper2D/Private/Rendering/MaterialExpressionSpriteTextureSampler.cpp b/Engine/Plugins/Experimental/Paper2D/Source/Paper2D/Private/Rendering/MaterialExpressionSpriteTextureSampler.cpp
new file mode 100644
--- /dev/null
+++ b/Engine/Plugins/Experimental/Paper2D/Source/Paper2D/Private/Rendering/MaterialExpressionSpriteTextureSampler.cpp
@@ -0,0 +1,50 @@
+// Copyright 1998-2014 Epic Games, Inc. All Rights Reserved.
+
+#include "Paper2DPrivatePCH.h"
+#include "MaterialExpressionSpriteTextureSampler.h"
+
+#define LOCTEXT_NAMESPACE "Paper2D"
+
+//////////////////////////////////////////////////////////////////////////
+// UMaterialExpressionSpriteTextureSampler
+
+UMaterialExpressionSpriteTextureSampler::UMaterialExpressionSpriteTextureSampler(const FPostConstructInitializeProperties& PCIP)
+	: Super(PCIP)
+{
+	ParameterName = TEXT("SpriteTexture");
+}
+
+void UMaterialExpressionSpriteTextureSampler::GetCaption(TArray<FString>& OutCaptions) const
+{
+	OutCaptions.Add(TEXT("Paper2D"));
+	OutCaptions.Add(TEXT("SpriteTexture"));
+}
+
+#if WITH_EDITOR
+FString UMaterialExpressionSpriteTextureSampler::GetKeywords() const
+{
+	FString ParentKeywords = Super::GetKeywords();
+	return ParentKeywords + TEXT(" Paper2D Sprite");
+}
+
+bool UMaterialExpressionSpriteTextureSampler::CanRenameNode() const
+{
+	// The parameter name is read only on this guy
+	return false;
+}
+
+void UMaterialExpressionSpriteTextureSampler::GetExpressionToolTip(TArray<FString>& OutToolTip)
+{
+	OutToolTip.Add(LOCTEXT("SpriteTextureSamplerTooltip", "This is a texture sampler 2D with the parameter name fixed as 'SpriteTexture'. The texture specified here will be replaced by the SourceTexture of a Paper2D sprite if this material is used on a sprite.").ToString());
+}
+
+void UMaterialExpressionSpriteTextureSampler::PostEditChangeProperty(FPropertyChangedEvent& PropertyChangedEvent)
+{
+	// Ensure that the parameter name never changes
+	ParameterName = TEXT("SpriteTexture");
+	Super::PostEditChangeProperty(PropertyChangedEvent);
+}
+
+#endif
+
+#undef LOCTEXT_NAMESPACE
diff --git a/Engine/Plugins/Experimental/Paper2D/Source/Paper2D/Public/MaterialExpressionSpriteTextureSampler.h b/Engine/Plugins/Experimental/Paper2D/Source/Paper2D/Public/MaterialExpressionSpriteTextureSampler.h
new file mode 100644
--- /dev/null
+++ b/Engine/Plugins/Experimental/Paper2D/Source/Paper2D/Public/MaterialExpressionSpriteTextureSampler.h
@@ -0,0 +1,23 @@
+// Copyright 1998-2014 Epic Games, Inc. All Rights Reserved.
+
+#pragma once
+
+#include "Runtime/Engine/Classes/Materials/MaterialExpressionTextureSampleParameter2D.h"
+#include "MaterialExpressionSpriteTextureSampler.generated.h"
+
+// This is a texture sampler 2D with the parameter name fixed as 'SpriteTexture'. The texture specified here will be replaced by the SourceTexture of a Paper2D sprite if this material is used on a sprite.
+UCLASS(hideCategories=(MaterialExpressionTextureSampleParameter, MaterialExpressionSpriteTextureSampler))
+class PAPER2D_API UMaterialExpressionSpriteTextureSampler : public UMaterialExpressionTextureSampleParameter2D
+{
+	GENERATED_UCLASS_BODY()
+
+	// UMaterialExpression interface
+	virtual void GetCaption(TArray<FString>& OutCaptions) const override;
+#if WITH_EDITOR
+	virtual FString GetKeywords() const override;
+	virtual bool CanRenameNode() const override;
+	virtual void GetExpressionToolTip(TArray<FString>& OutToolTip) override;
+	virtual void PostEditChangeProperty(FPropertyChangedEvent& PropertyChangedEvent) override;
+#endif
+	// End of UMaterialExpression interface
+};
