Paper2D: Allow newly created empty sprites automatically adjust their dimensions when the source texture is first set
Integrate CL# 2124122 from 4.3 to main

diff --git a/Engine/Plugins/Experimental/Paper2D/Source/Paper2D/Private/PaperSprite.cpp b/Engine/Plugins/Experimental/Paper2D/Source/Paper2D/Private/PaperSprite.cpp
--- a/Engine/Plugins/Experimental/Paper2D/Source/Paper2D/Private/PaperSprite.cpp
+++ b/Engine/Plugins/Experimental/Paper2D/Source/Paper2D/Private/PaperSprite.cpp
@@ -276,6 +276,9 @@ void UPaperSprite::PostEditChangeProperty(FPropertyChangedEvent& PropertyChanged
 		PixelsPerUnrealUnit = 1.0f;
 	}
 
+	SourceDimension.X = FMath::Max(SourceDimension.X, 0.0f);
+	SourceDimension.Y = FMath::Max(SourceDimension.Y, 0.0f);
+
 	//@TODO: Determine when these are really needed, as they're seriously expensive!
 	TComponentReregisterContext<UPaperSpriteComponent> ReregisterStaticComponents;
 	TComponentReregisterContext<UPaperFlipbookComponent> ReregisterAnimatedComponents;
@@ -316,13 +319,23 @@ void UPaperSprite::PostEditChangeProperty(FPropertyChangedEvent& PropertyChanged
 
 	if ((PropertyName == GET_MEMBER_NAME_CHECKED(UPaperSprite, SourceUV)) ||
 		(PropertyName == GET_MEMBER_NAME_CHECKED(UPaperSprite, SourceDimension)) ||
-		(PropertyName == GET_MEMBER_NAME_CHECKED(UPaperSprite, SourceTexture)) ||
 		(PropertyName == GET_MEMBER_NAME_CHECKED(UPaperSprite, CustomPivotPoint)) ||
 		(PropertyName == GET_MEMBER_NAME_CHECKED(UPaperSprite, PivotMode)) )
 	{
 		bBothModified = true;
 	}
 
+	if (PropertyName == GET_MEMBER_NAME_CHECKED(UPaperSprite, SourceTexture))
+	{
+		// If this is a brand new sprite that didn't have a texture set previously, act like we were factoried with the texture
+		if ((SourceTexture != nullptr) && SourceDimension.IsNearlyZero())
+		{
+			SourceUV = FVector2D::ZeroVector;
+			SourceDimension = FVector2D(SourceTexture->GetSizeX(), SourceTexture->GetSizeY());
+		}
+		bBothModified = true;
+	}
+
 	if (bCollisionDataModified || bBothModified)
 	{
 		RebuildCollisionData();
@@ -911,7 +924,7 @@ void UPaperSprite::InitializeSprite(UTexture2D* Texture, float InPixelsPerUnreal
 	}
 	else
 	{
-		InitializeSprite(NULL, FVector2D::ZeroVector, FVector2D(1.0f, 1.0f), InPixelsPerUnrealUnit);
+		InitializeSprite(NULL, FVector2D::ZeroVector, FVector2D::ZeroVector, InPixelsPerUnrealUnit);
 	}
 }
 
