Paper2D: Fix bug where the component displaying the 'Show Src Tex' view didn't refresh when SourceTexture was modified
Integrate CL# 2124121 from 4.3 to main

diff --git a/Engine/Plugins/Experimental/Paper2D/Source/Paper2DEditor/Private/SpriteEditor/SpriteEditorViewportClient.cpp b/Engine/Plugins/Experimental/Paper2D/Source/Paper2DEditor/Private/SpriteEditor/SpriteEditorViewportClient.cpp
--- a/Engine/Plugins/Experimental/Paper2D/Source/Paper2DEditor/Private/SpriteEditor/SpriteEditorViewportClient.cpp
+++ b/Engine/Plugins/Experimental/Paper2D/Source/Paper2DEditor/Private/SpriteEditor/SpriteEditorViewportClient.cpp
@@ -6,6 +6,7 @@
 
 #include "PreviewScene.h"
 #include "ScopedTransaction.h"
+#include "Runtime/Engine/Public/ComponentReregisterContext.h"
 
 #define LOCTEXT_NAMESPACE "SpriteEditor"
 
@@ -99,6 +100,7 @@ void FSpriteEditorViewportClient::UpdateSourceTextureSpriteFromSprite(UPaperSpri
 	{
 		if ((SourceSprite->GetSourceTexture() != TargetSprite->GetSourceTexture()) || (TargetSprite->PixelsPerUnrealUnit != SourceSprite->PixelsPerUnrealUnit))
 		{
+			FComponentReregisterContext ReregisterSprite(SourceTextureViewComponent);
 			TargetSprite->InitializeSprite(SourceSprite->SourceTexture, SourceSprite->PixelsPerUnrealUnit);
 		}
 	}
