Paper2D: Hide '2D queries are disabled' warning if sprite isn't actually set to 2D collision
Integrate CL# 2124120 from 4.3 to main

diff --git a/Engine/Plugins/Experimental/Paper2D/Source/Paper2DEditor/Private/SpriteEditor/SpriteDetailsCustomization.cpp b/Engine/Plugins/Experimental/Paper2D/Source/Paper2DEditor/Private/SpriteEditor/SpriteDetailsCustomization.cpp
--- a/Engine/Plugins/Experimental/Paper2D/Source/Paper2DEditor/Private/SpriteEditor/SpriteDetailsCustomization.cpp
+++ b/Engine/Plugins/Experimental/Paper2D/Source/Paper2DEditor/Private/SpriteEditor/SpriteDetailsCustomization.cpp
@@ -110,7 +110,7 @@ void FSpriteDetailsCustomization::BuildCollisionSection(IDetailCategoryBuilder&
 		.Visibility(ParticipatesInPhysics2D);
 
 	// Add a warning bar if 2D collision queries aren't enabled
-	TAttribute<EVisibility> WarnAbout2DQueriesBeingDisabledVisibility = TAttribute<EVisibility>::Create(TAttribute<EVisibility>::FGetter::CreateStatic(&FSpriteDetailsCustomization::Get2DPhysicsNotEnabledWarningVisibility));
+	TAttribute<EVisibility> WarnAbout2DQueriesBeingDisabledVisibility = TAttribute<EVisibility>::Create(TAttribute<EVisibility>::FGetter::CreateSP(this, &FSpriteDetailsCustomization::Get2DPhysicsNotEnabledWarningVisibility, SpriteCollisionDomainProperty));
 	FText QueryWarningFor2D = LOCTEXT("Query2DPhysicsWarning", "2D collision queries are disabled");
 	FText QueryTooltipFor2D = LOCTEXT("Query2DPhysicsWarningTooltip", "You can enable 2D queries in Project Settings..Physics by setting bEnable2DPhysics to true, otherwise only collision detection and response will work.\n\nNote: Only raycasts are partially supported; other queries, sweeps, and overlap tests are not implemented yet.");
 	GenerateWarningRow(CollisionCategory, /*bExperimental=*/ false, QueryWarningFor2D, QueryTooltipFor2D, TEXT("Shared/Editors/SpriteEditor"), TEXT("Disabled2DCollisionQueriesWarning"))
@@ -181,9 +181,10 @@ EVisibility FSpriteDetailsCustomization::GetCustomPivotVisibility(TSharedPtr<IPr
 	return EVisibility::Visible;
 }
 
-EVisibility FSpriteDetailsCustomization::Get2DPhysicsNotEnabledWarningVisibility()
+EVisibility FSpriteDetailsCustomization::Get2DPhysicsNotEnabledWarningVisibility(TSharedPtr<IPropertyHandle> Property) const
 {
-	return GetDefault<UPhysicsSettings>()->bEnable2DPhysics ? EVisibility::Collapsed : EVisibility::Visible;
+	// Hide the warning if 2D queries are enabled, or if the collision mode is not 2D physics
+	return GetDefault<UPhysicsSettings>()->bEnable2DPhysics ? EVisibility::Collapsed : PhysicsModeMatches(Property, ESpriteCollisionMode::Use2DPhysics);
 }
 
 #undef LOCTEXT_NAMESPACE
diff --git a/Engine/Plugins/Experimental/Paper2D/Source/Paper2DEditor/Private/SpriteEditor/SpriteDetailsCustomization.h b/Engine/Plugins/Experimental/Paper2D/Source/Paper2DEditor/Private/SpriteEditor/SpriteDetailsCustomization.h
--- a/Engine/Plugins/Experimental/Paper2D/Source/Paper2DEditor/Private/SpriteEditor/SpriteDetailsCustomization.h
+++ b/Engine/Plugins/Experimental/Paper2D/Source/Paper2DEditor/Private/SpriteEditor/SpriteDetailsCustomization.h
@@ -21,9 +21,9 @@ protected:
 	EVisibility PhysicsModeMatches(TSharedPtr<IPropertyHandle> Property, ESpriteCollisionMode::Type DesiredMode) const;
 	EVisibility AnyPhysicsMode(TSharedPtr<IPropertyHandle> Property) const;
 	EVisibility GetCustomPivotVisibility(TSharedPtr<IPropertyHandle> Property) const;
+	EVisibility Get2DPhysicsNotEnabledWarningVisibility(TSharedPtr<IPropertyHandle> Property) const;
 
 	static EVisibility GetAtlasGroupVisibility();
-	static EVisibility Get2DPhysicsNotEnabledWarningVisibility();
 
 	static FDetailWidgetRow& GenerateWarningRow(IDetailCategoryBuilder& WarningCategory, bool bExperimental, const FText& WarningText, const FText& Tooltip, const FString& ExcerptLink, const FString& ExcerptName);
 
