Paper2D: Disable terrain spline actors and tile map actors by default (gated by experimental flags in UPaperRuntimeSettings)
- Add associated actor factories so they can return false when the flag is disabled
- Gate the registration of the tile map editor mode by the flag
- Spline terrain factory supports creation via drag-dropping of material
Integrate CL# 2124118 from 4.3 to Main

diff --git a/Engine/Plugins/Experimental/Paper2D/Source/Paper2D/Classes/PaperRuntimeSettings.h b/Engine/Plugins/Experimental/Paper2D/Source/Paper2D/Classes/PaperRuntimeSettings.h
--- a/Engine/Plugins/Experimental/Paper2D/Source/Paper2D/Classes/PaperRuntimeSettings.h
+++ b/Engine/Plugins/Experimental/Paper2D/Source/Paper2D/Classes/PaperRuntimeSettings.h
@@ -19,4 +19,12 @@ class PAPER2D_API UPaperRuntimeSettings : public UObject
 	// Enables experimental *incomplete and unsupported* texture atlas groups that sprites can be assigned to
 	UPROPERTY(EditAnywhere, config, Category=Experimental)
 	bool bEnableSpriteAtlasGroups;
+
+	// Enables experimental *incomplete and unsupported* tile map editing. Note: You need to restart the editor when enabling this setting for the change to fully take effect.
+	UPROPERTY(EditAnywhere, config, Category=Experimental)
+	bool bEnableTileMapEditing;
+
+	// Enables experimental *incomplete and unsupported* 2D terrain spline editing. Note: You need to restart the editor when enabling this setting for the change to fully take effect.
+	UPROPERTY(EditAnywhere, config, Category=Experimental)
+	bool bEnableTerrainSplineEditing;
 };
diff --git a/Engine/Plugins/Experimental/Paper2D/Source/Paper2D/Private/PaperRuntimeSettings.cpp b/Engine/Plugins/Experimental/Paper2D/Source/Paper2D/Private/PaperRuntimeSettings.cpp
--- a/Engine/Plugins/Experimental/Paper2D/Source/Paper2D/Private/PaperRuntimeSettings.cpp
+++ b/Engine/Plugins/Experimental/Paper2D/Source/Paper2D/Private/PaperRuntimeSettings.cpp
@@ -10,5 +10,7 @@ UPaperRuntimeSettings::UPaperRuntimeSettings(const FPostConstructInitializePrope
 	: Super(PCIP)
 	, DefaultPixelsPerUnrealUnit(2.56f)
 	, bEnableSpriteAtlasGroups(false)
+	, bEnableTileMapEditing(false)
+	, bEnableTerrainSplineEditing(false)
 {
 }
diff --git a/Engine/Plugins/Experimental/Paper2D/Source/Paper2DEditor/Classes/TerrainSplineActorFactory.h b/Engine/Plugins/Experimental/Paper2D/Source/Paper2DEditor/Classes/TerrainSplineActorFactory.h
new file mode 100644
--- /dev/null
+++ b/Engine/Plugins/Experimental/Paper2D/Source/Paper2DEditor/Classes/TerrainSplineActorFactory.h
@@ -0,0 +1,15 @@
+// Copyright 1998-2014 Epic Games, Inc. All Rights Reserved.
+
+#include "TerrainSplineActorFactory.generated.h"
+
+UCLASS()
+class UTerrainSplineActorFactory : public UActorFactory
+{
+	GENERATED_UCLASS_BODY()
+
+	// UActorFactory interface
+	virtual void PostSpawnActor(UObject* Asset, AActor* NewActor) override;
+	virtual void PostCreateBlueprint(UObject* Asset, AActor* CDO) override;
+	virtual bool CanCreateActorFrom(const FAssetData& AssetData, FText& OutErrorMsg) override;
+	// End of UActorFactory interface
+};
diff --git a/Engine/Plugins/Experimental/Paper2D/Source/Paper2DEditor/Classes/TileMapActorFactory.h b/Engine/Plugins/Experimental/Paper2D/Source/Paper2DEditor/Classes/TileMapActorFactory.h
new file mode 100644
--- /dev/null
+++ b/Engine/Plugins/Experimental/Paper2D/Source/Paper2DEditor/Classes/TileMapActorFactory.h
@@ -0,0 +1,13 @@
+// Copyright 1998-2014 Epic Games, Inc. All Rights Reserved.
+
+#include "TileMapActorFactory.generated.h"
+
+UCLASS()
+class UTileMapActorFactory : public UActorFactory
+{
+	GENERATED_UCLASS_BODY()
+
+	// UActorFactory interface
+	virtual bool CanCreateActorFrom(const FAssetData& AssetData, FText& OutErrorMsg) override;
+	// End of UActorFactory interface
+};
diff --git a/Engine/Plugins/Experimental/Paper2D/Source/Paper2DEditor/Private/Paper2DEditorModule.cpp b/Engine/Plugins/Experimental/Paper2D/Source/Paper2DEditor/Private/Paper2DEditorModule.cpp
--- a/Engine/Plugins/Experimental/Paper2D/Source/Paper2DEditor/Private/Paper2DEditorModule.cpp
+++ b/Engine/Plugins/Experimental/Paper2D/Source/Paper2DEditor/Private/Paper2DEditorModule.cpp
@@ -115,11 +115,7 @@ public:
 		UThumbnailManager::Get().RegisterCustomRenderer(UPaperFlipbook::StaticClass(), UPaperFlipbookThumbnailRenderer::StaticClass());
 
 		// Register the editor modes
-		FEditorModeRegistry::Get().RegisterMode<FEdModeTileMap>(
-			FEdModeTileMap::EM_TileMap,
-			LOCTEXT("TileMapEditMode", "Tile Map Editor"),
-			FSlateIcon(),
-			true);
+		UpdateTileMapEditorModeInstallation();
 
 		// Integrate Paper2D actions associated with existing engine types (e.g., Texture2D) into the content browser
 		FPaperContentBrowserExtensions::InstallHooks();
@@ -189,6 +185,10 @@ private:
 		{
 			FPaperAtlasGenerator::HandleAssetChangedEvent(Atlas);
 		}
+		else if (UPaperRuntimeSettings* Settings = Cast<UPaperRuntimeSettings>(ObjectBeingModified))
+		{
+			UpdateTileMapEditorModeInstallation();
+		}
 	}
 
 	void RegisterSettings()
@@ -238,6 +238,25 @@ private:
 			IIntroTutorials::Get().UnregisterTutorialForAssetEditor(UPaperFlipbook::StaticClass());
 		}
 	}
+
+	// Installs or uninstalls the tile map editing mode depending on settings
+	void UpdateTileMapEditorModeInstallation()
+	{
+		const bool bAlreadyRegistered = FEditorModeRegistry::Get().GetFactoryMap().Contains(FEdModeTileMap::EM_TileMap);
+		const bool bShouldBeRegistered = GetDefault<UPaperRuntimeSettings>()->bEnableTileMapEditing;
+		if (bAlreadyRegistered && !bShouldBeRegistered)
+		{
+			FEditorModeRegistry::Get().UnregisterMode(FEdModeTileMap::EM_TileMap);
+		}
+		else if (!bAlreadyRegistered && bShouldBeRegistered)
+		{
+			FEditorModeRegistry::Get().RegisterMode<FEdModeTileMap>(
+				FEdModeTileMap::EM_TileMap,
+				LOCTEXT("TileMapEditMode", "Tile Map Editor"),
+				FSlateIcon(),
+				true);
+		}
+	}
 };
 
 //////////////////////////////////////////////////////////////////////////
diff --git a/Engine/Plugins/Experimental/Paper2D/Source/Paper2DEditor/Private/TerrainSplineEditing/TerrainSplineActorFactory.cpp b/Engine/Plugins/Experimental/Paper2D/Source/Paper2DEditor/Private/TerrainSplineEditing/TerrainSplineActorFactory.cpp
new file mode 100644
--- /dev/null
+++ b/Engine/Plugins/Experimental/Paper2D/Source/Paper2DEditor/Private/TerrainSplineEditing/TerrainSplineActorFactory.cpp
@@ -0,0 +1,54 @@
+// Copyright 1998-2014 Epic Games, Inc. All Rights Reserved.
+
+#include "Paper2DEditorPrivatePCH.h"
+#include "AssetData.h"
+
+//////////////////////////////////////////////////////////////////////////
+// UTerrainSplineActorFactory
+
+UTerrainSplineActorFactory::UTerrainSplineActorFactory(const FPostConstructInitializeProperties& PCIP)
+	: Super(PCIP)
+{
+	DisplayName = NSLOCTEXT("Paper2D", "TerrainSplineFactoryDisplayName", "Paper2D Terrain Spline");
+	NewActorClass = APaperTerrainActor::StaticClass();
+}
+
+void UTerrainSplineActorFactory::PostSpawnActor(UObject* Asset, AActor* NewActor)
+{
+	if (UPaperTerrainMaterial* TerrainMaterial = Cast<UPaperTerrainMaterial>(Asset))
+	{
+		GEditor->SetActorLabelUnique(NewActor, TerrainMaterial->GetName());
+
+		APaperTerrainActor* TypedActor = CastChecked<APaperTerrainActor>(NewActor);
+		UPaperTerrainComponent* RenderComponent = TypedActor->RenderComponent;
+		check(RenderComponent);
+
+		RenderComponent->UnregisterComponent();
+		RenderComponent->TerrainMaterial = TerrainMaterial;
+		RenderComponent->RegisterComponent();
+	}
+}
+
+void UTerrainSplineActorFactory::PostCreateBlueprint(UObject* Asset, AActor* CDO)
+{
+	checkf(false, TEXT("APaperTerrainActor isn't blueprintable; how did you get here?"));
+}
+
+bool UTerrainSplineActorFactory::CanCreateActorFrom(const FAssetData& AssetData, FText& OutErrorMsg)
+{
+	if (GetDefault<UPaperRuntimeSettings>()->bEnableTerrainSplineEditing)
+	{
+		if (AssetData.IsValid() && AssetData.GetClass()->IsChildOf(UPaperTerrainMaterial::StaticClass()))
+		{
+			return true;
+		}
+		else
+		{
+			return Super::CanCreateActorFrom(AssetData, OutErrorMsg);
+		}
+	}
+	else
+	{
+		return false;
+	}
+}
diff --git a/Engine/Plugins/Experimental/Paper2D/Source/Paper2DEditor/Private/TileMapEditing/TileMapActorFactory.cpp b/Engine/Plugins/Experimental/Paper2D/Source/Paper2DEditor/Private/TileMapEditing/TileMapActorFactory.cpp
new file mode 100644
--- /dev/null
+++ b/Engine/Plugins/Experimental/Paper2D/Source/Paper2DEditor/Private/TileMapEditing/TileMapActorFactory.cpp
@@ -0,0 +1,26 @@
+// Copyright 1998-2014 Epic Games, Inc. All Rights Reserved.
+
+#include "Paper2DEditorPrivatePCH.h"
+#include "AssetData.h"
+
+//////////////////////////////////////////////////////////////////////////
+// UTileMapActorFactory
+
+UTileMapActorFactory::UTileMapActorFactory(const FPostConstructInitializeProperties& PCIP)
+	: Super(PCIP)
+{
+	DisplayName = NSLOCTEXT("Paper2D", "TileMapFactoryDisplayName", "Paper2D Tile Map");
+	NewActorClass = APaperTileMapActor::StaticClass();
+}
+
+bool UTileMapActorFactory::CanCreateActorFrom(const FAssetData& AssetData, FText& OutErrorMsg)
+{
+	if (GetDefault<UPaperRuntimeSettings>()->bEnableTileMapEditing)
+	{
+		return Super::CanCreateActorFrom(AssetData, OutErrorMsg);
+	}
+	else
+	{
+		return false;
+	}
+}
