Paper2D: Fix crash when creating a blueprint from a placed flipbook or sprite actor using button in the details panel
Integrate 2126099 from 4.3 to Main

diff --git a/Engine/Plugins/Experimental/Paper2D/Source/Paper2DEditor/Private/PaperAnimatedSpriteActorFactory.cpp b/Engine/Plugins/Experimental/Paper2D/Source/Paper2DEditor/Private/PaperAnimatedSpriteActorFactory.cpp
--- a/Engine/Plugins/Experimental/Paper2D/Source/Paper2DEditor/Private/PaperAnimatedSpriteActorFactory.cpp
+++ b/Engine/Plugins/Experimental/Paper2D/Source/Paper2DEditor/Private/PaperAnimatedSpriteActorFactory.cpp
@@ -31,7 +31,16 @@ void UPaperFlipbookActorFactory::PostSpawnActor(UObject* Asset, AActor* NewActor
 
 void UPaperFlipbookActorFactory::PostCreateBlueprint(UObject* Asset, AActor* CDO)
 {
-	checkf(false, TEXT("APaperFlipbookActor isn't blueprintable; how did you get here?"));
+	if (UPaperFlipbook* Flipbook = Cast<UPaperFlipbook>(Asset))
+	{
+		if (APaperFlipbookActor* TypedActor = Cast<APaperFlipbookActor>(CDO))
+		{
+			UPaperFlipbookComponent* RenderComponent = TypedActor->RenderComponent;
+			check(RenderComponent);
+
+			RenderComponent->SetFlipbook(Flipbook);
+		}
+	}
 }
 
 bool UPaperFlipbookActorFactory::CanCreateActorFrom(const FAssetData& AssetData, FText& OutErrorMsg)
diff --git a/Engine/Plugins/Experimental/Paper2D/Source/Paper2DEditor/Private/PaperSpriteActorFactory.cpp b/Engine/Plugins/Experimental/Paper2D/Source/Paper2DEditor/Private/PaperSpriteActorFactory.cpp
--- a/Engine/Plugins/Experimental/Paper2D/Source/Paper2DEditor/Private/PaperSpriteActorFactory.cpp
+++ b/Engine/Plugins/Experimental/Paper2D/Source/Paper2DEditor/Private/PaperSpriteActorFactory.cpp
@@ -31,7 +31,16 @@ void UPaperSpriteActorFactory::PostSpawnActor(UObject* Asset, AActor* NewActor)
 
 void UPaperSpriteActorFactory::PostCreateBlueprint(UObject* Asset, AActor* CDO)
 {
-	checkf(false, TEXT("APaperSpriteActor isn't blueprintable; how did you get here?"));
+	if (UPaperSprite* Sprite = Cast<UPaperSprite>(Asset))
+	{
+		if (APaperSpriteActor* TypedActor = Cast<APaperSpriteActor>(CDO))
+		{
+			UPaperSpriteComponent* RenderComponent = TypedActor->RenderComponent;
+			check(RenderComponent);
+
+			RenderComponent->SetSprite(Sprite);
+		}
+	}
 }
 
 bool UPaperSpriteActorFactory::CanCreateActorFrom(const FAssetData& AssetData, FText& OutErrorMsg)
