Include Character.h in PaperCharacter.h
Include Interface_CollisionDataProvider.h in PaperSprite.h

diff --git a/Engine/Plugins/Experimental/Paper2D/Source/Paper2D/Classes/PaperCharacter.h b/Engine/Plugins/Experimental/Paper2D/Source/Paper2D/Classes/PaperCharacter.h
--- a/Engine/Plugins/Experimental/Paper2D/Source/Paper2D/Classes/PaperCharacter.h
+++ b/Engine/Plugins/Experimental/Paper2D/Source/Paper2D/Classes/PaperCharacter.h
@@ -1,7 +1,7 @@
 // Copyright 1998-2014 Epic Games, Inc. All Rights Reserved.
 
 #pragma once
-
+#include "GameFramework/Character.h"
 #include "PaperCharacter.generated.h"
 
 // APaperCharacter behaves like ACharacter, but uses a UPaperFlipbookComponent instead of a USkeletalMeshComponent as a visual representation
diff --git a/Engine/Plugins/Experimental/Paper2D/Source/Paper2D/Classes/PaperSprite.h b/Engine/Plugins/Experimental/Paper2D/Source/Paper2D/Classes/PaperSprite.h
--- a/Engine/Plugins/Experimental/Paper2D/Source/Paper2D/Classes/PaperSprite.h
+++ b/Engine/Plugins/Experimental/Paper2D/Source/Paper2D/Classes/PaperSprite.h
@@ -4,7 +4,7 @@
 
 #include "SpriteEditorOnlyTypes.h"
 #include "Engine/EngineTypes.h"
-
+#include "Interfaces/Interface_CollisionDataProvider.h"
 #include "PaperSprite.generated.h"
 
 UENUM()
