Paper2D: Add comments for several methods on UPaperSprite

diff --git a/Engine/Plugins/Experimental/Paper2D/Source/Paper2D/Classes/PaperSprite.h b/Engine/Plugins/Experimental/Paper2D/Source/Paper2D/Classes/PaperSprite.h
--- a/Engine/Plugins/Experimental/Paper2D/Source/Paper2D/Classes/PaperSprite.h
+++ b/Engine/Plugins/Experimental/Paper2D/Source/Paper2D/Classes/PaperSprite.h
@@ -176,12 +176,16 @@ public:
 	UTexture2D* GetSourceTexture() const { return SourceTexture; }
 #endif
 
+	// Return the scaling factor between pixels and Unreal units (cm)
 	float GetPixelsPerUnrealUnit() const { return PixelsPerUnrealUnit; }
+
+	// Return the scaling factor between Unreal units (cm) and pixels
 	float GetUnrealUnitsPerPixel() const { return 1.0f / PixelsPerUnrealUnit; }
 
 	// Returns the texture this should be rendered with
 	UTexture2D* GetBakedTexture() const;
 
+	// Return the default material for this sprite
 	UMaterialInterface* GetDefaultMaterial() const { return DefaultMaterial; }
 
 	// Returns the render bounds of this sprite
@@ -189,7 +193,11 @@ public:
 
 	// Search for a socket (note: do not cache this pointer; it's unsafe if the Socket array is edited)
 	FPaperSpriteSocket* FindSocket(FName SocketName);
+
+	// Returns true if the sprite has any sockets
 	bool HasAnySockets() const { return Sockets.Num() > 0; }
+	
+	// Returns a list of all of the sockets
 	void QuerySupportedSockets(TArray<FComponentSocketDescription>& OutSockets) const;
 
 	// IInterface_CollisionDataProvider interface
