Paper2D: Add support for importing tile maps and tile sets exported from Tiled 0.90 as JSON (WIP)

diff --git a/Engine/Plugins/Experimental/Paper2D/Paper2D.uplugin b/Engine/Plugins/Experimental/Paper2D/Paper2D.uplugin
--- a/Engine/Plugins/Experimental/Paper2D/Paper2D.uplugin
+++ b/Engine/Plugins/Experimental/Paper2D/Paper2D.uplugin
@@ -26,6 +26,10 @@
 			"Name" : "PaperJsonImporter",
 			"Type" : "Editor"
 		},
+		{
+			"Name" : "PaperTiledImporter",
+			"Type" : "Editor"
+		},
 		{
 			"Name" : "SmartSnapping",
 			"Type" : "Editor"
diff --git a/Engine/Plugins/Experimental/Paper2D/Source/Paper2D/Classes/PaperTileLayer.h b/Engine/Plugins/Experimental/Paper2D/Source/Paper2D/Classes/PaperTileLayer.h
--- a/Engine/Plugins/Experimental/Paper2D/Source/Paper2D/Classes/PaperTileLayer.h
+++ b/Engine/Plugins/Experimental/Paper2D/Source/Paper2D/Classes/PaperTileLayer.h
@@ -11,10 +11,10 @@ struct FPaperTileInfo
 {
 	GENERATED_USTRUCT_BODY()
 
-	UPROPERTY()
+	UPROPERTY(Category=Debug, VisibleAnywhere)
 	UPaperTileSet* TileSet;
 
-	UPROPERTY()
+	UPROPERTY(Category=Debug, VisibleAnywhere)
 	int32 PackedTileIndex;
 
 	FPaperTileInfo()
@@ -39,38 +39,54 @@ class PAPER2D_API UPaperTileLayer : public UObject
 {
 	GENERATED_UCLASS_BODY()
 
-	UPROPERTY()
+	// Name of the layer
+	UPROPERTY(Category=Debug, VisibleAnywhere)
 	FText LayerName;
 
-	UPROPERTY()
+	// Width of the layer (in tiles)
+	UPROPERTY(Category=Debug, VisibleAnywhere)
 	int32 LayerWidth;
 
-	UPROPERTY()
+	// Height of the layer (in tiles)
+	UPROPERTY(Category=Debug, VisibleAnywhere)
 	int32 LayerHeight;
 
-	UPROPERTY()
+	// Tile Set associated with the layer //@TODO: Remove this or make it editor only
+	UPROPERTY(Category=Debug, VisibleAnywhere)
 	UPaperTileSet* TileSet;
 
-	UPROPERTY()
+#if WITH_EDITORONLY_DATA
+	// Opacity of the layer (editor only)
+	UPROPERTY(Category=Debug, VisibleAnywhere)
+	float LayerOpacity;
+
+	// Is this layer currently hidden in the editor?
+	UPROPERTY(Category=Debug, VisibleAnywhere)
 	bool bHiddenInEditor;
+#endif
 
-	UPROPERTY()
+	// Should this layer be hidden in the game?
+	UPROPERTY(Category=Debug, VisibleAnywhere)
 	bool bHiddenInGame;
 
-	UPROPERTY()
+	// Is this layer a collision layer?
+	UPROPERTY(Category=Debug, VisibleAnywhere)
 	bool bCollisionLayer;
 
 protected:
-	UPROPERTY()
+	// The allocated width of the tile data (used to handle resizing without data loss)
+	UPROPERTY(Category=Debug, VisibleAnywhere)
 	int32 AllocatedWidth;
 
-	UPROPERTY()
+	// The allocated height of the tile data (used to handle resizing without data loss)
+	UPROPERTY(Category=Debug, VisibleAnywhere)
 	int32 AllocatedHeight;
 
 	UPROPERTY()
 	TArray<int32> AllocatedGrid_DEPRECATED;
 
-	UPROPERTY()
+	// The allocated tile data
+	UPROPERTY(Category=Debug, VisibleAnywhere)
 	TArray<FPaperTileInfo> AllocatedCells;
 
 public:
diff --git a/Engine/Plugins/Experimental/Paper2D/Source/Paper2D/Classes/PaperTileMap.h b/Engine/Plugins/Experimental/Paper2D/Source/Paper2D/Classes/PaperTileMap.h
--- a/Engine/Plugins/Experimental/Paper2D/Source/Paper2D/Classes/PaperTileMap.h
+++ b/Engine/Plugins/Experimental/Paper2D/Source/Paper2D/Classes/PaperTileMap.h
@@ -40,17 +40,24 @@ class PAPER2D_API UPaperTileMap : public UDataAsset //@TODO: Just to make it eas
 	UMaterialInterface* Material;
 
 	// The list of layers
-	UPROPERTY()
+	UPROPERTY(Category=Debug, VisibleAnywhere, EditInline) //@TODO: Remove this
 	TArray<UPaperTileLayer*> TileLayers;
 
 	// Collision domain (no collision, 2D, or 3D)
-	UPROPERTY(Category = Collision, EditAnywhere)
+	UPROPERTY(Category=Collision, EditAnywhere)
 	TEnumAsByte<ESpriteCollisionMode::Type> SpriteCollisionDomain;
 
 	// Baked physics data.
 	UPROPERTY()
 	class UBodySetup* BodySetup;
 
+public:
+#if WITH_EDITORONLY_DATA
+	/** Importing data and options used for this tile map */
+	UPROPERTY(Category=ImportSettings, VisibleAnywhere, EditInline)
+	class UAssetImportData* AssetImportData;
+#endif
+
 public:
 	// UObject interface
 #if WITH_EDITOR
diff --git a/Engine/Plugins/Experimental/Paper2D/Source/Paper2D/Classes/PaperTileSet.h b/Engine/Plugins/Experimental/Paper2D/Source/Paper2D/Classes/PaperTileSet.h
--- a/Engine/Plugins/Experimental/Paper2D/Source/Paper2D/Classes/PaperTileSet.h
+++ b/Engine/Plugins/Experimental/Paper2D/Source/Paper2D/Classes/PaperTileSet.h
@@ -11,15 +11,30 @@ class PAPER2D_API UPaperTileSet : public UDataAsset //@TODO: Just to make it eas
 {
 	GENERATED_UCLASS_BODY()
 
+	// The width of a single tile (in pixels)
 	UPROPERTY(Category=TileSet, EditAnywhere, meta=(UIMin=1, ClampMin=1))
 	int32 TileWidth;
 
+	// The height of a single tile (in pixels)
 	UPROPERTY(Category=TileSet, EditAnywhere, meta=(UIMin=1, ClampMin=1))
 	int32 TileHeight;
 
+	// The tile sheet texture associated with this tile set
 	UPROPERTY(Category=TileSet, EditAnywhere)
 	UTexture2D* TileSheet;
 
+	// The amount of padding around the perimeter of the tile sheet (in pixels)
+	UPROPERTY(Category=TileSet, EditAnywhere, meta=(UIMin=0, ClampMin=0))
+	int32 Margin;
+
+	// The amount of padding between tiles in the tile sheet (in pixels)
+	UPROPERTY(Category = TileSet, EditAnywhere, meta=(UIMin=0, ClampMin=0))
+	int32 Spacing;
+
+	// The drawing offset for tiles from this set (in pixels)
+	UPROPERTY(Category=TileSet, EditAnywhere, meta=(UIMin=0, ClampMin=0))
+	FIntPoint DrawingOffset;
+
 public:
 	int32 GetTileCount() const;
 	int32 GetTileCountX() const;
diff --git a/Engine/Plugins/Experimental/Paper2D/Source/Paper2DEditor/Classes/TileMapAssetImportData.h b/Engine/Plugins/Experimental/Paper2D/Source/Paper2DEditor/Classes/TileMapAssetImportData.h
new file mode 100644
--- /dev/null
+++ b/Engine/Plugins/Experimental/Paper2D/Source/Paper2DEditor/Classes/TileMapAssetImportData.h
@@ -0,0 +1,32 @@
+// Copyright 1998-2014 Epic Games, Inc. All Rights Reserved.
+
+#include "TileMapAssetImportData.generated.h"
+
+USTRUCT()
+struct FTileSetImportMapping
+{
+	GENERATED_USTRUCT_BODY()
+
+	UPROPERTY()
+	FString SourceName;
+
+	UPROPERTY()
+	TWeakObjectPtr<class UPaperTileSet> ImportedTileSet;
+
+	UPROPERTY()
+	TWeakObjectPtr<class UTexture> ImportedTexture;
+};
+
+/**
+ * Base class for import data and options used when importing a tile map
+ */
+UCLASS()
+class PAPER2DEDITOR_API UTileMapAssetImportData : public UAssetImportData
+{
+	GENERATED_UCLASS_BODY()
+
+	UPROPERTY()
+	TArray<FTileSetImportMapping> TileSetMap;
+
+	static UTileMapAssetImportData* GetImportDataForTileMap(class UPaperTileMap* TileMap);
+};
diff --git a/Engine/Plugins/Experimental/Paper2D/Source/Paper2DEditor/Private/PaperJSONHelpers.cpp b/Engine/Plugins/Experimental/Paper2D/Source/Paper2DEditor/Private/PaperJSONHelpers.cpp
new file mode 100644
--- /dev/null
+++ b/Engine/Plugins/Experimental/Paper2D/Source/Paper2DEditor/Private/PaperJSONHelpers.cpp
@@ -0,0 +1,135 @@
+// Copyright 1998-2014 Epic Games, Inc. All Rights Reserved.
+
+#include "Paper2DEditorPrivatePCH.h"
+#include "PaperJSONHelpers.h"
+#include "Json.h"
+
+//////////////////////////////////////////////////////////////////////////
+// FPaperJSONHelpers
+
+FString FPaperJSONHelpers::ReadString(TSharedPtr<class FJsonObject> Item, const FString& Key, const FString& DefaultValue)
+{
+	const TSharedPtr<FJsonValue>& ValuePtr = Item->GetField<EJson::String>(Key);
+	return ValuePtr.IsValid() ? ValuePtr->AsString() : DefaultValue;
+}
+
+TSharedPtr<class FJsonObject> FPaperJSONHelpers::ReadObject(TSharedPtr<class FJsonObject> Item, const FString& Key)
+{
+	if (Item->HasTypedField<EJson::Object>(Key))
+	{
+		return Item->GetObjectField(Key);
+	}
+	else
+	{
+		return NULL;
+	}
+}
+
+bool FPaperJSONHelpers::ReadBoolean(const TSharedPtr<class FJsonObject> Item, const FString& Key, bool bDefaultIfMissing)
+{
+	if (Item->HasTypedField<EJson::Boolean>(Key))
+	{
+		return Item->GetBoolField(Key);
+	}
+	else
+	{
+		return bDefaultIfMissing;
+	}
+}
+
+bool FPaperJSONHelpers::ReadFloatNoDefault(const TSharedPtr<class FJsonObject> Item, const FString& Key, float& Out_Value)
+{
+	if (Item->HasTypedField<EJson::Number>(Key))
+	{
+		Out_Value = Item->GetNumberField(Key);
+		return true;
+	}
+	else
+	{
+		Out_Value = 0.0f;
+		return false;
+	}
+}
+
+// Returns true if the field named Key is a Number (truncating it to an integer), populating Out_Value, or false if it is missing or the wrong type
+bool FPaperJSONHelpers::ReadIntegerNoDefault(const TSharedPtr<class FJsonObject> Item, const FString& Key, int32& Out_Value)
+{
+	if (Item->TryGetNumberField(Key, Out_Value))
+	{
+		return true;
+	}
+	else
+	{
+		Out_Value = 0;
+		return false;
+	}
+}
+
+bool FPaperJSONHelpers::ReadRectangle(const TSharedPtr<class FJsonObject> Item, const FString& Key, FVector2D& Out_XY, FVector2D& Out_WH)
+{
+	const TSharedPtr<FJsonObject> Struct = ReadObject(Item, Key);
+	if (Struct.IsValid())
+	{
+		if (ReadFloatNoDefault(Struct, TEXT("x"), /*out*/ Out_XY.X) &&
+			ReadFloatNoDefault(Struct, TEXT("y"), /*out*/ Out_XY.Y) &&
+			ReadFloatNoDefault(Struct, TEXT("w"), /*out*/ Out_WH.X) &&
+			ReadFloatNoDefault(Struct, TEXT("h"), /*out*/ Out_WH.Y))
+		{
+			return true;
+		}
+		else
+		{
+			Out_XY = FVector2D::ZeroVector;
+			Out_WH = FVector2D::ZeroVector;
+			return false;
+		}
+	}
+	else
+	{
+		return false;
+	}
+}
+
+bool FPaperJSONHelpers::ReadSize(const TSharedPtr<class FJsonObject> Item, const FString& Key, FVector2D& Out_WH)
+{
+	const TSharedPtr<FJsonObject> Struct = ReadObject(Item, Key);
+	if (Struct.IsValid())
+	{
+		if (ReadFloatNoDefault(Struct, TEXT("w"), /*out*/ Out_WH.X) &&
+			ReadFloatNoDefault(Struct, TEXT("h"), /*out*/ Out_WH.Y))
+		{
+			return true;
+		}
+		else
+		{
+			Out_WH = FVector2D::ZeroVector;
+			return false;
+		}
+	}
+	else
+	{
+		return false;
+	}
+}
+
+bool FPaperJSONHelpers::ReadIntPoint(const TSharedPtr<class FJsonObject> Item, const FString& Key, FIntPoint& Out_XY)
+{
+	const TSharedPtr<FJsonObject> Struct = ReadObject(Item, Key);
+	if (Struct.IsValid())
+	{
+		if (ReadIntegerNoDefault(Struct, TEXT("x"), /*out*/ Out_XY.X) &&
+			ReadIntegerNoDefault(Struct, TEXT("y"), /*out*/ Out_XY.Y))
+		{
+			return true;
+		}
+		else
+		{
+			Out_XY = FIntPoint::ZeroValue;
+			return false;
+		}
+	}
+	else
+	{
+		return false;
+	}
+}
diff --git a/Engine/Plugins/Experimental/Paper2D/Source/Paper2DEditor/Private/TileMapEditing/TileMapAssetImportData.cpp b/Engine/Plugins/Experimental/Paper2D/Source/Paper2DEditor/Private/TileMapEditing/TileMapAssetImportData.cpp
new file mode 100644
--- /dev/null
+++ b/Engine/Plugins/Experimental/Paper2D/Source/Paper2DEditor/Private/TileMapEditing/TileMapAssetImportData.cpp
@@ -0,0 +1,34 @@
+// Copyright 1998-2014 Epic Games, Inc. All Rights Reserved.
+
+#include "Paper2DEditorPrivatePCH.h"
+#include "AssetData.h"
+
+//////////////////////////////////////////////////////////////////////////
+// UTileMapAssetImportData
+
+UTileMapAssetImportData::UTileMapAssetImportData(const FPostConstructInitializeProperties& PCIP)
+	: Super(PCIP)
+{
+}
+
+UTileMapAssetImportData* UTileMapAssetImportData::GetImportDataForTileMap(UPaperTileMap* TileMap/*, UFbxSkeletalMeshImportData* TemplateForCreation*/)
+{
+	check(TileMap);
+
+	UTileMapAssetImportData* ImportData = Cast<UTileMapAssetImportData>(TileMap->AssetImportData);
+	if (ImportData == nullptr)
+	{
+		ImportData = ConstructObject<UTileMapAssetImportData>(UTileMapAssetImportData::StaticClass(), TileMap, NAME_None, RF_NoFlags/*, TemplateForCreation*/);
+
+		// Try to preserve the source file path if possible
+		if (TileMap->AssetImportData != nullptr)
+		{
+			ImportData->SourceFilePath = TileMap->AssetImportData->SourceFilePath;
+			ImportData->SourceFileTimestamp = TileMap->AssetImportData->SourceFileTimestamp;
+		}
+
+		TileMap->AssetImportData = ImportData;
+	}
+
+	return ImportData;
+}
diff --git a/Engine/Plugins/Experimental/Paper2D/Source/Paper2DEditor/Public/PaperJSONHelpers.h b/Engine/Plugins/Experimental/Paper2D/Source/Paper2DEditor/Public/PaperJSONHelpers.h
new file mode 100644
--- /dev/null
+++ b/Engine/Plugins/Experimental/Paper2D/Source/Paper2DEditor/Public/PaperJSONHelpers.h
@@ -0,0 +1,38 @@
+// Copyright 1998-2014 Epic Games, Inc. All Rights Reserved.
+
+#pragma once
+
+//////////////////////////////////////////////////////////////////////////
+// FPaperJSONHelpers
+
+class PAPER2DEDITOR_API FPaperJSONHelpers
+{
+public:
+	// returns the value of Key in Item, or DefaultValue if the Key is missing or the wrong type
+	static FString ReadString(TSharedPtr<class FJsonObject> Item, const FString& Key, const FString& DefaultValue);
+
+	// Returns the object named Key or NULL if it is missing or the wrong type
+	static TSharedPtr<class FJsonObject> ReadObject(TSharedPtr<class FJsonObject> Item, const FString& Key);
+
+	// Returns the bool named Key or bDefaultIfMissing if it is missing or the wrong type (note: no way to determine errors!)
+	static bool ReadBoolean(const TSharedPtr<class FJsonObject> Item, const FString& Key, bool bDefaultIfMissing);
+
+	// Returns true if the field named Key is a Number, populating Out_Value, or false if it is missing or the wrong type (Out_Value is set to 0.0f on failure)
+	static bool ReadFloatNoDefault(const TSharedPtr<class FJsonObject> Item, const FString& Key, float& Out_Value);
+
+	// Returns true if the field named Key is a Number (truncating it to an integer), populating Out_Value, or false if it is missing or the wrong type (Out_Value is set to 0 on failure)
+	static bool ReadIntegerNoDefault(const TSharedPtr<class FJsonObject> Item, const FString& Key, int32& Out_Value);
+
+	// Returns true if the object named Key is a struct containing four floats (x,y,w,h), populating XY and WH with the values)
+	static bool ReadRectangle(const TSharedPtr<class FJsonObject> Item, const FString& Key, FVector2D& Out_XY, FVector2D& Out_WH);
+
+	// Returns true if the object named Key is a struct containing two floats (w,h), populating WH with the values)
+	static bool ReadSize(const TSharedPtr<class FJsonObject> Item, const FString& Key, FVector2D& Out_WH);
+
+	// Returns true if the object named Key is a struct containing two floats (x,y), populating XY with the values)
+	static bool ReadIntPoint(const TSharedPtr<class FJsonObject> Item, const FString& Key, FIntPoint& Out_XY);
+
+private:
+	FPaperJSONHelpers() {}
+};
+
diff --git a/Engine/Plugins/Experimental/Paper2D/Source/PaperJsonImporter/Classes/PaperJsonImporterFactory.h b/Engine/Plugins/Experimental/Paper2D/Source/PaperJsonImporter/Classes/PaperJsonImporterFactory.h
--- a/Engine/Plugins/Experimental/Paper2D/Source/PaperJsonImporter/Classes/PaperJsonImporterFactory.h
+++ b/Engine/Plugins/Experimental/Paper2D/Source/PaperJsonImporter/Classes/PaperJsonImporterFactory.h
@@ -17,23 +17,5 @@ class UPaperJsonImporterFactory : public UFactory
 
 protected:
 	TSharedPtr<class FJsonObject> ParseJSON(const FString& FileContents, const FString& NameForErrors);
-
-	// returns the value of Key in Item, or DefaultValue if the Key is missing or the wrong type
-	FString ReadString(TSharedPtr<class FJsonObject> Item, const FString& Key, const FString& DefaultValue);
-
-	// Returns the object named Key or NULL if it is missing or the wrong type
-	TSharedPtr<class FJsonObject> ReadObject(TSharedPtr<class FJsonObject> Item, const FString& Key);
-
-	// Returns the bool named Key or bDefaultIfMissing if it is missing or the wrong type (note: no way to determine errors!)
-	bool ReadBoolean(const TSharedPtr<class FJsonObject> Item, const FString& Key, bool bDefaultIfMissing);
-
-	// Returns true if the field named Key is a Number, populating Out_Value, or false if it is missing or the wrong type
-	bool ReadFloatNoDefault(const TSharedPtr<class FJsonObject> Item, const FString& Key, float& Out_Value);
-
-	// Returns true if the object named Key is a struct containing four floats (x,y,w,h), populating XY and WH with the values)
-	bool ReadRectangle(const TSharedPtr<class FJsonObject> Item, const FString& Key, FVector2D& Out_XY, FVector2D& Out_WH);
-
-	// Returns true if the object named Key is a struct containing two floats (w,h), populating WH with the values)
-	bool ReadSize(const TSharedPtr<class FJsonObject> Item, const FString& Key, FVector2D& Out_WH);
 };
 
diff --git a/Engine/Plugins/Experimental/Paper2D/Source/PaperJsonImporter/PaperJsonImporter.Build.cs b/Engine/Plugins/Experimental/Paper2D/Source/PaperJsonImporter/PaperJsonImporter.Build.cs
--- a/Engine/Plugins/Experimental/Paper2D/Source/PaperJsonImporter/PaperJsonImporter.Build.cs
+++ b/Engine/Plugins/Experimental/Paper2D/Source/PaperJsonImporter/PaperJsonImporter.Build.cs
@@ -12,7 +12,8 @@ public class PaperJsonImporter : ModuleRules
 				"CoreUObject",
 				"Engine",
 				"Paper2D",
-				"UnrealEd"
+				"UnrealEd",
+				"Paper2DEditor"
 			});
 
 		PrivateIncludePathModuleNames.AddRange(
diff --git a/Engine/Plugins/Experimental/Paper2D/Source/PaperJsonImporter/Private/PaperJsonImporterFactory.cpp b/Engine/Plugins/Experimental/Paper2D/Source/PaperJsonImporter/Private/PaperJsonImporterFactory.cpp
--- a/Engine/Plugins/Experimental/Paper2D/Source/PaperJsonImporter/Private/PaperJsonImporterFactory.cpp
+++ b/Engine/Plugins/Experimental/Paper2D/Source/PaperJsonImporter/Private/PaperJsonImporterFactory.cpp
@@ -3,6 +3,7 @@
 #include "PaperJsonImporterPrivatePCH.h"
 #include "Paper2DClasses.h"
 #include "Json.h"
+#include "PaperJSONHelpers.h"
 #include "AssetToolsModule.h"
 #include "AssetRegistryModule.h"
 #include "PackageTools.h"
@@ -38,7 +39,7 @@ UPaperJsonImporterFactory::UPaperJsonImporterFactory(const FPostConstructInitial
 	bEditorImport = true;
 	bText = true;
 
-	Formats.Add(TEXT("json;JSON file"));
+	Formats.Add(TEXT("json;Spritesheet JSON file"));
 }
 
 FText UPaperJsonImporterFactory::GetToolTip() const
@@ -78,7 +79,7 @@ UObject* UPaperJsonImporterFactory::FactoryCreateText(UClass* InClass, UObject*
 		UTexture2D* ImportedTexture = NULL;
 
 		// Validate the meta block
-		TSharedPtr<FJsonObject> MetaBlock = ReadObject(SpriteDescriptorObject, TEXT("meta"));
+		TSharedPtr<FJsonObject> MetaBlock = FPaperJSONHelpers::ReadObject(SpriteDescriptorObject, TEXT("meta"));
 		if (MetaBlock.IsValid())
 		{
 			// Example contents:
@@ -89,8 +90,8 @@ UObject* UPaperJsonImporterFactory::FactoryCreateText(UClass* InClass, UObject*
 			//   "size": {"w":2048,"h":2048},    (ignored)
 			//   "scale": "1"                    (ignored)
 
-			const FString AppName = ReadString(MetaBlock, TEXT("app"), TEXT(""));
-			const FString Image = ReadString(MetaBlock, TEXT("image"), TEXT(""));
+			const FString AppName = FPaperJSONHelpers::ReadString(MetaBlock, TEXT("app"), TEXT(""));
+			const FString Image = FPaperJSONHelpers::ReadString(MetaBlock, TEXT("image"), TEXT(""));
 			
  			if (bLoadedSuccessfully)
 			{
@@ -163,7 +164,7 @@ UObject* UPaperJsonImporterFactory::FactoryCreateText(UClass* InClass, UObject*
 		// Parse the frames
 		if (bLoadedSuccessfully)
 		{
-			TSharedPtr<FJsonObject> FramesBlock = ReadObject(SpriteDescriptorObject, TEXT("frames"));
+			TSharedPtr<FJsonObject> FramesBlock = FPaperJSONHelpers::ReadObject(SpriteDescriptorObject, TEXT("frames"));
 			if (FramesBlock.IsValid())
 			{
 				GWarn->BeginSlowTask(NSLOCTEXT("Paper2D", "PaperJsonImporterFactory_ParsingSprites", "Parsing Sprite Frame"), true, true);
@@ -201,8 +202,8 @@ UObject* UPaperJsonImporterFactory::FactoryCreateText(UClass* InClass, UObject*
 						//   "spriteSourceSize": {"x":0,"y":11,"w":216,"h":240},
 						//   "sourceSize": {"w":216,"h":240}
 
-						Frame.bRotated = ReadBoolean(FrameData, TEXT("rotated"), false);
-						Frame.bTrimmed = ReadBoolean(FrameData, TEXT("trimmed"), false);
+						Frame.bRotated = FPaperJSONHelpers::ReadBoolean(FrameData, TEXT("rotated"), false);
+						Frame.bTrimmed = FPaperJSONHelpers::ReadBoolean(FrameData, TEXT("trimmed"), false);
 
 						if (Frame.bRotated || !Frame.bTrimmed)
 						{
@@ -210,11 +211,11 @@ UObject* UPaperJsonImporterFactory::FactoryCreateText(UClass* InClass, UObject*
 							UE_LOG(LogPaperJsonImporter, Warning, TEXT("Frame %s is either rotated or not trimmed, which may not be handled correctly"), *Frame.FrameName.ToString());
 						}
 
-						bReadFrameSuccessfully = bReadFrameSuccessfully && ReadRectangle(FrameData, "frame", /*out*/ Frame.SpritePosInSheet, /*out*/ Frame.SpriteSizeInSheet);
+						bReadFrameSuccessfully = bReadFrameSuccessfully && FPaperJSONHelpers::ReadRectangle(FrameData, "frame", /*out*/ Frame.SpritePosInSheet, /*out*/ Frame.SpriteSizeInSheet);
 
 						if (Frame.bTrimmed)
 						{
-							bReadFrameSuccessfully = bReadFrameSuccessfully && ReadRectangle(FrameData, "spriteSourceSize", /*out*/ Frame.SpriteSourcePos, /*out*/ Frame.SpriteSourceSize);
+							bReadFrameSuccessfully = bReadFrameSuccessfully && FPaperJSONHelpers::ReadRectangle(FrameData, "spriteSourceSize", /*out*/ Frame.SpriteSourcePos, /*out*/ Frame.SpriteSourceSize);
 						}
 						else
 						{
@@ -347,99 +348,3 @@ TSharedPtr<FJsonObject> UPaperJsonImporterFactory::ParseJSON(const FString& File
 		return NULL;
 	}
 }
-
-// returns the value of Key in Item, or DefaultValue if the Key is missing or the wrong type
-FString UPaperJsonImporterFactory::ReadString(TSharedPtr<class FJsonObject> Item, const FString& Key, const FString& DefaultValue)
-{
-	const TSharedPtr<FJsonValue>& ValuePtr = Item->GetField<EJson::String>(Key);
-	return ValuePtr.IsValid() ? ValuePtr->AsString() : DefaultValue;
-}
-
-// Returns the object named Key or NULL if it is missing or the wrong type
-TSharedPtr<class FJsonObject> UPaperJsonImporterFactory::ReadObject(TSharedPtr<class FJsonObject> Item, const FString& Key)
-{
-	if (Item->HasTypedField<EJson::Object>(Key))
-	{
-		return Item->GetObjectField(Key);
-	}
-	else
-	{
-		return NULL;
-	}
-}
-
-// Returns the bool named Key or bDefaultIfMissing if it is missing or the wrong type (note: no way to determine errors!)
-bool UPaperJsonImporterFactory::ReadBoolean(const TSharedPtr<class FJsonObject> Item, const FString& Key, bool bDefaultIfMissing)
-{
-	if (Item->HasTypedField<EJson::Boolean>(Key))
-	{
-		return Item->GetBoolField(Key);
-	}
-	else
-	{
-		return bDefaultIfMissing;
-	}
-}
-
-bool UPaperJsonImporterFactory::ReadFloatNoDefault(const TSharedPtr<class FJsonObject> Item, const FString& Key, float& Out_Value)
-{
-	if (Item->HasTypedField<EJson::Number>(Key))
-	{
-		Out_Value = Item->GetNumberField(Key);
-		return true;
-	}
-	else
-	{
-		Out_Value = 0.0f;
-		return false;
-	}
-}
-
-// Returns true if the object named Key is a struct containing four floats (x,y,w,h), populating XY and WH with the values)
-bool UPaperJsonImporterFactory::ReadRectangle(const TSharedPtr<class FJsonObject> Item, const FString& Key, FVector2D& Out_XY, FVector2D& Out_WH)
-{
-	const TSharedPtr<FJsonObject> Struct = ReadObject(Item, Key);
-	if (Struct.IsValid())
-	{
-		if (ReadFloatNoDefault(Struct, TEXT("x"), /*out*/ Out_XY.X) &&
-			ReadFloatNoDefault(Struct, TEXT("y"), /*out*/ Out_XY.Y) &&
-			ReadFloatNoDefault(Struct, TEXT("w"), /*out*/ Out_WH.X) &&
-			ReadFloatNoDefault(Struct, TEXT("h"), /*out*/ Out_WH.Y))
-		{
-			return true;
-		}
-		else
-		{
-			Out_XY = FVector2D::ZeroVector;
-			Out_WH = FVector2D::ZeroVector;
-			return false;
-		}
-	}
-	else
-	{
-		return false;
-	}
-}
-
-// Returns true if the object named Key is a struct containing two floats (w,h), populating WH with the values)
-bool UPaperJsonImporterFactory::ReadSize(const TSharedPtr<class FJsonObject> Item, const FString& Key, FVector2D& Out_WH)
-{
-	const TSharedPtr<FJsonObject> Struct = ReadObject(Item, Key);
-	if (Struct.IsValid())
-	{
-		if (ReadFloatNoDefault(Struct, TEXT("w"), /*out*/ Out_WH.X) &&
-			ReadFloatNoDefault(Struct, TEXT("h"), /*out*/ Out_WH.Y))
-		{
-			return true;
-		}
-		else
-		{
-			Out_WH = FVector2D::ZeroVector;
-			return false;
-		}
-	}
-	else
-	{
-		return false;
-	}
-}
diff --git a/Engine/Plugins/Experimental/Paper2D/Source/PaperTiledImporter/Classes/PaperTiledImporterFactory.h b/Engine/Plugins/Experimental/Paper2D/Source/PaperTiledImporter/Classes/PaperTiledImporterFactory.h
new file mode 100644
--- /dev/null
+++ b/Engine/Plugins/Experimental/Paper2D/Source/PaperTiledImporter/Classes/PaperTiledImporterFactory.h
@@ -0,0 +1,127 @@
+// Copyright 1998-2014 Epic Games, Inc. All Rights Reserved.
+
+#pragma once
+
+#include "PaperTileLayer.h"
+#include "PaperTiledImporterFactory.generated.h"
+
+namespace ETiledOrientation
+{
+	enum Type
+	{
+		Unknown,
+		Orthogonal,
+		Isometric
+	};
+}
+
+// Imports a tile map (and associated textures & tile sets) exported from Tiled (http://www.mapeditor.org/)
+UCLASS()
+class UPaperTiledImporterFactory : public UFactory, public FReimportHandler
+{
+	GENERATED_UCLASS_BODY()
+
+	// UFactory interface
+	virtual FText GetToolTip() const override;
+	virtual bool FactoryCanImport(const FString& Filename) override;
+	virtual UObject* FactoryCreateText(UClass* InClass, UObject* InParent, FName InName, EObjectFlags Flags, UObject* Context, const TCHAR* Type, const TCHAR*& Buffer, const TCHAR* BufferEnd, FFeedbackContext* Warn) override;
+	// End of UFactory interface
+
+	// FReimportHandler interface
+	virtual bool CanReimport(UObject* Obj, TArray<FString>& OutFilenames) override;
+	virtual void SetReimportPaths(UObject* Obj, const TArray<FString>& NewReimportPaths) override;
+	virtual EReimportResult::Type Reimport(UObject* Obj) override;
+	// End of FReimportHandler interface
+
+protected:
+	TSharedPtr<FJsonObject> ParseJSON(const FString& FileContents, const FString& NameForErrors, bool bSilent = false);
+
+	void ParseGlobalInfoFromJSON(TSharedPtr<FJsonObject> Tree, struct FTileMapFromTiled& OutParsedInfo, const FString& NameForErrors, bool bSilent = false);
+
+	static UObject* CreateNewAsset(UClass* AssetClass, const FString& TargetPath, const FString& DesiredName, EObjectFlags Flags);
+	static UTexture2D* ImportTexture(const FString& SourceFilename, const FString& TargetSubPath);
+};
+
+
+struct FTileLayerFromTiled
+{
+	// Name of the layer
+	FString Name;
+
+	// Array of tiles
+	TArray<int32> TileIndices;
+
+	// Width and height in tiles
+	int32 Width;
+	int32 Height;
+
+	// Saved layer opacity
+	float Opacity;
+
+	// Is the layer currently visible
+	bool bVisible;
+
+	// 'tilelayer' ??? - not sure what this is
+	FString Type;
+
+	// Offset
+	int32 OffsetX;
+	int32 OffsetY;
+
+
+	FTileLayerFromTiled();
+
+	void ParseFromJSON(TSharedPtr<FJsonObject> Tree, const FString& NameForErrors, bool bSilent = false);
+
+	bool IsValid() const;
+};
+
+struct FTileSetFromTiled
+{
+	int32 FirstGID;
+	FString Name;
+
+	FString ImagePath;
+	int32 ImageWidth;
+	int32 ImageHeight;
+
+	//Properties
+	//TileProperties
+	int32 TileOffsetX;
+	int32 TileOffsetY;
+
+	int32 Margin;
+	int32 Spacing;
+
+	int32 TileWidth;
+	int32 TileHeight;
+
+	FTileSetFromTiled();
+
+	void ParseFromJSON(TSharedPtr<FJsonObject> Tree, const FString& NameForErrors, bool bSilent = false);
+
+	bool IsValid() const;
+};
+
+
+struct FTileMapFromTiled
+{
+	int32 FileVersion;
+	int32 Width;
+	int32 Height;
+	int32 TileWidth;
+	int32 TileHeight;
+
+	ETiledOrientation::Type Orientation;
+
+	TArray<FTileSetFromTiled> TileSets;
+	TArray<UPaperTileSet*> CreatedTileSetAssets;
+	TArray<FTileLayerFromTiled> Layers;
+	//map of properties
+
+	FPaperTileInfo ConvertTileGIDToPaper2D(int32 GID) const;
+
+	FTileMapFromTiled();
+	bool IsValid() const;
+};
+
diff --git a/Engine/Plugins/Experimental/Paper2D/Source/PaperTiledImporter/PaperTiledImporter.Build.cs b/Engine/Plugins/Experimental/Paper2D/Source/PaperTiledImporter/PaperTiledImporter.Build.cs
new file mode 100644
--- /dev/null
+++ b/Engine/Plugins/Experimental/Paper2D/Source/PaperTiledImporter/PaperTiledImporter.Build.cs
@@ -0,0 +1,31 @@
+// Copyright 1998-2014 Epic Games, Inc. All Rights Reserved.
+
+using UnrealBuildTool;
+
+public class PaperTiledImporter : ModuleRules
+{
+	public PaperTiledImporter(TargetInfo Target)
+	{
+		PrivateDependencyModuleNames.AddRange(
+			new string[] {
+				"Core",
+				"CoreUObject",
+				"Engine",
+				"Paper2D",
+				"UnrealEd",
+				"Paper2DEditor"
+			});
+
+		PrivateIncludePathModuleNames.AddRange(
+			new string[] {
+				"AssetTools",
+				"AssetRegistry"
+			});
+
+		DynamicallyLoadedModuleNames.AddRange(
+			new string[] {
+				"AssetTools",
+				"AssetRegistry"
+			});
+	}
+}
diff --git a/Engine/Plugins/Experimental/Paper2D/Source/PaperTiledImporter/Private/PaperTiledImporterFactory.cpp b/Engine/Plugins/Experimental/Paper2D/Source/PaperTiledImporter/Private/PaperTiledImporterFactory.cpp
new file mode 100644
--- /dev/null
+++ b/Engine/Plugins/Experimental/Paper2D/Source/PaperTiledImporter/Private/PaperTiledImporterFactory.cpp
@@ -0,0 +1,644 @@
+// Copyright 1998-2014 Epic Games, Inc. All Rights Reserved.
+
+#include "PaperTiledImporterPrivatePCH.h"
+#include "Paper2DClasses.h"
+#include "Paper2DEditorClasses.h"
+#include "Json.h"
+#include "PaperJSONHelpers.h"
+#include "AssetToolsModule.h"
+#include "AssetRegistryModule.h"
+#include "PackageTools.h"
+
+#define LOCTEXT_NAMESPACE "Paper2D"
+#define TILED_IMPORT_ERROR(FormatString, ...) \
+	if (!bSilent) { UE_LOG(LogPaperTiledImporter, Warning, FormatString, __VA_ARGS__); }
+#define TILED_IMPORT_WARNING(FormatString, ...) \
+	if (!bSilent) { UE_LOG(LogPaperTiledImporter, Warning, FormatString, __VA_ARGS__); }
+
+//////////////////////////////////////////////////////////////////////////
+// FRequiredIntField
+
+struct FRequiredIntField
+{
+	int32& Value;
+	const FString Key;
+	int32 MinValue;
+
+	FRequiredIntField(int32& InValue, const FString& InKey)
+		: Value(InValue)
+		, Key(InKey)
+		, MinValue(1)
+	{
+	}
+
+	FRequiredIntField(int32& InValue, const FString& InKey, int32 InMinValue)
+		: Value(InValue)
+		, Key(InKey)
+		, MinValue(InMinValue)
+	{
+	}
+};
+
+bool ParseIntegerFields(FRequiredIntField* IntFieldArray, int32 ArrayCount, TSharedPtr<FJsonObject> Tree, const FString& NameForErrors, bool bSilent)
+{
+	bool bSuccessfullyParsed = true;
+
+	for (int32 ArrayIndex = 0; ArrayIndex < ArrayCount; ++ArrayIndex)
+	{
+		const FRequiredIntField& Field = IntFieldArray[ArrayIndex];
+
+		if (!Tree->TryGetNumberField(Field.Key, /*out*/ Field.Value))
+		{
+			TILED_IMPORT_ERROR(TEXT("Failed to parse '%s'.  Missing '%s' property"), *NameForErrors, *Field.Key);
+			bSuccessfullyParsed = false;
+			Field.Value = 0;
+		}
+		else
+		{
+			if (Field.Value < Field.MinValue)
+			{
+				TILED_IMPORT_ERROR(TEXT("Failed to parse '%s'.  Invalid value for '%s' (%d but must be at least %d)"), *NameForErrors, *Field.Key, Field.Value, Field.MinValue);
+				bSuccessfullyParsed = false;
+				Field.Value = Field.MinValue;
+			}
+		}
+	}
+
+	return bSuccessfullyParsed;
+}
+
+
+//////////////////////////////////////////////////////////////////////////
+// UPaperTiledImporterFactory
+
+UPaperTiledImporterFactory::UPaperTiledImporterFactory(const FPostConstructInitializeProperties& PCIP)
+	: Super(PCIP)
+{
+	bCreateNew = false;
+	//bEditAfterNew = true;
+	SupportedClass = UPaperTileMap::StaticClass();
+
+	bEditorImport = true;
+	bText = true;
+
+	Formats.Add(TEXT("json;Tiled JSON file"));
+}
+
+FText UPaperTiledImporterFactory::GetToolTip() const
+{
+	return LOCTEXT("PaperTiledImporterFactoryDescription", "Tile maps exported from Tiled");
+}
+
+bool UPaperTiledImporterFactory::FactoryCanImport(const FString& Filename)
+{
+	if (GetDefault<UPaperRuntimeSettings>()->bEnableTileMapEditing)
+	{
+		FString FileContent;
+		if (FFileHelper::LoadFileToString(/*out*/ FileContent, *Filename))
+		{
+			TSharedPtr<FJsonObject> DescriptorObject = ParseJSON(FileContent, FString(), /*bSilent=*/ true);
+			if (DescriptorObject.IsValid())
+			{
+				FTileMapFromTiled GlobalInfo;
+				ParseGlobalInfoFromJSON(DescriptorObject, /*out*/ GlobalInfo, FString(), /*bSilent=*/ true);
+
+				return GlobalInfo.IsValid();
+			}
+		}
+	}
+
+	return false;
+}
+
+UObject* UPaperTiledImporterFactory::FactoryCreateText(UClass* InClass, UObject* InParent, FName InName, EObjectFlags Flags, UObject* Context, const TCHAR* Type, const TCHAR*& Buffer, const TCHAR* BufferEnd, FFeedbackContext* Warn)
+{
+	Flags |= RF_Transactional;
+
+	FEditorDelegates::OnAssetPreImport.Broadcast(this, InClass, InParent, InName, Type);
+
+ 	FAssetToolsModule& AssetToolsModule = FModuleManager::GetModuleChecked<FAssetToolsModule>("AssetTools");
+ 
+ 	bool bLoadedSuccessfully = true;
+ 
+ 	const FString CurrentFilename = UFactory::GetCurrentFilename();
+ 	FString CurrentSourcePath;
+ 	FString FilenameNoExtension;
+ 	FString UnusedExtension;
+ 	FPaths::Split(CurrentFilename, CurrentSourcePath, FilenameNoExtension, UnusedExtension);
+ 
+ 	const FString LongPackagePath = FPackageName::GetLongPackagePath(InParent->GetOutermost()->GetPathName());
+ 
+ 	const FString NameForErrors(InName.ToString());
+ 	const FString FileContent(BufferEnd - Buffer, Buffer);
+ 	TSharedPtr<FJsonObject> DescriptorObject = ParseJSON(FileContent, NameForErrors);
+// 
+// 	TSet<FName> FrameNames;
+// 	TArray<FSpriteFrame> ParsedFrames;
+// 
+ 	UPaperTileMap* Result = nullptr;
+ 
+	FTileMapFromTiled GlobalInfo;
+	if (DescriptorObject.IsValid())
+	{
+		ParseGlobalInfoFromJSON(DescriptorObject, /*out*/ GlobalInfo, NameForErrors);
+	}
+
+	if (GlobalInfo.IsValid())
+	{
+		if (GlobalInfo.FileVersion != 1)
+		{
+			UE_LOG(LogPaperTiledImporter, Warning, TEXT("JSON exported from Tiled in file '%s' has an unknown version %d (expected version 1).  Parsing will continue but some things may not import correctly"), *NameForErrors, GlobalInfo.FileVersion);
+		}
+
+		// Load the tile sets
+		const TArray<TSharedPtr<FJsonValue>>* TileSetDescriptors;
+		if (DescriptorObject->TryGetArrayField(TEXT("tilesets"), /*out*/ TileSetDescriptors))
+		{
+			for (TSharedPtr<FJsonValue> TileSetDescriptor : *TileSetDescriptors)
+			{
+				FTileSetFromTiled& TileSet = *new (GlobalInfo.TileSets) FTileSetFromTiled();
+				TileSet.ParseFromJSON(TileSetDescriptor->AsObject(), NameForErrors);
+				bLoadedSuccessfully = bLoadedSuccessfully && TileSet.IsValid();
+			}
+		}
+		else
+		{
+			UE_LOG(LogPaperTiledImporter, Warning, TEXT("JSON exported from Tiled in file '%s' has no tile sets."), *NameForErrors);
+			bLoadedSuccessfully = false;
+		}
+
+		// Load the layers
+		const TArray<TSharedPtr<FJsonValue>>* LayerDescriptors;
+		if (DescriptorObject->TryGetArrayField(TEXT("layers"), /*out*/ LayerDescriptors))
+		{
+			for (TSharedPtr<FJsonValue> LayerDescriptor : *LayerDescriptors)
+			{
+				FTileLayerFromTiled& TileLayer = *new (GlobalInfo.Layers) FTileLayerFromTiled();
+				TileLayer.ParseFromJSON(LayerDescriptor->AsObject(), NameForErrors);
+				bLoadedSuccessfully = bLoadedSuccessfully && TileLayer.IsValid();
+			}
+		}
+		else
+		{
+			UE_LOG(LogPaperTiledImporter, Warning, TEXT("JSON exported from Tiled in file '%s' has no layers."), *NameForErrors);
+			bLoadedSuccessfully = false;
+		}
+
+		// Load the properties
+		//@TODO: What to do with Tiled properties?
+
+
+
+		// Create the new tile map asset and import basic/global data
+		Result = NewNamedObject<UPaperTileMap>(InParent, InName, Flags);
+
+		Result->Modify();
+		Result->MapWidth = GlobalInfo.Width;
+		Result->MapHeight = GlobalInfo.Height;
+		Result->TileWidth = GlobalInfo.TileWidth;
+		Result->TileHeight = GlobalInfo.TileHeight;
+		//@TODO: Need to handle Orientation
+
+		// Create the tile sets
+		for (const FTileSetFromTiled& TileSetData : GlobalInfo.TileSets)
+		{
+			if (TileSetData.IsValid())
+			{
+				const FString TargetTileSetPath = LongPackagePath;
+				const FString TargetTexturePath = LongPackagePath + TEXT("/Textures");
+
+				UPaperTileSet* TileSetAsset = CastChecked<UPaperTileSet>(CreateNewAsset(UPaperTileSet::StaticClass(), TargetTileSetPath, TileSetData.Name, Flags));
+				TileSetAsset->Modify();
+
+				TileSetAsset->TileWidth = TileSetData.TileWidth;
+				TileSetAsset->TileHeight = TileSetData.TileHeight;
+				TileSetAsset->Margin = TileSetData.Margin;
+				TileSetAsset->Spacing = TileSetData.Spacing;
+				TileSetAsset->DrawingOffset = FIntPoint(TileSetData.TileOffsetX, TileSetData.TileOffsetY);
+
+				//@TODO: Ignoring ImageWidth and ImageHeight, and Properties/TileProperties aren't even loaded yet
+
+				// Import the texture
+				const FString SourceImageFilename = FPaths::Combine(*CurrentSourcePath, *TileSetData.ImagePath);
+				TileSetAsset->TileSheet = ImportTexture(SourceImageFilename, TargetTexturePath);
+
+				if (TileSetAsset->TileSheet == nullptr)
+				{
+					UE_LOG(LogPaperTiledImporter, Warning, TEXT("Failed to import tile set image '%s' referenced from tile set '%s'."), *TileSetData.ImagePath, *TileSetData.Name);
+					bLoadedSuccessfully = false;
+				}
+
+				TileSetAsset->PostEditChange();
+				GlobalInfo.CreatedTileSetAssets.Add(TileSetAsset);
+			}
+			else
+			{
+				GlobalInfo.CreatedTileSetAssets.Add(nullptr);
+			}
+		}
+
+
+		// Create the layers
+		for (const FTileLayerFromTiled& LayerData : GlobalInfo.Layers)
+		{
+			if (LayerData.IsValid())
+			{
+				UPaperTileLayer* NewLayer = NewObject<UPaperTileLayer>(Result);
+				NewLayer->SetFlags(RF_Transactional);
+
+				NewLayer->LayerName = FText::FromString(LayerData.Name);
+				NewLayer->bHiddenInEditor = !LayerData.bVisible;
+				NewLayer->LayerOpacity = FMath::Clamp<float>(LayerData.Opacity, 0.0f, 1.0f);
+				//@TODO: No support for Type, OffsetX, or OffsetY
+
+				NewLayer->LayerWidth = LayerData.Width;
+				NewLayer->LayerHeight = LayerData.Height;
+				NewLayer->DestructiveAllocateMap(LayerData.Width, LayerData.Height);
+
+				int32 SourceIndex = 0;
+				for (int32 Y = 0; Y < LayerData.Height; ++Y)
+				{
+					for (int32 X = 0; X < LayerData.Width; ++X)
+					{
+						const int32 SourceTileGID = LayerData.TileIndices[SourceIndex++];
+						const FPaperTileInfo CellContents = GlobalInfo.ConvertTileGIDToPaper2D(SourceTileGID);
+						NewLayer->SetCell(X, Y, CellContents);
+					}
+				}
+
+				//@TODO: Handle TileIndices!
+
+				Result->TileLayers.Add(NewLayer);
+			}
+		}
+
+		Result->PostEditChange();
+	}
+ 	else
+ 	{
+ 		// Failed to parse the JSON
+ 		bLoadedSuccessfully = false;
+ 	}
+
+	if (Result != nullptr)
+	{
+		// Store the current file path and timestamp for re-import purposes
+		UTileMapAssetImportData* ImportData = UTileMapAssetImportData::GetImportDataForTileMap(Result);
+		ImportData->SourceFilePath = FReimportManager::SanitizeImportFilename(CurrentFilename, Result);
+		ImportData->SourceFileTimestamp = IFileManager::Get().GetTimeStamp(*CurrentFilename).ToString();
+
+		//GlobalInfo.TileSets
+	}
+
+	FEditorDelegates::OnAssetPostImport.Broadcast(this, Result);
+
+	return Result;
+}
+
+TSharedPtr<FJsonObject> UPaperTiledImporterFactory::ParseJSON(const FString& FileContents, const FString& NameForErrors, bool bSilent)
+{
+	// Load the file up (JSON format)
+	if (!FileContents.IsEmpty())
+	{
+		const TSharedRef< TJsonReader<> >& Reader = TJsonReaderFactory<>::Create(FileContents);
+
+		TSharedPtr<FJsonObject> DescriptorObject;
+		if (FJsonSerializer::Deserialize(Reader, /*out*/ DescriptorObject) && DescriptorObject.IsValid())
+		{
+			// File was loaded and deserialized OK!
+			return DescriptorObject;
+		}
+		else
+		{
+			if (!bSilent)
+			{
+				//@TODO: PAPER2D: How to correctly surface import errors to the user?
+				UE_LOG(LogPaperTiledImporter, Warning, TEXT("Failed to parse tile map JSON file '%s'.  Error: '%s'"), *NameForErrors, *Reader->GetErrorMessage());
+			}
+			return nullptr;
+		}
+	}
+	else
+	{
+		if (!bSilent)
+		{
+			UE_LOG(LogPaperTiledImporter, Warning, TEXT("Tile map JSON file '%s' was empty.  This tile map cannot be imported."), *NameForErrors);
+		}
+		return nullptr;
+	}
+}
+
+bool UPaperTiledImporterFactory::CanReimport(UObject* Obj, TArray<FString>& OutFilenames)
+{
+	if (UPaperTileMap* TileMap = Cast<UPaperTileMap>(Obj))
+	{
+		if (TileMap->AssetImportData != nullptr)
+		{
+			OutFilenames.Add(FReimportManager::ResolveImportFilename(TileMap->AssetImportData->SourceFilePath, TileMap));
+		}
+		else
+		{
+			OutFilenames.Add(TEXT(""));
+		}
+		return true;
+	}
+	return false;
+}
+
+void UPaperTiledImporterFactory::SetReimportPaths(UObject* Obj, const TArray<FString>& NewReimportPaths)
+{
+	if (UPaperTileMap* TileMap = Cast<UPaperTileMap>(Obj))
+	{
+		if (ensure(NewReimportPaths.Num() == 1))
+		{
+			UTileMapAssetImportData* ImportData = UTileMapAssetImportData::GetImportDataForTileMap(TileMap);
+
+			ImportData->SourceFilePath = FReimportManager::SanitizeImportFilename(NewReimportPaths[0], TileMap);
+		}
+	}
+}
+
+EReimportResult::Type UPaperTiledImporterFactory::Reimport(UObject* Obj)
+{
+	if (UPaperTileMap* TileMap = Cast<UPaperTileMap>(Obj))
+	{
+		//@TODO: Not implemented yet
+		ensureMsg(false, TEXT("Tile map reimport is not implemented yet"));
+	}
+	return EReimportResult::Failed;
+}
+
+
+UObject* UPaperTiledImporterFactory::CreateNewAsset(UClass* AssetClass, const FString& TargetPath, const FString& DesiredName, EObjectFlags Flags)
+{
+	FAssetToolsModule& AssetToolsModule = FModuleManager::GetModuleChecked<FAssetToolsModule>("AssetTools");
+
+	// Create a unique package name and asset name for the frame
+	const FString TentativePackagePath = PackageTools::SanitizePackageName(TargetPath + TEXT("/") + DesiredName);
+	FString DefaultSuffix;
+	FString AssetName;
+	FString PackageName;
+	AssetToolsModule.Get().CreateUniqueAssetName(TentativePackagePath, /*out*/ DefaultSuffix, /*out*/ PackageName, /*out*/ AssetName);
+
+	// Create a package for the asset
+	UObject* OuterForAsset = CreatePackage(nullptr, *PackageName);
+
+	// Create a frame in the package
+	UObject* NewAsset = ConstructObject<UObject>(AssetClass, OuterForAsset, *AssetName, Flags);
+	FAssetRegistryModule::AssetCreated(NewAsset);
+
+	return NewAsset;
+}
+
+UTexture2D* UPaperTiledImporterFactory::ImportTexture(const FString& SourceFilename, const FString& TargetSubPath)
+{
+	FAssetToolsModule& AssetToolsModule = FModuleManager::GetModuleChecked<FAssetToolsModule>("AssetTools");
+
+	TArray<FString> FileNames;
+	FileNames.Add(SourceFilename);
+
+	//@TODO: Avoid the first compression, since we're going to recompress
+	TArray<UObject*> ImportedAssets = AssetToolsModule.Get().ImportAssets(FileNames, TargetSubPath);
+	UTexture2D* ImportedTexture = ImportedTexture = (ImportedAssets.Num() > 0) ? Cast<UTexture2D>(ImportedAssets[0]) : nullptr;
+
+	if (ImportedTexture != nullptr)
+	{
+		// Change the compression settings
+		ImportedTexture->Modify();
+		ImportedTexture->LODGroup = TEXTUREGROUP_UI;
+		ImportedTexture->CompressionSettings = TC_EditorIcon;
+		ImportedTexture->Filter = TF_Nearest;
+		ImportedTexture->PostEditChange();
+	}
+
+	return ImportedTexture;
+}
+
+void UPaperTiledImporterFactory::ParseGlobalInfoFromJSON(TSharedPtr<FJsonObject> Tree, FTileMapFromTiled& OutParsedInfo, const FString& NameForErrors, bool bSilent)
+{
+	bool bSuccessfullyParsed = true;
+
+	// Parse all of the integer fields
+	FRequiredIntField IntFields[] = {
+		{ OutParsedInfo.FileVersion, TEXT("version") },
+		{ OutParsedInfo.Width, TEXT("width") },
+		{ OutParsedInfo.Height, TEXT("height") },
+		{ OutParsedInfo.TileWidth, TEXT("tilewidth") },
+		{ OutParsedInfo.TileHeight, TEXT("tileheight") }
+	};
+	bSuccessfullyParsed = bSuccessfullyParsed && ParseIntegerFields(IntFields, ARRAY_COUNT(IntFields), Tree, NameForErrors, bSilent);
+
+	// Parse the orientation
+	const FString OrientationModeStr = FPaperJSONHelpers::ReadString(Tree, TEXT("orientation"), TEXT(""));
+	if (OrientationModeStr == TEXT("orthogonal"))
+	{
+		OutParsedInfo.Orientation = ETiledOrientation::Orthogonal;
+	}
+	else if (OrientationModeStr == TEXT("isometric"))
+	{
+		OutParsedInfo.Orientation = ETiledOrientation::Isometric;
+	}
+	else
+	{
+		TILED_IMPORT_ERROR(TEXT("Failed to parse '%s'.  Invalid value for '%s' (%s but expected 'orthogonal' or 'isometric')"), *NameForErrors, TEXT("orientation"), *OrientationModeStr);
+		bSuccessfullyParsed = false;
+		OutParsedInfo.Orientation = ETiledOrientation::Unknown;
+	}
+}
+
+//////////////////////////////////////////////////////////////////////////
+// FTileMapFromTiled
+
+FTileMapFromTiled::FTileMapFromTiled()
+	: FileVersion(0)
+	, Width(0)
+	, Height(0)
+	, TileWidth(0)
+	, TileHeight(0)
+	, Orientation(ETiledOrientation::Unknown)
+{
+}
+
+bool FTileMapFromTiled::IsValid() const
+{
+	return (FileVersion != 0) && (Width > 0) && (Height > 0) && (TileWidth > 0) && (TileHeight > 0) && (Orientation != ETiledOrientation::Unknown);
+}
+
+FPaperTileInfo FTileMapFromTiled::ConvertTileGIDToPaper2D(int32 GID) const
+{
+	// Clear the mirroring / rotation bits, we don't support those yet
+	//@TODO: Handle mirroring/flipping flags
+	const int32 Flags = GID >> 29;
+	const int32 TileIndex = GID & ~(7U << 29);
+
+	FPaperTileInfo Result;
+
+	for (int32 SetIndex = TileSets.Num() - 1; SetIndex >= 0; SetIndex--)
+	{
+		const int32 RelativeIndex = TileIndex - TileSets[SetIndex].FirstGID;
+		if (RelativeIndex >= 0)
+		{
+			// We've found the source tile set and are done searching, but only import a non-null if that tile set imported successfully
+			if (UPaperTileSet* Set = CreatedTileSetAssets[SetIndex])
+			{
+				Result.TileSet = Set;
+				Result.PackedTileIndex = RelativeIndex;
+			}
+			break;
+		}
+	}
+
+	return Result;
+}
+
+//////////////////////////////////////////////////////////////////////////
+// FTileSetFromTiled
+
+FTileSetFromTiled::FTileSetFromTiled()
+	: FirstGID(INDEX_NONE)
+	, ImageWidth(0)
+	, ImageHeight(0)
+	, TileOffsetX(0)
+	, TileOffsetY(0)
+	, Margin(0)
+	, Spacing(0)
+	, TileWidth(0)
+	, TileHeight(0)
+{
+}
+
+void FTileSetFromTiled::ParseFromJSON(TSharedPtr<FJsonObject> Tree, const FString& NameForErrors, bool bSilent)
+{
+	bool bSuccessfullyParsed = true;
+
+	// Parse all of the integer fields
+	FRequiredIntField IntFields[] = {
+		{ FirstGID, TEXT("firstgid"), 1 },
+		{ ImageWidth, TEXT("imagewidth"), 1 },
+		{ ImageHeight, TEXT("imageheight"), 1 },
+		{ Margin, TEXT("margin"), 0 },
+		{ Spacing, TEXT("spacing"), 0 },
+		{ TileWidth, TEXT("tilewidth"), 1 },
+		{ TileHeight, TEXT("tileheight"), 1 }
+	};
+
+	bSuccessfullyParsed = bSuccessfullyParsed && ParseIntegerFields(IntFields, ARRAY_COUNT(IntFields), Tree, NameForErrors, bSilent);
+
+	// Parse the tile offset
+	if (bSuccessfullyParsed)
+	{
+		FIntPoint TileOffsetTemp;
+		if (FPaperJSONHelpers::ReadIntPoint(Tree, TEXT("tileoffset"), /*out*/ TileOffsetTemp))
+		{
+			TileOffsetX = TileOffsetTemp.X;
+			TileOffsetY = TileOffsetTemp.Y;
+		}
+		else
+		{
+			TILED_IMPORT_ERROR(TEXT("Failed to parse '%s'.  Invalid or missing value for '%s'"), *NameForErrors, TEXT("tileoffset"));
+			bSuccessfullyParsed = false;
+		}
+	}
+
+	//@TODO: Parse the properties
+	//@TODO: Parse the tile properties
+
+	// Parse the tile set name
+	Name = FPaperJSONHelpers::ReadString(Tree, TEXT("name"), TEXT(""));
+	if (Name.IsEmpty())
+	{
+		TILED_IMPORT_WARNING(TEXT("Expected a non-empty name for each tile set in '%s', generating a new name"), *NameForErrors);
+		Name = FString::Printf(TEXT("TileSetStartingAt%d"), FirstGID);
+	}
+
+	// Parse the image path
+	ImagePath = FPaperJSONHelpers::ReadString(Tree, TEXT("image"), TEXT(""));
+	if (ImagePath.IsEmpty())
+	{
+		TILED_IMPORT_ERROR(TEXT("Failed to parse '%s'.  Invalid value for '%s' (%s but expected a path to an image)"), *NameForErrors, TEXT("image"), *ImagePath);
+		bSuccessfullyParsed = false;
+	}
+}
+
+bool FTileSetFromTiled::IsValid() const
+{
+	return (TileWidth > 0) && (TileHeight > 0) && (FirstGID > 0);
+}
+
+//////////////////////////////////////////////////////////////////////////
+// FTileLayerFromTiled
+
+FTileLayerFromTiled::FTileLayerFromTiled()
+	: Width(0)
+	, Height(0)
+	, Opacity(1.0f)
+	, bVisible(true)
+	, OffsetX(0)
+	, OffsetY(0)
+{
+}
+
+void FTileLayerFromTiled::ParseFromJSON(TSharedPtr<FJsonObject> Tree, const FString& NameForErrors, bool bSilent)
+{
+	bool bSuccessfullyParsed = true;
+
+	// Parse all of the integer fields
+	FRequiredIntField IntFields[] = {
+		{ Width, TEXT("width"), 1 },
+		{ Height, TEXT("height"), 1 },
+		{ OffsetX, TEXT("x"), 0 },
+		{ OffsetY, TEXT("y"), 0 }
+	};
+
+	bSuccessfullyParsed = bSuccessfullyParsed && ParseIntegerFields(IntFields, ARRAY_COUNT(IntFields), Tree, NameForErrors, bSilent);
+
+	if (!Tree->TryGetBoolField(TEXT("visible"), /*out*/ bVisible))
+	{
+		bVisible = true;
+	}
+
+	if (!FPaperJSONHelpers::ReadFloatNoDefault(Tree, TEXT("opacity"), /*out*/ Opacity))
+	{
+		Opacity = 1.0f;
+	}
+
+	if (!Tree->TryGetStringField(TEXT("name"), /*out*/ Name))
+	{
+		//@TODO: Figure out what to do with the layer type!
+		TILED_IMPORT_ERROR(TEXT("Failed to parse '%s'.  Expected a layer name"), *NameForErrors);
+		bSuccessfullyParsed = false;
+	}
+
+	if (!Tree->TryGetStringField(TEXT("type"), /*out*/ Type))
+	{
+		//@TODO: Figure out what to do with the layer type!
+		TILED_IMPORT_ERROR(TEXT("Failed to parse '%s'.  Invalid value for '%s' (%s but expected a ???)"), *NameForErrors, TEXT("type"), *Type);
+		bSuccessfullyParsed = false;
+	}
+
+	const TArray<TSharedPtr<FJsonValue>>* DataArray;
+	if (Tree->TryGetArrayField(TEXT("data"), /*out*/ DataArray))
+	{
+		TileIndices.Reserve(DataArray->Num());
+		for (TSharedPtr<FJsonValue> TileEntry : *DataArray)
+		{
+			int32 TileID = (int32)TileEntry->AsNumber();
+			TileIndices.Add(TileID);
+		}
+	}
+	else
+	{
+		TILED_IMPORT_ERROR(TEXT("Failed to parse '%s'.  Missing layer data for layer '%s'"), *NameForErrors, *Name);
+		bSuccessfullyParsed = false;
+	}
+}
+
+bool FTileLayerFromTiled::IsValid() const
+{
+	return (Width > 0) && (Height > 0) && (TileIndices.Num() == (Width*Height));
+}
+
+//////////////////////////////////////////////////////////////////////////
+
+#undef TILED_IMPORT_ERROR
+#undef LOCTEXT_NAMESPACE
\ No newline at end of file
diff --git a/Engine/Plugins/Experimental/Paper2D/Source/PaperTiledImporter/Private/PaperTiledImporterModule.cpp b/Engine/Plugins/Experimental/Paper2D/Source/PaperTiledImporter/Private/PaperTiledImporterModule.cpp
new file mode 100644
--- /dev/null
+++ b/Engine/Plugins/Experimental/Paper2D/Source/PaperTiledImporter/Private/PaperTiledImporterModule.cpp
@@ -0,0 +1,23 @@
+// Copyright 1998-2014 Epic Games, Inc. All Rights Reserved.
+
+#include "PaperTiledImporterPrivatePCH.h"
+
+//////////////////////////////////////////////////////////////////////////
+// FPaperTiledImporterModule
+
+class FPaperTiledImporterModule : public FDefaultModuleImpl
+{
+public:
+	virtual void StartupModule() override
+	{
+	}
+
+	virtual void ShutdownModule() override
+	{
+	}
+};
+
+//////////////////////////////////////////////////////////////////////////
+
+IMPLEMENT_MODULE(FPaperTiledImporterModule, PaperTiledImporter);
+DEFINE_LOG_CATEGORY(LogPaperTiledImporter);
diff --git a/Engine/Plugins/Experimental/Paper2D/Source/PaperTiledImporter/Private/PaperTiledImporterPrivatePCH.h b/Engine/Plugins/Experimental/Paper2D/Source/PaperTiledImporter/Private/PaperTiledImporterPrivatePCH.h
new file mode 100644
--- /dev/null
+++ b/Engine/Plugins/Experimental/Paper2D/Source/PaperTiledImporter/Private/PaperTiledImporterPrivatePCH.h
@@ -0,0 +1,9 @@
+// Copyright 1998-2014 Epic Games, Inc. All Rights Reserved.
+
+#pragma once
+
+#include "UnrealEd.h"
+
+DECLARE_LOG_CATEGORY_EXTERN(LogPaperTiledImporter, Verbose, All);
+
+#include "PaperTiledImporterClasses.h"
