Paper2D: Finished deprecating layer-wide tileset pointer and removed several bogus uses

diff --git a/Engine/Plugins/Experimental/Paper2D/Source/Paper2D/Classes/PaperTileLayer.h b/Engine/Plugins/Experimental/Paper2D/Source/Paper2D/Classes/PaperTileLayer.h
--- a/Engine/Plugins/Experimental/Paper2D/Source/Paper2D/Classes/PaperTileLayer.h
+++ b/Engine/Plugins/Experimental/Paper2D/Source/Paper2D/Classes/PaperTileLayer.h
@@ -51,10 +51,6 @@ class PAPER2D_API UPaperTileLayer : public UObject
 	UPROPERTY(Category=Debug, VisibleAnywhere)
 	int32 LayerHeight;
 
-	// Tile Set associated with the layer //@TODO: Remove this or make it editor only
-	UPROPERTY(Category=Debug, VisibleAnywhere)
-	UPaperTileSet* TileSet;
-
 #if WITH_EDITORONLY_DATA
 	// Opacity of the layer (editor only)
 	UPROPERTY(Category=Debug, VisibleAnywhere)
@@ -82,13 +78,17 @@ protected:
 	UPROPERTY(Category=Debug, VisibleAnywhere)
 	int32 AllocatedHeight;
 
-	UPROPERTY()
-	TArray<int32> AllocatedGrid_DEPRECATED;
-
 	// The allocated tile data
 	UPROPERTY(Category=Debug, VisibleAnywhere)
 	TArray<FPaperTileInfo> AllocatedCells;
 
+private:
+	UPROPERTY()
+	UPaperTileSet* TileSet_DEPRECATED;
+
+	UPROPERTY()
+	TArray<int32> AllocatedGrid_DEPRECATED;
+
 public:
 	void ConvertToTileSetPerCell();
 
diff --git a/Engine/Plugins/Experimental/Paper2D/Source/Paper2D/Private/PaperTileLayer.cpp b/Engine/Plugins/Experimental/Paper2D/Source/Paper2D/Private/PaperTileLayer.cpp
--- a/Engine/Plugins/Experimental/Paper2D/Source/Paper2D/Private/PaperTileLayer.cpp
+++ b/Engine/Plugins/Experimental/Paper2D/Source/Paper2D/Private/PaperTileLayer.cpp
@@ -153,8 +153,8 @@ void UPaperTileLayer::ConvertToTileSetPerCell()
 	const int32 NumCells = AllocatedWidth * AllocatedHeight;
 	for (int32 Index = 0; Index < NumCells; ++Index)
 	{
-		FPaperTileInfo* Info = new (AllocatedCells)FPaperTileInfo();
-		Info->TileSet = TileSet;
+		FPaperTileInfo* Info = new (AllocatedCells) FPaperTileInfo();
+		Info->TileSet = TileSet_DEPRECATED;
 		Info->PackedTileIndex = AllocatedGrid_DEPRECATED[Index];
 	}
 }
\ No newline at end of file
diff --git a/Engine/Plugins/Experimental/Paper2D/Source/Paper2D/Private/PaperTileMapRenderSceneProxy.cpp b/Engine/Plugins/Experimental/Paper2D/Source/Paper2D/Private/PaperTileMapRenderSceneProxy.cpp
--- a/Engine/Plugins/Experimental/Paper2D/Source/Paper2D/Private/PaperTileMapRenderSceneProxy.cpp
+++ b/Engine/Plugins/Experimental/Paper2D/Source/Paper2D/Private/PaperTileMapRenderSceneProxy.cpp
@@ -103,7 +103,7 @@ void FPaperTileMapRenderSceneProxy::DrawDynamicElements(FPrimitiveDrawInterface*
 
 			FTransform LocalToWorld(GetLocalToWorld());
 
-			const float TW = TileMap ->TileWidth;
+			const float TW = TileMap->TileWidth;
 			const float TH = TileMap->TileHeight;
 			//@TODO: PAPER: Tiles probably shouldn't be drawn with their pivot at the center - RE: all the 0.5f in here
 
@@ -171,10 +171,8 @@ void FPaperTileMapRenderSceneProxy::DrawDynamicElements(FPrimitiveDrawInterface*
 			ensure(Layer->LayerWidth == TileMap->MapWidth);
 			ensure(Layer->LayerHeight == TileMap->MapHeight);
 
-			UPaperTileSet* TileSet = Layer->TileSet;
-				
-			int32 TileWidth = 4;
-			int32 TileHeight = 4;
+			const int32 TileWidth = TileMap->TileWidth;
+			const int32 TileHeight = TileMap->TileHeight;
 
 			UTexture2D* LastSourceTexture = nullptr;
 			FVector2D InverseTextureSize(1.0f, 1.0f);
@@ -227,8 +225,6 @@ void FPaperTileMapRenderSceneProxy::DrawDynamicElements(FPrimitiveDrawInterface*
 
 					if (SourceTexture != LastSourceTexture)
 					{
-						TileWidth = TileSet->TileWidth;
-						TileHeight = TileSet->TileHeight;
 						InverseTextureSize = FVector2D(1.0f / SourceTexture->GetSizeX(), 1.0f / SourceTexture->GetSizeY());
 						SourceDimensionsUV = FVector2D(TileWidth * InverseTextureSize.X, TileHeight * InverseTextureSize.Y);
 						LastSourceTexture = SourceTexture;
diff --git a/Engine/Plugins/Experimental/Paper2D/Source/Paper2DEditor/Private/TileMapEditing/EdModeTileMap.cpp b/Engine/Plugins/Experimental/Paper2D/Source/Paper2DEditor/Private/TileMapEditing/EdModeTileMap.cpp
--- a/Engine/Plugins/Experimental/Paper2D/Source/Paper2DEditor/Private/TileMapEditing/EdModeTileMap.cpp
+++ b/Engine/Plugins/Experimental/Paper2D/Source/Paper2DEditor/Private/TileMapEditing/EdModeTileMap.cpp
@@ -462,7 +462,7 @@ bool FEdModeTileMap::PaintTiles(const FViewportCursorLocation& Ray)
 					}
 
 					Ink.PackedTileIndex = SX + (SY * InkSource->GetTileCountX());
-					Ink.TileSet = Layer->TileSet;
+					Ink.TileSet = InkSource;
 				}
 
 				if (Layer->GetCell(DX, DY) != Ink)
@@ -662,7 +662,7 @@ bool FEdModeTileMap::FloodFillTiles(const FViewportCursorLocation& Ray)
 					{
 						const int32 TileSetX = PaintSourceTopLeft.X + InsideBrushX;
 						const int32 TileSetY = PaintSourceTopLeft.Y + InsideBrushY;
-						NewInk.TileSet = Layer->TileSet;
+						NewInk.TileSet = InkSource;
 						NewInk.PackedTileIndex = TileSetX + (TileSetY * InkSource->GetTileCountX());
 					}
 
diff --git a/Engine/Plugins/Experimental/Paper2D/Source/Paper2DEditor/Private/TileMapEditing/PaperTileMapDetailsCustomization.cpp b/Engine/Plugins/Experimental/Paper2D/Source/Paper2DEditor/Private/TileMapEditing/PaperTileMapDetailsCustomization.cpp
--- a/Engine/Plugins/Experimental/Paper2D/Source/Paper2DEditor/Private/TileMapEditing/PaperTileMapDetailsCustomization.cpp
+++ b/Engine/Plugins/Experimental/Paper2D/Source/Paper2DEditor/Private/TileMapEditing/PaperTileMapDetailsCustomization.cpp
@@ -150,7 +150,6 @@ UPaperTileLayer* FPaperTileMapDetailsCustomization::AddLayer(bool bCollisionLaye
 		NewLayer->LayerWidth = TileMap->MapWidth;
 		NewLayer->LayerHeight = TileMap->MapHeight;
 		NewLayer->DestructiveAllocateMap(NewLayer->LayerWidth, NewLayer->LayerHeight);
-		NewLayer->TileSet = TileMap->DefaultLayerTileSet;
 		NewLayer->LayerName = LOCTEXT("DefaultNewLayerName", "New Layer");
 		NewLayer->bCollisionLayer = bCollisionLayer;
 
diff --git a/Engine/Plugins/Experimental/Paper2D/Source/Paper2DEditor/Private/TileMapEditing/STileLayerItem.cpp b/Engine/Plugins/Experimental/Paper2D/Source/Paper2DEditor/Private/TileMapEditing/STileLayerItem.cpp
--- a/Engine/Plugins/Experimental/Paper2D/Source/Paper2DEditor/Private/TileMapEditing/STileLayerItem.cpp
+++ b/Engine/Plugins/Experimental/Paper2D/Source/Paper2DEditor/Private/TileMapEditing/STileLayerItem.cpp
@@ -215,15 +215,16 @@ void STileLayerItem::OnLayerNameCommitted(const FText& NewText, ETextCommit::Typ
 
 void STileLayerItem::OnChangeLayerTileSet(UObject* NewAsset)
 {
-	if (UPaperTileSet* TileSet = Cast<UPaperTileSet>(NewAsset))
-	{
-		MyLayer->TileSet = TileSet;
-	}
+	//@TODO: PAPER2D: Dead code now?
+// 	if (UPaperTileSet* TileSet = Cast<UPaperTileSet>(NewAsset))
+// 	{
+// 		MyLayer->TileSet = TileSet;
+// 	}
 }
 
 UObject* STileLayerItem::GetLayerTileset() const
 {
-	return MyLayer->TileSet;
+	return nullptr;// MyLayer->TileSet;
 }
 
 FReply STileLayerItem::OnToggleVisibility()
