Paper2D: Tile map editing improvements
- Add a placeholder standalone tile map editor (doesn't edit much yet)
- Add a component asset broker allowing tile maps components to be used in Blueprints

diff --git a/Engine/Plugins/Experimental/Paper2D/Source/Paper2D/Classes/PaperTileMapRenderComponent.h b/Engine/Plugins/Experimental/Paper2D/Source/Paper2D/Classes/PaperTileMapRenderComponent.h
--- a/Engine/Plugins/Experimental/Paper2D/Source/Paper2D/Classes/PaperTileMapRenderComponent.h
+++ b/Engine/Plugins/Experimental/Paper2D/Source/Paper2D/Classes/PaperTileMapRenderComponent.h
@@ -6,7 +6,7 @@
 
 #include "PaperTileMapRenderComponent.generated.h"
 
-UCLASS(hideCategories=Object, Experimental)
+UCLASS(hideCategories=Object, Experimental, meta=(BlueprintSpawnableComponent))
 class PAPER2D_API UPaperTileMapRenderComponent : public UPrimitiveComponent
 {
 	GENERATED_UCLASS_BODY()
diff --git a/Engine/Plugins/Experimental/Paper2D/Source/Paper2DEditor/Private/Paper2DEditorModule.cpp b/Engine/Plugins/Experimental/Paper2D/Source/Paper2DEditor/Private/Paper2DEditorModule.cpp
--- a/Engine/Plugins/Experimental/Paper2D/Source/Paper2DEditor/Private/Paper2DEditorModule.cpp
+++ b/Engine/Plugins/Experimental/Paper2D/Source/Paper2DEditor/Private/Paper2DEditorModule.cpp
@@ -3,23 +3,33 @@
 #include "Paper2DEditorPrivatePCH.h"
 #include "Paper2DEditorModule.h"
 
-#include "SpriteAssetTypeActions.h"
-#include "FlipbookAssetTypeActions.h"
-#include "TileSetAssetTypeActions.h"
-
 #include "AssetToolsModule.h"
 #include "PropertyEditorModule.h"
 #include "PaperStyle.h"
 #include "PaperEditorCommands.h"
-#include "TileMapEditing/EdModeTileMap.h"
-#include "TileMapEditing/PaperTileMapDetailsCustomization.h"
+
+#include "AssetEditorToolkit.h"
+
+#include "ContentBrowserExtensions/ContentBrowserExtensions.h"
+
+// Sprite support
+#include "SpriteAssetTypeActions.h"
+#include "PaperSpriteAssetBroker.h"
 #include "SpriteEditor/SpriteDetailsCustomization.h"
 #include "SpriteEditor/SpritePolygonCollectionCustomization.h"
-#include "PaperSpriteAssetBroker.h"
+
+// Flipbook support
+#include "FlipbookAssetTypeActions.h"
 #include "PaperFlipbookAssetBroker.h"
-#include "AssetEditorToolkit.h"
 
-#include "ContentBrowserExtensions/ContentBrowserExtensions.h"
+// Tile set support
+#include "TileSetAssetTypeActions.h"
+
+// Tile map support
+#include "TileMapEditing/TileMapAssetTypeActions.h"
+#include "TileMapEditing/PaperTileMapAssetBroker.h"
+#include "TileMapEditing/EdModeTileMap.h"
+#include "TileMapEditing/PaperTileMapDetailsCustomization.h"
 
 // Atlas support
 #include "Atlasing/AtlasAssetTypeActions.h"
@@ -62,6 +72,7 @@ private:
 
 	TSharedPtr<IComponentAssetBroker> PaperSpriteBroker;
 	TSharedPtr<IComponentAssetBroker> PaperFlipbookBroker;
+	TSharedPtr<IComponentAssetBroker> PaperTileMapBroker;
 
 	FCoreDelegates::FOnObjectPropertyChanged::FDelegate OnPropertyChangedHandle;
 
@@ -85,6 +96,7 @@ public:
 		RegisterAssetTypeAction(AssetTools, MakeShareable(new FSpriteAssetTypeActions));
 		RegisterAssetTypeAction(AssetTools, MakeShareable(new FFlipbookAssetTypeActions));
 		RegisterAssetTypeAction(AssetTools, MakeShareable(new FTileSetAssetTypeActions));
+		RegisterAssetTypeAction(AssetTools, MakeShareable(new FTileMapAssetTypeActions));
 		RegisterAssetTypeAction(AssetTools, MakeShareable(new FAtlasAssetTypeActions));
 
 		PaperSpriteBroker = MakeShareable(new FPaperSpriteAssetBroker);
@@ -93,6 +105,9 @@ public:
 		PaperFlipbookBroker = MakeShareable(new FPaperFlipbookAssetBroker);
 		FComponentAssetBrokerage::RegisterBroker(PaperFlipbookBroker, UPaperFlipbookComponent::StaticClass(), true, true);
 
+		PaperTileMapBroker = MakeShareable(new FPaperTileMapAssetBroker);
+		FComponentAssetBrokerage::RegisterBroker(PaperTileMapBroker, UPaperTileMapRenderComponent::StaticClass(), true, true);
+
 		// Register the details customizations
 		{
 			FPropertyEditorModule& PropertyModule = FModuleManager::LoadModuleChecked<FPropertyEditorModule>("PropertyEditor");
@@ -113,6 +128,7 @@ public:
 		UThumbnailManager::Get().RegisterCustomRenderer(UPaperSprite::StaticClass(), UPaperSpriteThumbnailRenderer::StaticClass());
 		UThumbnailManager::Get().RegisterCustomRenderer(UPaperTileSet::StaticClass(), UPaperTileSetThumbnailRenderer::StaticClass());
 		UThumbnailManager::Get().RegisterCustomRenderer(UPaperFlipbook::StaticClass(), UPaperFlipbookThumbnailRenderer::StaticClass());
+		//@TODO: PAPER2D: UThumbnailManager::Get().RegisterCustomRenderer(UPaperTileMap::StaticClass(), UPaperTileMapThumbnailRenderer::StaticClass());
 
 		// Register the editor modes
 		UpdateTileMapEditorModeInstallation();
@@ -139,6 +155,7 @@ public:
 
 			FPaperContentBrowserExtensions::RemoveHooks();
 
+			FComponentAssetBrokerage::UnregisterBroker(PaperTileMapBroker);
 			FComponentAssetBrokerage::UnregisterBroker(PaperFlipbookBroker);
 			FComponentAssetBrokerage::UnregisterBroker(PaperSpriteBroker);
 
@@ -147,6 +164,7 @@ public:
 
 			// Unregister the thumbnail renderers
 			UThumbnailManager::Get().UnregisterCustomRenderer(UPaperSprite::StaticClass());
+			UThumbnailManager::Get().UnregisterCustomRenderer(UPaperTileMap::StaticClass());
 			UThumbnailManager::Get().UnregisterCustomRenderer(UPaperTileSet::StaticClass());
 			UThumbnailManager::Get().UnregisterCustomRenderer(UPaperFlipbook::StaticClass());
 
diff --git a/Engine/Plugins/Experimental/Paper2D/Source/Paper2DEditor/Private/TileMapEditing/PaperTileMapAssetBroker.h b/Engine/Plugins/Experimental/Paper2D/Source/Paper2DEditor/Private/TileMapEditing/PaperTileMapAssetBroker.h
new file mode 100644
--- /dev/null
+++ b/Engine/Plugins/Experimental/Paper2D/Source/Paper2DEditor/Private/TileMapEditing/PaperTileMapAssetBroker.h
@@ -0,0 +1,47 @@
+// Copyright 1998-2014 Epic Games, Inc. All Rights Reserved.
+
+#pragma once
+
+#include "ComponentAssetBroker.h"
+
+//////////////////////////////////////////////////////////////////////////
+// FPaperTileMapAssetBroker
+
+class FPaperTileMapAssetBroker : public IComponentAssetBroker
+{
+public:
+	UClass* GetSupportedAssetClass() override
+	{
+		return UPaperTileMap::StaticClass();
+	}
+
+	virtual bool AssignAssetToComponent(UActorComponent* InComponent, UObject* InAsset) override
+	{
+		if (UPaperTileMapRenderComponent* RenderComp = Cast<UPaperTileMapRenderComponent>(InComponent))
+		{
+			UPaperTileMap* TileMap = Cast<UPaperTileMap>(InAsset);
+
+			if ((TileMap != nullptr) || (InAsset == nullptr))
+			{
+				RenderComp->TileMap = TileMap;
+				return true;
+			}
+		}
+
+		return false;
+	}
+
+	virtual UObject* GetAssetFromComponent(UActorComponent* InComponent) override
+	{
+		if (UPaperTileMapRenderComponent* RenderComp = Cast<UPaperTileMapRenderComponent>(InComponent))
+		{
+			if ((RenderComp->TileMap != nullptr) && (RenderComp->TileMap->IsAsset()))
+			{
+				return RenderComp->TileMap;
+			}
+		}
+
+		return nullptr;
+	}
+};
+
diff --git a/Engine/Plugins/Experimental/Paper2D/Source/Paper2DEditor/Private/TileMapEditing/STileMapEditorViewportToolbar.cpp b/Engine/Plugins/Experimental/Paper2D/Source/Paper2DEditor/Private/TileMapEditing/STileMapEditorViewportToolbar.cpp
new file mode 100644
--- /dev/null
+++ b/Engine/Plugins/Experimental/Paper2D/Source/Paper2DEditor/Private/TileMapEditing/STileMapEditorViewportToolbar.cpp
@@ -0,0 +1,16 @@
+// Copyright 1998-2014 Epic Games, Inc. All Rights Reserved.
+
+#include "Paper2DEditorPrivatePCH.h"
+#include "STileMapEditorViewportToolbar.h"
+
+#define LOCTEXT_NAMESPACE "STileMapEditorViewportToolbar"
+
+///////////////////////////////////////////////////////////
+// STileMapEditorViewportToolbar
+
+void STileMapEditorViewportToolbar::Construct(const FArguments& InArgs, TSharedPtr<class ICommonEditorViewportToolbarInfoProvider> InInfoProvider)
+{
+	SCommonEditorViewportToolbarBase::Construct(SCommonEditorViewportToolbarBase::FArguments(), InInfoProvider);
+}
+
+#undef LOCTEXT_NAMESPACE
diff --git a/Engine/Plugins/Experimental/Paper2D/Source/Paper2DEditor/Private/TileMapEditing/STileMapEditorViewportToolbar.h b/Engine/Plugins/Experimental/Paper2D/Source/Paper2DEditor/Private/TileMapEditing/STileMapEditorViewportToolbar.h
new file mode 100644
--- /dev/null
+++ b/Engine/Plugins/Experimental/Paper2D/Source/Paper2DEditor/Private/TileMapEditing/STileMapEditorViewportToolbar.h
@@ -0,0 +1,15 @@
+// Copyright 1998-2014 Epic Games, Inc. All Rights Reserved.
+
+#pragma once
+
+#include "Editor/UnrealEd/Public/SCommonEditorViewportToolbarBase.h"
+
+// In-viewport toolbar widget used in the tile map editor
+class STileMapEditorViewportToolbar : public SCommonEditorViewportToolbarBase
+{
+public:
+	SLATE_BEGIN_ARGS(STileMapEditorViewportToolbar) {}
+	SLATE_END_ARGS()
+
+	void Construct(const FArguments& InArgs, TSharedPtr<class ICommonEditorViewportToolbarInfoProvider> InInfoProvider);
+};
diff --git a/Engine/Plugins/Experimental/Paper2D/Source/Paper2DEditor/Private/TileMapEditing/TileMapAssetTypeActions.cpp b/Engine/Plugins/Experimental/Paper2D/Source/Paper2DEditor/Private/TileMapEditing/TileMapAssetTypeActions.cpp
new file mode 100644
--- /dev/null
+++ b/Engine/Plugins/Experimental/Paper2D/Source/Paper2DEditor/Private/TileMapEditing/TileMapAssetTypeActions.cpp
@@ -0,0 +1,49 @@
+// Copyright 1998-2014 Epic Games, Inc. All Rights Reserved.
+
+#include "Paper2DEditorPrivatePCH.h"
+
+#include "TileMapAssetTypeActions.h"
+#include "TileMapEditor.h"
+
+#define LOCTEXT_NAMESPACE "AssetTypeActions"
+
+//////////////////////////////////////////////////////////////////////////
+// FTileMapAssetTypeActions
+
+FText FTileMapAssetTypeActions::GetName() const
+{
+	return LOCTEXT("FTileMapAssetTypeActionsName", "Tile Map");
+}
+
+FColor FTileMapAssetTypeActions::GetTypeColor() const
+{
+	return FColor(0, 255, 255);
+}
+
+UClass* FTileMapAssetTypeActions::GetSupportedClass() const
+{
+	return UPaperTileMap::StaticClass();
+}
+
+void FTileMapAssetTypeActions::OpenAssetEditor(const TArray<UObject*>& InObjects, TSharedPtr<class IToolkitHost> EditWithinLevelEditor)
+{
+	EToolkitMode::Type Mode = EditWithinLevelEditor.IsValid() ? EToolkitMode::WorldCentric : EToolkitMode::Standalone;
+
+	for (auto ObjIt = InObjects.CreateConstIterator(); ObjIt; ++ObjIt)
+	{
+		if (UPaperTileMap* TileMap = Cast<UPaperTileMap>(*ObjIt))
+		{
+			TSharedRef<FTileMapEditor> NewTileMapEditor(new FTileMapEditor());
+			NewTileMapEditor->InitTileMapEditor(Mode, EditWithinLevelEditor, TileMap);
+		}
+	}
+}
+
+uint32 FTileMapAssetTypeActions::GetCategories()
+{
+	return EAssetTypeCategories::Misc;
+}
+
+//////////////////////////////////////////////////////////////////////////
+
+#undef LOCTEXT_NAMESPACE
\ No newline at end of file
diff --git a/Engine/Plugins/Experimental/Paper2D/Source/Paper2DEditor/Private/TileMapEditing/TileMapAssetTypeActions.h b/Engine/Plugins/Experimental/Paper2D/Source/Paper2DEditor/Private/TileMapEditing/TileMapAssetTypeActions.h
new file mode 100644
--- /dev/null
+++ b/Engine/Plugins/Experimental/Paper2D/Source/Paper2DEditor/Private/TileMapEditing/TileMapAssetTypeActions.h
@@ -0,0 +1,17 @@
+// Copyright 1998-2014 Epic Games, Inc. All Rights Reserved.
+
+#pragma once
+
+#include "AssetTypeActions_Base.h"
+
+class FTileMapAssetTypeActions : public FAssetTypeActions_Base
+{
+public:
+	// IAssetTypeActions interface
+	virtual FText GetName() const override;
+	virtual FColor GetTypeColor() const override;
+	virtual UClass* GetSupportedClass() const override;
+	virtual void OpenAssetEditor(const TArray<UObject*>& InObjects, TSharedPtr<class IToolkitHost> EditWithinLevelEditor = TSharedPtr<IToolkitHost>()) override;
+	virtual uint32 GetCategories() override;
+	// End of IAssetTypeActions interface
+};
\ No newline at end of file
diff --git a/Engine/Plugins/Experimental/Paper2D/Source/Paper2DEditor/Private/TileMapEditing/TileMapEditor.cpp b/Engine/Plugins/Experimental/Paper2D/Source/Paper2DEditor/Private/TileMapEditing/TileMapEditor.cpp
new file mode 100644
--- /dev/null
+++ b/Engine/Plugins/Experimental/Paper2D/Source/Paper2DEditor/Private/TileMapEditing/TileMapEditor.cpp
@@ -0,0 +1,458 @@
+// Copyright 1998-2014 Epic Games, Inc. All Rights Reserved.
+
+#include "Paper2DEditorPrivatePCH.h"
+#include "TileMapEditor.h"
+#include "SSingleObjectDetailsPanel.h"
+#include "SceneViewport.h"
+
+#include "GraphEditor.h"
+
+#include "TileMapEditorViewportClient.h"
+#include "TileMapEditorCommands.h"
+#include "SEditorViewport.h"
+#include "WorkspaceMenuStructureModule.h"
+#include "Paper2DEditorModule.h"
+#include "STileMapEditorViewportToolbar.h"
+
+#define LOCTEXT_NAMESPACE "TileMapEditor"
+
+//////////////////////////////////////////////////////////////////////////
+
+const FName TileMapEditorAppName = FName(TEXT("TileMapEditorApp"));
+
+//////////////////////////////////////////////////////////////////////////
+
+struct FTileMapEditorTabs
+{
+	// Tab identifiers
+	static const FName DetailsID;
+	static const FName ViewportID;
+};
+
+//////////////////////////////////////////////////////////////////////////
+
+const FName FTileMapEditorTabs::DetailsID(TEXT("Details"));
+const FName FTileMapEditorTabs::ViewportID(TEXT("Viewport"));
+
+//////////////////////////////////////////////////////////////////////////
+// STileMapEditorViewport
+
+class STileMapEditorViewport : public SEditorViewport, public ICommonEditorViewportToolbarInfoProvider
+{
+public:
+	SLATE_BEGIN_ARGS(STileMapEditorViewport) {}
+	SLATE_END_ARGS()
+
+	void Construct(const FArguments& InArgs, TSharedPtr<FTileMapEditor> InTileMapEditor);
+
+	// SEditorViewport interface
+	virtual void BindCommands() override;
+	virtual TSharedRef<FEditorViewportClient> MakeEditorViewportClient() override;
+	virtual TSharedPtr<SWidget> MakeViewportToolbar() override;
+	virtual EVisibility GetTransformToolbarVisibility() const override;
+	// End of SEditorViewport interface
+
+	// ICommonEditorViewportToolbarInfoProvider interface
+	virtual TSharedRef<class SEditorViewport> GetViewportWidget() override;
+	virtual TSharedPtr<FExtender> GetExtenders() const override;
+	virtual void OnFloatingButtonClicked() override;
+	// End of ICommonEditorViewportToolbarInfoProvider interface
+
+	// Invalidate any references to the tile map being edited; it has changed
+	void NotifyTileMapBeingEditedHasChanged()
+	{
+		EditorViewportClient->NotifyTileMapBeingEditedHasChanged();
+	}
+private:
+	// Pointer back to owning tile map editor instance (the keeper of state)
+	TWeakPtr<class FTileMapEditor> TileMapEditorPtr;
+
+	// Viewport client
+	TSharedPtr<FTileMapEditorViewportClient> EditorViewportClient;
+
+private:
+	// Returns true if the viewport is visible
+	bool IsVisible() const;
+};
+
+void STileMapEditorViewport::Construct(const FArguments& InArgs, TSharedPtr<FTileMapEditor> InTileMapEditor)
+{
+	TileMapEditorPtr = InTileMapEditor;
+
+	SEditorViewport::Construct(SEditorViewport::FArguments());
+}
+
+void STileMapEditorViewport::BindCommands()
+{
+	SEditorViewport::BindCommands();
+
+	const FTileMapEditorCommands& Commands = FTileMapEditorCommands::Get();
+
+	TSharedRef<FTileMapEditorViewportClient> EditorViewportClientRef = EditorViewportClient.ToSharedRef();
+
+	// Show toggles
+	CommandList->MapAction(
+		Commands.SetShowCollision,
+		FExecuteAction::CreateSP(EditorViewportClientRef, &FEditorViewportClient::SetShowCollision),
+		FCanExecuteAction(),
+		FIsActionChecked::CreateSP(EditorViewportClientRef, &FEditorViewportClient::IsSetShowCollisionChecked));
+
+	CommandList->MapAction(
+		Commands.SetShowPivot,
+		FExecuteAction::CreateSP(EditorViewportClientRef, &FTileMapEditorViewportClient::ToggleShowPivot),
+		FCanExecuteAction(),
+		FIsActionChecked::CreateSP(EditorViewportClientRef, &FTileMapEditorViewportClient::IsShowPivotChecked));
+
+	// Misc. actions
+	CommandList->MapAction(
+		Commands.FocusOnTileMap,
+		FExecuteAction::CreateSP(EditorViewportClientRef, &FTileMapEditorViewportClient::FocusOnTileMap));
+}
+
+TSharedRef<FEditorViewportClient> STileMapEditorViewport::MakeEditorViewportClient()
+{
+	EditorViewportClient = MakeShareable(new FTileMapEditorViewportClient(TileMapEditorPtr, SharedThis(this)));
+
+	EditorViewportClient->VisibilityDelegate.BindSP(this, &STileMapEditorViewport::IsVisible);
+
+	return EditorViewportClient.ToSharedRef();
+}
+
+bool STileMapEditorViewport::IsVisible() const
+{
+	return true;//@TODO: Determine this better so viewport ticking optimizations can take place
+}
+
+TSharedPtr<SWidget> STileMapEditorViewport::MakeViewportToolbar()
+{
+	return SNew(STileMapEditorViewportToolbar, SharedThis(this));
+}
+
+EVisibility STileMapEditorViewport::GetTransformToolbarVisibility() const
+{
+	return EVisibility::Visible;
+}
+
+TSharedRef<class SEditorViewport> STileMapEditorViewport::GetViewportWidget()
+{
+	return SharedThis(this);
+}
+
+TSharedPtr<FExtender> STileMapEditorViewport::GetExtenders() const
+{
+	TSharedPtr<FExtender> Result(MakeShareable(new FExtender));
+	return Result;
+}
+
+void STileMapEditorViewport::OnFloatingButtonClicked()
+{
+}
+
+/////////////////////////////////////////////////////
+// STileMapPropertiesTabBody
+
+class STileMapPropertiesTabBody : public SSingleObjectDetailsPanel
+{
+public:
+	SLATE_BEGIN_ARGS(STileMapPropertiesTabBody) {}
+	SLATE_END_ARGS()
+
+private:
+	// Pointer back to owning TileMap editor instance (the keeper of state)
+	TWeakPtr<class FTileMapEditor> TileMapEditorPtr;
+public:
+	void Construct(const FArguments& InArgs, TSharedPtr<FTileMapEditor> InTileMapEditor)
+	{
+		TileMapEditorPtr = InTileMapEditor;
+
+		SSingleObjectDetailsPanel::Construct(SSingleObjectDetailsPanel::FArguments());
+	}
+
+	// SSingleObjectDetailsPanel interface
+	virtual UObject* GetObjectToObserve() const override
+	{
+		return TileMapEditorPtr.Pin()->GetTileMapBeingEdited();
+	}
+
+	virtual TSharedRef<SWidget> PopulateSlot(TSharedRef<SWidget> PropertyEditorWidget) override
+	{
+		return SNew(SVerticalBox)
+			+ SVerticalBox::Slot()
+			.FillHeight(1)
+			[
+				PropertyEditorWidget
+			];
+	}
+	// End of SSingleObjectDetailsPanel interface
+};
+
+//////////////////////////////////////////////////////////////////////////
+// FTileMapEditor
+
+TSharedRef<SDockTab> FTileMapEditor::SpawnTab_Viewport(const FSpawnTabArgs& Args)
+{
+	return SNew(SDockTab)
+		.Label(LOCTEXT("ViewportTab_Title", "Viewport"))
+		[
+			SNew(SOverlay)
+
+			// The TileMap editor viewport
+			+ SOverlay::Slot()
+			[
+				ViewportPtr.ToSharedRef()
+			]
+
+			// Bottom-right corner text indicating the preview nature of the tile map editor
+			+ SOverlay::Slot()
+				.Padding(10)
+				.VAlign(VAlign_Bottom)
+				.HAlign(HAlign_Right)
+				[
+					SNew(STextBlock)
+					.Visibility(EVisibility::HitTestInvisible)
+					.TextStyle(FEditorStyle::Get(), "Graph.CornerText")
+					.Text(LOCTEXT("TileMapEditorViewportExperimentalWarning", "Experimental!"))
+				]
+		];
+}
+
+TSharedRef<SDockTab> FTileMapEditor::SpawnTab_Details(const FSpawnTabArgs& Args)
+{
+	TSharedPtr<FTileMapEditor> TileMapEditorPtr = SharedThis(this);
+
+	// Spawn the tab
+	return SNew(SDockTab)
+		.Icon(FEditorStyle::GetBrush("LevelEditor.Tabs.Details"))
+		.Label(LOCTEXT("DetailsTab_Title", "Details"))
+		[
+			SNew(STileMapPropertiesTabBody, TileMapEditorPtr)
+		];
+}
+
+void FTileMapEditor::RegisterTabSpawners(const TSharedRef<class FTabManager>& TabManager)
+{
+	FAssetEditorToolkit::RegisterTabSpawners(TabManager);
+
+	const IWorkspaceMenuStructure& MenuStructure = WorkspaceMenu::GetMenuStructure();
+
+	TabManager->RegisterTabSpawner(FTileMapEditorTabs::ViewportID, FOnSpawnTab::CreateSP(this, &FTileMapEditor::SpawnTab_Viewport))
+		.SetDisplayName(LOCTEXT("ViewportTab", "Viewport"))
+		.SetGroup(MenuStructure.GetAssetEditorCategory());
+
+	TabManager->RegisterTabSpawner(FTileMapEditorTabs::DetailsID, FOnSpawnTab::CreateSP(this, &FTileMapEditor::SpawnTab_Details))
+		.SetDisplayName(LOCTEXT("DetailsTabLabel", "Details"))
+		.SetGroup(MenuStructure.GetAssetEditorCategory());
+}
+
+void FTileMapEditor::UnregisterTabSpawners(const TSharedRef<class FTabManager>& TabManager)
+{
+	FAssetEditorToolkit::UnregisterTabSpawners(TabManager);
+
+	TabManager->UnregisterTabSpawner(FTileMapEditorTabs::ViewportID);
+	TabManager->UnregisterTabSpawner(FTileMapEditorTabs::DetailsID);
+}
+
+void FTileMapEditor::InitTileMapEditor(const EToolkitMode::Type Mode, const TSharedPtr< class IToolkitHost >& InitToolkitHost, class UPaperTileMap* InitTileMap)
+{
+	FAssetEditorManager::Get().CloseOtherEditors(InitTileMap, this);
+	TileMapBeingEdited = InitTileMap;
+
+	FTileMapEditorCommands::Register();
+
+	BindCommands();
+
+	ViewportPtr = SNew(STileMapEditorViewport, SharedThis(this));
+
+	// Default layout
+	const TSharedRef<FTabManager::FLayout> StandaloneDefaultLayout = FTabManager::NewLayout("Standalone_TileMapEditor_Layout_v1")
+		->AddArea
+		(
+			FTabManager::NewPrimaryArea()
+			->SetOrientation(Orient_Vertical)
+			->Split
+			(
+				FTabManager::NewStack()
+				->SetSizeCoefficient(0.1f)
+				->SetHideTabWell(true)
+				->AddTab(GetToolbarTabId(), ETabState::OpenedTab)
+			)
+			->Split
+			(
+				FTabManager::NewSplitter()
+				->SetOrientation(Orient_Horizontal)
+				->SetSizeCoefficient(0.9f)
+				->Split
+				(
+					FTabManager::NewStack()
+					->SetSizeCoefficient(0.8f)
+					->SetHideTabWell(true)
+					->AddTab(FTileMapEditorTabs::ViewportID, ETabState::OpenedTab)
+				)
+				->Split
+				(
+					FTabManager::NewSplitter()
+					->SetOrientation(Orient_Vertical)
+					->SetSizeCoefficient(0.2f)
+					->Split
+					(
+						FTabManager::NewStack()
+						->SetSizeCoefficient(0.75f)
+						->AddTab(FTileMapEditorTabs::DetailsID, ETabState::OpenedTab)
+					)
+				)
+			)
+		);
+
+	// Initialize the asset editor and spawn nothing (dummy layout)
+	InitAssetEditor(Mode, InitToolkitHost, TileMapEditorAppName, StandaloneDefaultLayout, /*bCreateDefaultStandaloneMenu=*/ true, /*bCreateDefaultToolbar=*/ true, InitTileMap);
+
+	// Extend things
+	ExtendMenu();
+	ExtendToolbar();
+	RegenerateMenusAndToolbars();
+}
+
+void FTileMapEditor::BindCommands()
+{
+	// 	const FTileMapEditorCommands& Commands = FTileMapEditorCommands::Get();
+	// 
+	// 	const TSharedRef<FUICommandList>& UICommandList = GetToolkitCommands();
+	// 
+	// 	UICommandList->MapAction( FGenericCommands::Get().Delete,
+	// 		FExecuteAction::CreateSP( this, &FStaticMeshEditor::DeleteSelectedSockets ),
+	// 		FCanExecuteAction::CreateSP( this, &FStaticMeshEditor::HasSelectedSockets ));
+	// 
+	// 	UICommandList->MapAction( FGenericCommands::Get().Undo, 
+	// 		FExecuteAction::CreateSP( this, &FStaticMeshEditor::UndoAction ) );
+	// 
+	// 	UICommandList->MapAction( FGenericCommands::Get().Redo, 
+	// 		FExecuteAction::CreateSP( this, &FStaticMeshEditor::RedoAction ) );
+	// 
+	// 	UICommandList->MapAction(
+	// 		FGenericCommands::Get().Duplicate,
+	// 		FExecuteAction::CreateSP(this, &FStaticMeshEditor::DuplicateSelectedSocket),
+	// 		FCanExecuteAction::CreateSP(this, &FStaticMeshEditor::HasSelectedSockets));
+}
+
+FName FTileMapEditor::GetToolkitFName() const
+{
+	return FName("TileMapEditor");
+}
+
+FText FTileMapEditor::GetBaseToolkitName() const
+{
+	return LOCTEXT("TileMapEditorAppLabelBase", "Tile Map Editor");
+}
+
+FText FTileMapEditor::GetToolkitName() const
+{
+	const bool bDirtyState = TileMapBeingEdited->GetOutermost()->IsDirty();
+
+	FFormatNamedArguments Args;
+	Args.Add(TEXT("TileMapName"), FText::FromString(TileMapBeingEdited->GetName()));
+	Args.Add(TEXT("DirtyState"), bDirtyState ? FText::FromString(TEXT("*")) : FText::GetEmpty());
+	return FText::Format(LOCTEXT("TileMapEditorAppLabel", "{TileMapName}{DirtyState}"), Args);
+}
+
+FString FTileMapEditor::GetWorldCentricTabPrefix() const
+{
+	return TEXT("TileMapEditor");
+}
+
+FString FTileMapEditor::GetDocumentationLink() const
+{
+	return TEXT("Engine/Paper2D/TileMapEditor");
+}
+
+FLinearColor FTileMapEditor::GetWorldCentricTabColorScale() const
+{
+	return FLinearColor::White;
+}
+
+void FTileMapEditor::AddReferencedObjects(FReferenceCollector& Collector)
+{
+	Collector.AddReferencedObject(TileMapBeingEdited);
+}
+
+void FTileMapEditor::ExtendMenu()
+{
+}
+
+void FTileMapEditor::ExtendToolbar()
+{
+	struct Local
+	{
+		static void FillToolbar(FToolBarBuilder& ToolbarBuilder)
+		{
+			// 			ToolbarBuilder.BeginSection("Realtime");
+			// 			{
+			// 				ToolbarBuilder.AddToolBarButton(FEditorViewportCommands::Get().ToggleRealTime);
+			// 			}
+			// 			ToolbarBuilder.EndSection();
+
+// 			ToolbarBuilder.BeginSection("Command");
+// 			{
+// 				ToolbarBuilder.AddToolBarButton(FTileMapEditorCommands::Get().SetShowSourceTexture);
+// 			}
+// 			ToolbarBuilder.EndSection();
+// 
+// 			ToolbarBuilder.BeginSection("Tools");
+// 			{
+// 				ToolbarBuilder.AddToolBarButton(FTileMapEditorCommands::Get().AddPolygon);
+// 				ToolbarBuilder.AddToolBarButton(FTileMapEditorCommands::Get().SnapAllVertices);
+// 			}
+// 			ToolbarBuilder.EndSection();
+		}
+	};
+
+	TSharedPtr<FExtender> ToolbarExtender = MakeShareable(new FExtender);
+
+	ToolbarExtender->AddToolBarExtension(
+		"Asset",
+		EExtensionHook::After,
+		ViewportPtr->GetCommandList(),
+		FToolBarExtensionDelegate::CreateStatic(&Local::FillToolbar)
+		);
+
+// 	ToolbarExtender->AddToolBarExtension(
+// 		"Asset",
+// 		EExtensionHook::After,
+// 		ViewportPtr->GetCommandList(),
+// 		FToolBarExtensionDelegate::CreateSP(this, &FTileMapEditor::CreateModeToolbarWidgets));
+
+	AddToolbarExtender(ToolbarExtender);
+// 
+// 	IPaper2DEditorModule* Paper2DEditorModule = &FModuleManager::LoadModuleChecked<IPaper2DEditorModule>("Paper2DEditor");
+// 	AddToolbarExtender(Paper2DEditorModule->GetTileMapEditorToolBarExtensibilityManager()->GetAllExtenders());
+}
+
+void FTileMapEditor::SetTileMapBeingEdited(UPaperTileMap* NewTileMap)
+{
+	if ((NewTileMap != TileMapBeingEdited) && (NewTileMap != NULL))
+	{
+		UPaperTileMap* OldTileMap = TileMapBeingEdited;
+		TileMapBeingEdited = NewTileMap;
+
+		// Let the viewport know that we are editing something different
+		ViewportPtr->NotifyTileMapBeingEditedHasChanged();
+
+		// Let the editor know that are editing something different
+		RemoveEditingObject(OldTileMap);
+		AddEditingObject(NewTileMap);
+	}
+}
+
+
+void FTileMapEditor::CreateModeToolbarWidgets(FToolBarBuilder& IgnoredBuilder)
+{
+// 	FToolBarBuilder ToolbarBuilder(ViewportPtr->GetCommandList(), FMultiBoxCustomization::None);
+// 	ToolbarBuilder.AddToolBarButton(FTileMapEditorCommands::Get().EnterViewMode);
+// 	ToolbarBuilder.AddToolBarButton(FTileMapEditorCommands::Get().EnterSourceRegionEditMode);
+// 	ToolbarBuilder.AddToolBarButton(FTileMapEditorCommands::Get().EnterCollisionEditMode);
+// 	ToolbarBuilder.AddToolBarButton(FTileMapEditorCommands::Get().EnterRenderingEditMode);
+// 	AddToolbarWidget(ToolbarBuilder.MakeWidget());
+}
+
+//////////////////////////////////////////////////////////////////////////
+
+#undef LOCTEXT_NAMESPACE
diff --git a/Engine/Plugins/Experimental/Paper2D/Source/Paper2DEditor/Private/TileMapEditing/TileMapEditor.h b/Engine/Plugins/Experimental/Paper2D/Source/Paper2DEditor/Private/TileMapEditing/TileMapEditor.h
new file mode 100644
--- /dev/null
+++ b/Engine/Plugins/Experimental/Paper2D/Source/Paper2DEditor/Private/TileMapEditing/TileMapEditor.h
@@ -0,0 +1,49 @@
+// Copyright 1998-2014 Epic Games, Inc. All Rights Reserved.
+
+#pragma once
+
+#include "WorkflowCentricApplication.h"
+#include "SPaperEditorViewport.h"
+
+//////////////////////////////////////////////////////////////////////////
+// FTileMapEditor
+
+class FTileMapEditor : public FAssetEditorToolkit, public FGCObject
+{
+public:
+	// IToolkit interface
+	virtual void RegisterTabSpawners(const TSharedRef<class FTabManager>& TabManager) override;
+	virtual void UnregisterTabSpawners(const TSharedRef<class FTabManager>& TabManager) override;
+	// End of IToolkit interface
+
+	// FAssetEditorToolkit
+	virtual FName GetToolkitFName() const override;
+	virtual FText GetBaseToolkitName() const override;
+	virtual FText GetToolkitName() const override;
+	virtual FLinearColor GetWorldCentricTabColorScale() const override;
+	virtual FString GetWorldCentricTabPrefix() const override;
+	virtual FString GetDocumentationLink() const override;
+	// End of FAssetEditorToolkit
+
+	// FSerializableObject interface
+	virtual void AddReferencedObjects(FReferenceCollector& Collector) override;
+	// End of FSerializableObject interface
+public:
+	void InitTileMapEditor(const EToolkitMode::Type Mode, const TSharedPtr< class IToolkitHost >& InitToolkitHost, class UPaperTileMap* InitTileMap);
+
+	UPaperTileMap* GetTileMapBeingEdited() const { return TileMapBeingEdited; }
+	void SetTileMapBeingEdited(UPaperTileMap* NewTileMap);
+protected:
+	UPaperTileMap* TileMapBeingEdited;
+	TSharedPtr<class STileMapEditorViewport> ViewportPtr;
+
+protected:
+	void BindCommands();
+	void ExtendMenu();
+	void ExtendToolbar();
+
+	TSharedRef<SDockTab> SpawnTab_Viewport(const FSpawnTabArgs& Args);
+	TSharedRef<SDockTab> SpawnTab_Details(const FSpawnTabArgs& Args);
+
+	void CreateModeToolbarWidgets(FToolBarBuilder& ToolbarBuilder);
+};
diff --git a/Engine/Plugins/Experimental/Paper2D/Source/Paper2DEditor/Private/TileMapEditing/TileMapEditorCommands.cpp b/Engine/Plugins/Experimental/Paper2D/Source/Paper2DEditor/Private/TileMapEditing/TileMapEditorCommands.cpp
new file mode 100644
--- /dev/null
+++ b/Engine/Plugins/Experimental/Paper2D/Source/Paper2DEditor/Private/TileMapEditing/TileMapEditorCommands.cpp
@@ -0,0 +1,18 @@
+// Copyright 1998-2014 Epic Games, Inc. All Rights Reserved.
+
+#include "Paper2DEditorPrivatePCH.h"
+#include "TileMapEditorCommands.h"
+
+//////////////////////////////////////////////////////////////////////////
+// FTileMapEditorCommands
+
+void FTileMapEditorCommands::RegisterCommands()
+{
+	// Show toggles
+	UI_COMMAND(SetShowCollision, "Collision", "Toggles display of the simplified collision mesh of the static mesh, if one has been assigned.", EUserInterfaceActionType::ToggleButton, FInputGesture(EKeys::C, EModifierKey::Alt));
+
+	UI_COMMAND(SetShowPivot, "Pivot", "Display the pivot location of the sprite.", EUserInterfaceActionType::ToggleButton, FInputGesture());
+
+	// Misc. actions
+	UI_COMMAND(FocusOnTileMap, "Focus on tile map", "Centers and zooms the view to focus on the tile map.", EUserInterfaceActionType::Button, FInputGesture(EKeys::Home, EModifierKey::None));
+}
diff --git a/Engine/Plugins/Experimental/Paper2D/Source/Paper2DEditor/Private/TileMapEditing/TileMapEditorCommands.h b/Engine/Plugins/Experimental/Paper2D/Source/Paper2DEditor/Private/TileMapEditing/TileMapEditorCommands.h
new file mode 100644
--- /dev/null
+++ b/Engine/Plugins/Experimental/Paper2D/Source/Paper2DEditor/Private/TileMapEditing/TileMapEditorCommands.h
@@ -0,0 +1,30 @@
+// Copyright 1998-2014 Epic Games, Inc. All Rights Reserved.
+
+#pragma once
+
+class FTileMapEditorCommands : public TCommands<FTileMapEditorCommands>
+{
+public:
+	FTileMapEditorCommands()
+		: TCommands<FTileMapEditorCommands>(
+			TEXT("TileMapEditor"), // Context name for fast lookup
+			NSLOCTEXT("Contexts", "TileMapEditor", "Tile Map Editor"), // Localized context name for displaying
+			NAME_None, // Parent
+			FEditorStyle::GetStyleSetName() // Icon Style Set
+			)
+	{
+	}
+
+	// TCommand<> interface
+	virtual void RegisterCommands() override;
+	// End of TCommand<> interface
+
+public:
+	// Show toggles
+	TSharedPtr<FUICommandInfo> SetShowCollision;
+	
+	TSharedPtr<FUICommandInfo> SetShowPivot;
+
+	// Misc. actions
+	TSharedPtr<FUICommandInfo> FocusOnTileMap;
+};
\ No newline at end of file
diff --git a/Engine/Plugins/Experimental/Paper2D/Source/Paper2DEditor/Private/TileMapEditing/TileMapEditorViewportClient.cpp b/Engine/Plugins/Experimental/Paper2D/Source/Paper2DEditor/Private/TileMapEditing/TileMapEditorViewportClient.cpp
new file mode 100644
--- /dev/null
+++ b/Engine/Plugins/Experimental/Paper2D/Source/Paper2DEditor/Private/TileMapEditing/TileMapEditorViewportClient.cpp
@@ -0,0 +1,259 @@
+// Copyright 1998-2014 Epic Games, Inc. All Rights Reserved.
+
+#include "Paper2DEditorPrivatePCH.h"
+#include "TileMapEditorViewportClient.h"
+#include "SceneViewport.h"
+
+#include "PreviewScene.h"
+#include "ScopedTransaction.h"
+#include "Runtime/Engine/Public/ComponentReregisterContext.h"
+
+#define LOCTEXT_NAMESPACE "TileMapEditor"
+
+//////////////////////////////////////////////////////////////////////////
+// FTileMapEditorViewportClient
+
+FTileMapEditorViewportClient::FTileMapEditorViewportClient(TWeakPtr<FTileMapEditor> InTileMapEditor, TWeakPtr<class STileMapEditorViewport> InTileMapEditorViewportPtr)
+	: TileMapEditorPtr(InTileMapEditor)
+	, TileMapEditorViewportPtr(InTileMapEditorViewportPtr)
+{
+	check(TileMapEditorPtr.IsValid() && TileMapEditorViewportPtr.IsValid());
+
+	PreviewScene = &OwnedPreviewScene;
+
+	SetRealtime(true);
+
+	WidgetMode = FWidget::WM_Translate;
+	bManipulating = false;
+	bManipulationDirtiedSomething = false;
+	ScopedTransaction = NULL;
+
+	bShowPivot = true;
+
+	bDeferZoomToTileMap = true;
+
+	EngineShowFlags.DisableAdvancedFeatures();
+	EngineShowFlags.CompositeEditorPrimitives = true;
+
+	// Create a render component for the tile map being edited
+	{
+		RenderTileMapComponent = NewObject<UPaperTileMapRenderComponent>();
+		UPaperTileMap* TileMap = GetTileMapBeingEdited();
+		RenderTileMapComponent->TileMap = TileMap;
+
+		PreviewScene->AddComponent(RenderTileMapComponent, FTransform::Identity);
+	}
+}
+
+void FTileMapEditorViewportClient::DrawBoundsAsText(FViewport& InViewport, FSceneView& View, FCanvas& Canvas, int32& YPos)
+{
+	FNumberFormattingOptions NoDigitGroupingFormat;
+	NoDigitGroupingFormat.UseGrouping = false;
+
+	UPaperTileMap* TileMap = GetTileMapBeingEdited();
+	FBoxSphereBounds Bounds;//@TODO: = TileMap->GetRenderBounds();
+
+	const FText DisplaySizeText = FText::Format(LOCTEXT("BoundsSize", "Approx. Size: {0}x{1}x{2}"),
+		FText::AsNumber((int32)(Bounds.BoxExtent.X * 2.0f), &NoDigitGroupingFormat),
+		FText::AsNumber((int32)(Bounds.BoxExtent.Y * 2.0f), &NoDigitGroupingFormat),
+		FText::AsNumber((int32)(Bounds.BoxExtent.Z * 2.0f), &NoDigitGroupingFormat));
+
+	Canvas.DrawShadowedString(
+		6,
+		YPos,
+		*DisplaySizeText.ToString(),
+		GEngine->GetSmallFont(),
+		FLinearColor::White);
+	YPos += 18;
+}
+
+void FTileMapEditorViewportClient::DrawCanvas(FViewport& Viewport, FSceneView& View, FCanvas& Canvas)
+{
+	FEditorViewportClient::DrawCanvas(Viewport, View, Canvas);
+
+	const bool bIsHitTesting = Canvas.IsHitTesting();
+	if (!bIsHitTesting)
+	{
+		Canvas.SetHitProxy(NULL);
+	}
+
+	if (!TileMapEditorPtr.IsValid())
+	{
+		return;
+	}
+
+	int32 YPos = 42;
+}
+
+void FTileMapEditorViewportClient::Draw(const FSceneView* View, FPrimitiveDrawInterface* PDI)
+{
+	FEditorViewportClient::Draw(View, PDI);
+
+	if (bShowPivot)
+	{
+		FUnrealEdUtils::DrawWidget(View, PDI, RenderTileMapComponent->ComponentToWorld.ToMatrixWithScale(), 0, 0, EAxisList::XZ, EWidgetMovementMode::WMM_Translate);
+	}
+}
+
+void FTileMapEditorViewportClient::Draw(FViewport* Viewport, FCanvas* Canvas)
+{
+	FEditorViewportClient::Draw(Viewport, Canvas);
+}
+
+void FTileMapEditorViewportClient::Tick(float DeltaSeconds)
+{
+	if (UPaperTileMap* TileMap = RenderTileMapComponent->TileMap)
+	{
+		// Zoom in on the tile map
+		//@TODO: Fix this properly so it doesn't need to be deferred, or wait for the viewport to initialize
+		FIntPoint Size = Viewport->GetSizeXY();
+		if (bDeferZoomToTileMap && (Size.X > 0) && (Size.Y > 0))
+		{
+			UPaperTileMapRenderComponent* ComponentToFocusOn = RenderTileMapComponent;
+			FocusViewportOnBox(ComponentToFocusOn->Bounds.GetBox(), true);
+			bDeferZoomToTileMap = false;
+		}		
+	}
+
+	FPaperEditorViewportClient::Tick(DeltaSeconds);
+
+	if (!GIntraFrameDebuggingGameThread)
+	{
+		OwnedPreviewScene.GetWorld()->Tick(LEVELTICK_All, DeltaSeconds);
+	}
+}
+
+void FTileMapEditorViewportClient::ToggleShowMeshEdges()
+{
+	EngineShowFlags.MeshEdges = !EngineShowFlags.MeshEdges;
+	Invalidate();
+}
+
+bool FTileMapEditorViewportClient::IsShowMeshEdgesChecked() const
+{
+	return EngineShowFlags.MeshEdges;
+}
+
+void FTileMapEditorViewportClient::UpdateMouseDelta()
+{
+	FPaperEditorViewportClient::UpdateMouseDelta();
+}
+
+void FTileMapEditorViewportClient::ProcessClick(FSceneView& View, HHitProxy* HitProxy, FKey Key, EInputEvent Event, uint32 HitX, uint32 HitY)
+{
+	FPaperEditorViewportClient::ProcessClick(View, HitProxy, Key, Event, HitX, HitY);
+}
+
+bool FTileMapEditorViewportClient::InputKey(FViewport* Viewport, int32 ControllerId, FKey Key, EInputEvent Event, float AmountDepressed, bool bGamepad)
+{
+	bool bHandled = false;
+
+	// Start the drag
+	//@TODO: EKeys::LeftMouseButton
+	//@TODO: Event.IE_Pressed
+	// Implement InputAxis
+	// StartTracking
+
+	// Pass keys to standard controls, if we didn't consume input
+	return (bHandled) ? true : FEditorViewportClient::InputKey(Viewport,  ControllerId, Key, Event, AmountDepressed, bGamepad);
+}
+
+bool FTileMapEditorViewportClient::InputWidgetDelta(FViewport* Viewport, EAxisList::Type CurrentAxis, FVector& Drag, FRotator& Rot, FVector& Scale)
+{
+	bool bHandled = false;
+	return bHandled;
+}
+
+void FTileMapEditorViewportClient::TrackingStarted(const struct FInputEventState& InInputState, bool bIsDragging, bool bNudge)
+{
+	if (!bManipulating && bIsDragging)
+	{
+		BeginTransaction(LOCTEXT("ModificationInViewport", "Modification in Viewport"));
+		bManipulating = true;
+		bManipulationDirtiedSomething = false;
+	}
+}
+
+void FTileMapEditorViewportClient::TrackingStopped()
+{
+	if (bManipulating)
+	{
+		EndTransaction();
+		bManipulating = false;
+	}
+}
+
+FWidget::EWidgetMode FTileMapEditorViewportClient::GetWidgetMode() const
+{
+	return FWidget::WM_None;
+}
+
+FVector FTileMapEditorViewportClient::GetWidgetLocation() const
+{
+
+	return FVector::ZeroVector;
+}
+
+FMatrix FTileMapEditorViewportClient::GetWidgetCoordSystem() const
+{
+	return FMatrix::Identity;
+}
+
+ECoordSystem FTileMapEditorViewportClient::GetWidgetCoordSystemSpace() const
+{
+	return COORD_World;
+}
+
+void FTileMapEditorViewportClient::BeginTransaction(const FText& SessionName)
+{
+	if (ScopedTransaction == NULL)
+	{
+		ScopedTransaction = new FScopedTransaction(SessionName);
+
+		UPaperTileMap* TileMap = GetTileMapBeingEdited();
+		TileMap->Modify();
+	}
+}
+
+void FTileMapEditorViewportClient::EndTransaction()
+{
+	if (bManipulationDirtiedSomething)
+	{
+		RenderTileMapComponent->TileMap->PostEditChange();
+	}
+	
+	bManipulationDirtiedSomething = false;
+
+	if (ScopedTransaction != NULL)
+	{
+		delete ScopedTransaction;
+		ScopedTransaction = NULL;
+	}
+}
+
+void FTileMapEditorViewportClient::NotifyTileMapBeingEditedHasChanged()
+{
+	//@TODO: Ideally we do this before switching
+	EndTransaction();
+	ClearSelectionSet();
+
+	// Update components to know about the new tile map being edited
+	UPaperTileMap* TileMap = GetTileMapBeingEdited();
+	RenderTileMapComponent->TileMap = TileMap;
+
+	//
+	bDeferZoomToTileMap = true;
+}
+
+void FTileMapEditorViewportClient::FocusOnTileMap()
+{
+	FocusViewportOnBox(RenderTileMapComponent->Bounds.GetBox());
+}
+
+void FTileMapEditorViewportClient::ClearSelectionSet()
+{
+}
+
+//////////////////////////////////////////////////////////////////////////
+
+#undef LOCTEXT_NAMESPACE
\ No newline at end of file
diff --git a/Engine/Plugins/Experimental/Paper2D/Source/Paper2DEditor/Private/TileMapEditing/TileMapEditorViewportClient.h b/Engine/Plugins/Experimental/Paper2D/Source/Paper2DEditor/Private/TileMapEditing/TileMapEditorViewportClient.h
new file mode 100644
--- /dev/null
+++ b/Engine/Plugins/Experimental/Paper2D/Source/Paper2DEditor/Private/TileMapEditing/TileMapEditorViewportClient.h
@@ -0,0 +1,96 @@
+// Copyright 1998-2014 Epic Games, Inc. All Rights Reserved.
+
+#pragma once
+
+#include "Paper2DEditorPrivatePCH.h"
+#include "TileMapEditor.h"
+#include "SSingleObjectDetailsPanel.h"
+#include "SceneViewport.h"
+#include "PaperEditorViewportClient.h"
+#include "PreviewScene.h"
+#include "ScopedTransaction.h"
+
+//////////////////////////////////////////////////////////////////////////
+// FTileMapEditorViewportClient
+
+class FTileMapEditorViewportClient : public FPaperEditorViewportClient
+{
+public:
+	/** Constructor */
+	FTileMapEditorViewportClient(TWeakPtr<FTileMapEditor> InTileMapEditor, TWeakPtr<class STileMapEditorViewport> InTileMapEditorViewportPtr);
+
+	// FViewportClient interface
+	virtual void Draw(FViewport* Viewport, FCanvas* Canvas) override;
+	virtual void Draw(const FSceneView* View, FPrimitiveDrawInterface* PDI);
+	virtual void DrawCanvas(FViewport& InViewport, FSceneView& View, FCanvas& Canvas) override;
+	virtual void Tick(float DeltaSeconds) override;
+	// End of FViewportClient interface
+
+	// FEditorViewportClient interface
+	virtual void UpdateMouseDelta() override;
+	virtual void ProcessClick(FSceneView& View, HHitProxy* HitProxy, FKey Key, EInputEvent Event, uint32 HitX, uint32 HitY) override;
+	virtual bool InputKey(FViewport* Viewport, int32 ControllerId, FKey Key, EInputEvent Event, float AmountDepressed, bool bGamepad) override;
+	virtual bool InputWidgetDelta(FViewport* Viewport, EAxisList::Type CurrentAxis, FVector& Drag, FRotator& Rot, FVector& Scale) override;
+	virtual void TrackingStarted(const struct FInputEventState& InInputState, bool bIsDragging, bool bNudge) override;
+	virtual void TrackingStopped() override;
+	virtual FWidget::EWidgetMode GetWidgetMode() const override;
+	virtual FVector GetWidgetLocation() const;
+	virtual FMatrix GetWidgetCoordSystem() const;
+	virtual ECoordSystem GetWidgetCoordSystemSpace() const;
+	// End of FEditorViewportClient interface
+
+	void ToggleShowPivot() { bShowPivot = !bShowPivot; Invalidate(); }
+	bool IsShowPivotChecked() const { return bShowPivot; }
+
+	void ToggleShowMeshEdges();
+	bool IsShowMeshEdgesChecked() const;
+
+	//
+	void FocusOnTileMap();
+
+	// Invalidate any references to the tile map being edited; it has changed
+	void NotifyTileMapBeingEditedHasChanged();
+private:
+	// The preview scene
+	FPreviewScene OwnedPreviewScene;
+
+	// Tile map editor that owns this viewport
+	TWeakPtr<FTileMapEditor> TileMapEditorPtr;
+
+	// Render component for the tile map being edited
+	UPaperTileMapRenderComponent* RenderTileMapComponent;
+
+	// Widget mode
+	FWidget::EWidgetMode WidgetMode;
+
+	// Are we currently manipulating something?
+	bool bManipulating;
+
+	// Did we dirty something during manipulation?
+	bool bManipulationDirtiedSomething;
+
+	// Pointer back to the tile map editor viewport control that owns us
+	TWeakPtr<class STileMapEditorViewport> TileMapEditorViewportPtr;
+
+	// The current transaction for undo/redo
+	class FScopedTransaction* ScopedTransaction;
+
+	// Should we show the sprite pivot?
+	bool bShowPivot;
+
+	// Should we zoom to the tile map next tick?
+	bool bDeferZoomToTileMap;
+private:
+	UPaperTileMap* GetTileMapBeingEdited() const
+	{
+		return TileMapEditorPtr.Pin()->GetTileMapBeingEdited();
+	}
+
+	void DrawTriangleList(FViewport& InViewport, FSceneView& View, FCanvas& Canvas, const TArray<FVector2D>& Triangles);
+	void DrawBoundsAsText(FViewport& InViewport, FSceneView& View, FCanvas& Canvas, int32& YPos);
+	
+	void BeginTransaction(const FText& SessionName);
+	void EndTransaction();
+
+	void ClearSelectionSet();
+};
