Paper2D: Add support for dragging tile map assets into the editor viewport to create tile map actors

diff --git a/Engine/Plugins/Experimental/Paper2D/Source/Paper2DEditor/Classes/TileMapActorFactory.h b/Engine/Plugins/Experimental/Paper2D/Source/Paper2DEditor/Classes/TileMapActorFactory.h
--- a/Engine/Plugins/Experimental/Paper2D/Source/Paper2DEditor/Classes/TileMapActorFactory.h
+++ b/Engine/Plugins/Experimental/Paper2D/Source/Paper2DEditor/Classes/TileMapActorFactory.h
@@ -8,6 +8,8 @@ class UTileMapActorFactory : public UActorFactory
 	GENERATED_UCLASS_BODY()
 
 	// UActorFactory interface
+	virtual void PostSpawnActor(UObject* Asset, AActor* NewActor) override;
+	virtual void PostCreateBlueprint(UObject* Asset, AActor* CDO) override;
 	virtual bool CanCreateActorFrom(const FAssetData& AssetData, FText& OutErrorMsg) override;
 	// End of UActorFactory interface
 };
diff --git a/Engine/Plugins/Experimental/Paper2D/Source/Paper2DEditor/Private/TileMapEditing/TileMapActorFactory.cpp b/Engine/Plugins/Experimental/Paper2D/Source/Paper2DEditor/Private/TileMapEditing/TileMapActorFactory.cpp
--- a/Engine/Plugins/Experimental/Paper2D/Source/Paper2DEditor/Private/TileMapEditing/TileMapActorFactory.cpp
+++ b/Engine/Plugins/Experimental/Paper2D/Source/Paper2DEditor/Private/TileMapEditing/TileMapActorFactory.cpp
@@ -13,14 +13,53 @@ UTileMapActorFactory::UTileMapActorFactory(const FPostConstructInitializePropert
 	NewActorClass = APaperTileMapActor::StaticClass();
 }
 
+void UTileMapActorFactory::PostSpawnActor(UObject* Asset, AActor* NewActor)
+{
+	if (UPaperTileMap* TileMap = Cast<UPaperTileMap>(Asset))
+	{
+		GEditor->SetActorLabelUnique(NewActor, TileMap->GetName());
+
+		APaperTileMapActor* TypedActor = CastChecked<APaperTileMapActor>(NewActor);
+		UPaperTileMapRenderComponent* RenderComponent = TypedActor->RenderComponent;
+		check(RenderComponent);
+
+		RenderComponent->UnregisterComponent();
+		RenderComponent->TileMap = TileMap;
+		RenderComponent->RegisterComponent();
+	}
+}
+
+void UTileMapActorFactory::PostCreateBlueprint(UObject* Asset, AActor* CDO)
+{
+	if (UPaperTileMap* TileMap = Cast<UPaperTileMap>(Asset))
+	{
+		if (APaperTileMapActor* TypedActor = Cast<APaperTileMapActor>(CDO))
+		{
+			UPaperTileMapRenderComponent* RenderComponent = TypedActor->RenderComponent;
+			check(RenderComponent);
+
+			RenderComponent->TileMap = TileMap;
+		}
+	}
+}
+
 bool UTileMapActorFactory::CanCreateActorFrom(const FAssetData& AssetData, FText& OutErrorMsg)
 {
 	if (GetDefault<UPaperRuntimeSettings>()->bEnableTileMapEditing)
 	{
-		return Super::CanCreateActorFrom(AssetData, OutErrorMsg);
+		if (AssetData.IsValid() && AssetData.GetClass()->IsChildOf(UPaperTileMap::StaticClass()))
+		{
+			return true;
+		}
+		else
+		{
+			OutErrorMsg = NSLOCTEXT("Paper2D", "CanCreateActorFrom_NoTileMap", "No tile map was specified.");
+			return false;
+		}
 	}
 	else
 	{
+		OutErrorMsg = NSLOCTEXT("Paper2D", "CanCreateActorFrom_Disabled", "Tile map support is disabled in the Paper2D settings.");
 		return false;
 	}
 }
