Paper2D: Add display text for tile map render bounds to tile map editor

diff --git a/Engine/Plugins/Experimental/Paper2D/Source/Paper2DEditor/Private/TileMapEditing/TileMapEditorViewportClient.cpp b/Engine/Plugins/Experimental/Paper2D/Source/Paper2DEditor/Private/TileMapEditing/TileMapEditorViewportClient.cpp
--- a/Engine/Plugins/Experimental/Paper2D/Source/Paper2DEditor/Private/TileMapEditing/TileMapEditorViewportClient.cpp
+++ b/Engine/Plugins/Experimental/Paper2D/Source/Paper2DEditor/Private/TileMapEditing/TileMapEditorViewportClient.cpp
@@ -51,7 +51,7 @@ void FTileMapEditorViewportClient::DrawBoundsAsText(FViewport& InViewport, FScen
 	NoDigitGroupingFormat.UseGrouping = false;
 
 	UPaperTileMap* TileMap = GetTileMapBeingEdited();
-	FBoxSphereBounds Bounds;//@TODO: = TileMap->GetRenderBounds();
+	FBoxSphereBounds Bounds = TileMap->GetRenderBounds();
 
 	const FText DisplaySizeText = FText::Format(LOCTEXT("BoundsSize", "Approx. Size: {0}x{1}x{2}"),
 		FText::AsNumber((int32)(Bounds.BoxExtent.X * 2.0f), &NoDigitGroupingFormat),
@@ -83,6 +83,9 @@ void FTileMapEditorViewportClient::DrawCanvas(FViewport& Viewport, FSceneView& V
 	}
 
 	int32 YPos = 42;
+
+	// Draw the render bounds
+	DrawBoundsAsText(Viewport, View, Canvas, /*inout*/ YPos);
 }
 
 void FTileMapEditorViewportClient::Draw(const FSceneView* View, FPrimitiveDrawInterface* PDI)
