Paper2D: Add a proxy specifically for UPaperSpriteComponent and fix flipped normals on lit sprites
- Add FPaperSpriteSceneProxy to pair with UPaperSpriteComponent, removing some code from the common FPaperRenderSceneProxy parent
- Fix broken tangent calculations in FPaperRenderSceneProxy that caused bogus lighting on lit sprites

diff --git a/Engine/Plugins/Experimental/Paper2D/Source/Paper2D/Classes/PaperSpriteComponent.h b/Engine/Plugins/Experimental/Paper2D/Source/Paper2D/Classes/PaperSpriteComponent.h
--- a/Engine/Plugins/Experimental/Paper2D/Source/Paper2D/Classes/PaperSpriteComponent.h
+++ b/Engine/Plugins/Experimental/Paper2D/Source/Paper2D/Classes/PaperSpriteComponent.h
@@ -65,7 +65,7 @@ public:
 	// End of UPrimitiveComponent interface
 
 protected:
-	friend class FPaperRenderSceneProxy;
+	friend class FPaperSpriteSceneProxy;
 };
 
 // Allow the old name to continue to work for one release
diff --git a/Engine/Plugins/Experimental/Paper2D/Source/Paper2D/Private/PaperRenderSceneProxy.cpp b/Engine/Plugins/Experimental/Paper2D/Source/Paper2D/Private/PaperRenderSceneProxy.cpp
--- a/Engine/Plugins/Experimental/Paper2D/Source/Paper2D/Private/PaperRenderSceneProxy.cpp
+++ b/Engine/Plugins/Experimental/Paper2D/Source/Paper2D/Private/PaperRenderSceneProxy.cpp
@@ -7,30 +7,24 @@
 #include "Rendering/PaperBatchSceneProxy.h"
 #include "PhysicsEngine/BodySetup2D.h"
 
-static FPackedNormal PackedNormalX(FVector(1.0f, 0.0f, 0.0f));
-static FPackedNormal PackedNormalY(FVector(0.0f, 1.0f, 0.0f));
+//////////////////////////////////////////////////////////////////////////
+// FPaperSpriteVertex
+
+FPackedNormal FPaperSpriteVertex::PackedNormalX(FVector(1.0f, 0.0f, 0.0f));
+FPackedNormal FPaperSpriteVertex::PackedNormalZ(FVector(0.0f, -1.0f, 0.0f));
 
-/** A material sprite vertex. */
-struct FPaperSpriteVertex
+void FPaperSpriteVertex::SetTangentsFromPaperAxes()
 {
-	FVector Position;
-	FPackedNormal TangentX;
-	FPackedNormal TangentZ;
-	FLinearColor Color;
-	FVector2D TexCoords;
-
-	FPaperSpriteVertex() {}
-
-	FPaperSpriteVertex(const FVector& InPosition, const FVector2D& InTextureCoordinate, const FLinearColor& InColor)
-		: Position(InPosition)
-		, TangentX(PackedNormalX)
-		, TangentZ(PackedNormalY)
-		, Color(InColor)
-		, TexCoords(InTextureCoordinate)
-	{}
-};
+	PackedNormalX = PaperAxisX;
+	PackedNormalZ = PaperAxisZ;
+	// store determinant of basis in w component of normal vector
+	PackedNormalZ.Vector.W = (GetBasisDeterminantSign(PaperAxisX, PaperAxisY, PaperAxisZ) < 0.0f) ? 0 : 255;
+}
+
+//////////////////////////////////////////////////////////////////////////
+// FPaperSpriteVertexBuffer
 
-/** A dummy vertex buffer used to give the FMaterialSpriteVertexFactory something to reference as a stream source. */
+/** A dummy vertex buffer used to give the FPaperSpriteVertexFactory something to reference as a stream source. */
 class FPaperSpriteVertexBuffer : public FVertexBuffer
 {
 public:
@@ -164,30 +158,9 @@ FPaperRenderSceneProxy::FPaperRenderSceneProxy(const UPrimitiveComponent* InComp
 	: FPrimitiveSceneProxy(InComponent)
 	, Material(NULL)
 	, Owner(InComponent->GetOwner())
-	, SourceSprite(NULL)
 	, WireframeColor(FLinearColor::White)
 	, CollisionResponse(InComponent->GetCollisionResponseToChannels())
 {
-	if (const UPaperSpriteComponent* RenderComp = Cast<const UPaperSpriteComponent>(InComponent))
-	{
-		WireframeColor = RenderComp->GetWireframeColor();
-
-		SourceSprite = RenderComp->SourceSprite; //@TODO: This is totally not threadsafe, and won't keep up to date if the actor's sprite changes, etc....
-		if (SourceSprite)
-		{
-			Material = SourceSprite->GetDefaultMaterial();
-		}
-		
-		if (RenderComp->MaterialOverride)
-		{
-			Material = RenderComp->MaterialOverride;
-		}
-
-		if (Material)
-		{
-			MaterialRelevance = Material->GetRelevance();
-		}
-	}
 }
 
 FPaperRenderSceneProxy::~FPaperRenderSceneProxy()
@@ -216,57 +189,6 @@ void FPaperRenderSceneProxy::DrawDynamicElements(FPrimitiveDrawInterface* PDI, c
 
 	checkSlow(IsInRenderingThread());
 
-	if (SourceSprite != NULL)
-	{
-		// Show 3D physics
-		if ((View->Family->EngineShowFlags.Collision /*@TODO: && bIsCollisionEnabled*/) && AllowDebugViewmodes())
-		{
-			if (UBodySetup2D* BodySetup2D = Cast<UBodySetup2D>(SourceSprite->BodySetup))
-			{
-				//@TODO: Draw 2D debugging geometry
-			}
-			else if (UBodySetup* BodySetup = SourceSprite->BodySetup)
-			{
-				if (FMath::Abs(GetLocalToWorld().Determinant()) < SMALL_NUMBER)
-				{
-					// Catch this here or otherwise GeomTransform below will assert
-					// This spams so commented out
-					//UE_LOG(LogStaticMesh, Log, TEXT("Zero scaling not supported (%s)"), *StaticMesh->GetPathName());
-				}
-				else
-				{
-					// Make a material for drawing solid collision stuff
-					const UMaterial* LevelColorationMaterial = View->Family->EngineShowFlags.Lighting 
-						? GEngine->ShadedLevelColorationLitMaterial : GEngine->ShadedLevelColorationUnlitMaterial;
-
-					const FColoredMaterialRenderProxy CollisionMaterialInstance(
-						LevelColorationMaterial->GetRenderProxy(IsSelected(), IsHovered()),
-						WireframeColor
-						);
-
-					// Draw the static mesh's body setup.
-
-					// Get transform without scaling.
-					FTransform GeomTransform(GetLocalToWorld());
-
-					// In old wireframe collision mode, always draw the wireframe highlighted (selected or not).
-					bool bDrawWireSelected = IsSelected();
-					if (View->Family->EngineShowFlags.Collision)
-					{
-						bDrawWireSelected = true;
-					}
-
-					// Differentiate the color based on bBlockNonZeroExtent.  Helps greatly with skimming a level for optimization opportunities.
-					const FColor CollisionColor = FColor(220,149,223,255);
-
-					const bool bUseSeparateColorPerHull = (Owner == NULL);
-					const bool bDrawSolid = false;
-					BodySetup->AggGeom.DrawAggGeom(PDI, GeomTransform, GetSelectionColor(CollisionColor, bDrawWireSelected, IsHovered()), &CollisionMaterialInstance, bUseSeparateColorPerHull, bDrawSolid);
-				}
-			}
-		}
-	}
-
 #if !(UE_BUILD_SHIPPING || UE_BUILD_TEST)
 	if (View->Family->EngineShowFlags.Paper2DSprites)
 	{
@@ -299,6 +221,13 @@ void FPaperRenderSceneProxy::DrawDynamicElements(FPrimitiveDrawInterface* PDI, c
 
 FVertexFactory* FPaperRenderSceneProxy::GetPaperSpriteVertexFactory() const
 {
+	static bool bInitAxes = false;
+	if (!bInitAxes)
+	{
+		bInitAxes = true;
+		FPaperSpriteVertex::SetTangentsFromPaperAxes();
+	}
+
 	static TGlobalResource<FPaperSpriteVertexFactory> GPaperSpriteVertexFactory;
 	return &GPaperSpriteVertexFactory;
 }
diff --git a/Engine/Plugins/Experimental/Paper2D/Source/Paper2D/Private/PaperRenderSceneProxy.h b/Engine/Plugins/Experimental/Paper2D/Source/Paper2D/Private/PaperRenderSceneProxy.h
--- a/Engine/Plugins/Experimental/Paper2D/Source/Paper2D/Private/PaperRenderSceneProxy.h
+++ b/Engine/Plugins/Experimental/Paper2D/Source/Paper2D/Private/PaperRenderSceneProxy.h
@@ -4,6 +4,34 @@
 
 #define TEST_BATCHING 0
 
+//////////////////////////////////////////////////////////////////////////
+// FPaperSpriteVertex
+
+/** A Paper2D sprite vertex. */
+struct FPaperSpriteVertex
+{
+	FVector Position;
+	FPackedNormal TangentX;
+	FPackedNormal TangentZ;
+	FLinearColor Color;
+	FVector2D TexCoords;
+
+	FPaperSpriteVertex() {}
+
+	FPaperSpriteVertex(const FVector& InPosition, const FVector2D& InTextureCoordinate, const FLinearColor& InColor)
+		: Position(InPosition)
+		, TangentX(PackedNormalX)
+		, TangentZ(PackedNormalZ)
+		, Color(InColor)
+		, TexCoords(InTextureCoordinate)
+	{}
+
+	static void SetTangentsFromPaperAxes();
+
+	static FPackedNormal PackedNormalX;
+	static FPackedNormal PackedNormalZ;
+};
+
 //////////////////////////////////////////////////////////////////////////
 // FPaperRenderSceneProxy
 
@@ -46,7 +74,6 @@ protected:
 	class UMaterialInterface* Material;
 	FVector Origin;
 	AActor* Owner;
-	const UPaperSprite* SourceSprite;
 
 	// The color to draw as in wireframe mode
 	FLinearColor WireframeColor;
diff --git a/Engine/Plugins/Experimental/Paper2D/Source/Paper2D/Private/PaperSpriteComponent.cpp b/Engine/Plugins/Experimental/Paper2D/Source/Paper2D/Private/PaperSpriteComponent.cpp
--- a/Engine/Plugins/Experimental/Paper2D/Source/Paper2D/Private/PaperSpriteComponent.cpp
+++ b/Engine/Plugins/Experimental/Paper2D/Source/Paper2D/Private/PaperSpriteComponent.cpp
@@ -1,7 +1,7 @@
 // Copyright 1998-2014 Epic Games, Inc. All Rights Reserved.
 
 #include "Paper2DPrivatePCH.h"
-#include "PaperRenderSceneProxy.h"
+#include "PaperSpriteSceneProxy.h"
 #include "PaperSpriteComponent.h"
 #include "PhysicsEngine/BodySetup2D.h"
 
@@ -29,7 +29,7 @@ void UPaperSpriteComponent::PostEditChangeProperty(FPropertyChangedEvent& Proper
 
 FPrimitiveSceneProxy* UPaperSpriteComponent::CreateSceneProxy()
 {
-	FPaperRenderSceneProxy* NewProxy = new FPaperRenderSceneProxy(this);
+	FPaperSpriteSceneProxy* NewProxy = new FPaperSpriteSceneProxy(this);
 	FSpriteDrawCallRecord DrawCall;
 	DrawCall.BuildFromSprite(SourceSprite);
 	DrawCall.Color = SpriteColor;
diff --git a/Engine/Plugins/Experimental/Paper2D/Source/Paper2D/Private/PaperSpriteSceneProxy.cpp b/Engine/Plugins/Experimental/Paper2D/Source/Paper2D/Private/PaperSpriteSceneProxy.cpp
new file mode 100644
--- /dev/null
+++ b/Engine/Plugins/Experimental/Paper2D/Source/Paper2D/Private/PaperSpriteSceneProxy.cpp
@@ -0,0 +1,86 @@
+// Copyright 1998-2014 Epic Games, Inc. All Rights Reserved.
+
+#include "Paper2DPrivatePCH.h"
+#include "PaperSpriteSceneProxy.h"
+
+//////////////////////////////////////////////////////////////////////////
+// FPaperSpriteSceneProxy
+
+FPaperSpriteSceneProxy::FPaperSpriteSceneProxy(const UPaperSpriteComponent* InComponent)
+	: FPaperRenderSceneProxy(InComponent)
+	, SourceSprite(nullptr)
+{
+	WireframeColor = InComponent->GetWireframeColor();
+	SourceSprite = InComponent->SourceSprite; //@TODO: This is totally not threadsafe, and won't keep up to date if the actor's sprite changes, etc....
+	
+	if (SourceSprite)
+	{
+		Material = SourceSprite->GetDefaultMaterial();
+	}
+		
+	if (InComponent->MaterialOverride)
+	{
+		Material = InComponent->MaterialOverride;
+	}
+
+	if (Material)
+	{
+		MaterialRelevance = Material->GetRelevance();
+	}
+}
+
+void FPaperSpriteSceneProxy::DrawDynamicElements(FPrimitiveDrawInterface* PDI, const FSceneView* View)
+{
+	if (SourceSprite != nullptr)
+	{
+		// Show 3D physics
+		if ((View->Family->EngineShowFlags.Collision /*@TODO: && bIsCollisionEnabled*/) && AllowDebugViewmodes())
+		{
+			if (UBodySetup2D* BodySetup2D = Cast<UBodySetup2D>(SourceSprite->BodySetup))
+			{
+				//@TODO: Draw 2D debugging geometry
+			}
+			else if (UBodySetup* BodySetup = SourceSprite->BodySetup)
+			{
+				if (FMath::Abs(GetLocalToWorld().Determinant()) < SMALL_NUMBER)
+				{
+					// Catch this here or otherwise GeomTransform below will assert
+					// This spams so commented out
+					//UE_LOG(LogStaticMesh, Log, TEXT("Zero scaling not supported (%s)"), *StaticMesh->GetPathName());
+				}
+				else
+				{
+					// Make a material for drawing solid collision stuff
+					const UMaterial* LevelColorationMaterial = View->Family->EngineShowFlags.Lighting 
+						? GEngine->ShadedLevelColorationLitMaterial : GEngine->ShadedLevelColorationUnlitMaterial;
+
+					const FColoredMaterialRenderProxy CollisionMaterialInstance(
+						LevelColorationMaterial->GetRenderProxy(IsSelected(), IsHovered()),
+						WireframeColor
+						);
+
+					// Draw the static mesh's body setup.
+
+					// Get transform without scaling.
+					FTransform GeomTransform(GetLocalToWorld());
+
+					// In old wireframe collision mode, always draw the wireframe highlighted (selected or not).
+					bool bDrawWireSelected = IsSelected();
+					if (View->Family->EngineShowFlags.Collision)
+					{
+						bDrawWireSelected = true;
+					}
+
+					// Differentiate the color based on bBlockNonZeroExtent.  Helps greatly with skimming a level for optimization opportunities.
+					const FColor CollisionColor = FColor(220,149,223,255);
+
+					const bool bUseSeparateColorPerHull = (Owner == NULL);
+					const bool bDrawSolid = false;
+					BodySetup->AggGeom.DrawAggGeom(PDI, GeomTransform, GetSelectionColor(CollisionColor, bDrawWireSelected, IsHovered()), &CollisionMaterialInstance, bUseSeparateColorPerHull, bDrawSolid);
+				}
+			}
+		}
+	}
+
+	FPaperRenderSceneProxy::DrawDynamicElements(PDI, View);
+}
diff --git a/Engine/Plugins/Experimental/Paper2D/Source/Paper2D/Private/PaperSpriteSceneProxy.h b/Engine/Plugins/Experimental/Paper2D/Source/Paper2D/Private/PaperSpriteSceneProxy.h
new file mode 100644
--- /dev/null
+++ b/Engine/Plugins/Experimental/Paper2D/Source/Paper2D/Private/PaperSpriteSceneProxy.h
@@ -0,0 +1,21 @@
+// Copyright 1998-2014 Epic Games, Inc. All Rights Reserved.
+
+#pragma once
+
+#include "PaperRenderSceneProxy.h"
+
+//////////////////////////////////////////////////////////////////////////
+// FPaperSpriteSceneProxy
+
+class FPaperSpriteSceneProxy : public FPaperRenderSceneProxy
+{
+public:
+	FPaperSpriteSceneProxy(const UPaperSpriteComponent* InComponent);
+
+	// FPrimitiveSceneProxy interface
+	virtual void DrawDynamicElements(FPrimitiveDrawInterface* PDI, const FSceneView* View) override;
+	// End of FPrimitiveSceneProxy interface
+
+protected:
+	const UPaperSprite* SourceSprite;
+};
