Paper2D: Make tile set viewer/editor draw translucent/masked textures as translucent

diff --git a/Engine/Plugins/Experimental/Paper2D/Source/Paper2DEditor/Private/TileSetEditor.cpp b/Engine/Plugins/Experimental/Paper2D/Source/Paper2DEditor/Private/TileSetEditor.cpp
--- a/Engine/Plugins/Experimental/Paper2D/Source/Paper2DEditor/Private/TileSetEditor.cpp
+++ b/Engine/Plugins/Experimental/Paper2D/Source/Paper2DEditor/Private/TileSetEditor.cpp
@@ -75,6 +75,8 @@ void FTileSetEditorViewportClient::Draw(FViewport* Viewport, FCanvas* Canvas)
 
 	if (Texture != NULL)
 	{
+		const bool bUseTranslucentBlend = Texture->HasAlphaChannel();
+
 		// Fully stream in the texture before drawing it.
 		Texture->SetForceMipLevelsToBeResident(30.0f);
 		Texture->WaitForStreaming();
@@ -99,7 +101,7 @@ void FTileSetEditorViewportClient::Draw(FViewport* Viewport, FCanvas* Canvas)
 		const float Width = Texture->GetSurfaceWidth() * ZoomAmount;
 		const float Height = Texture->GetSurfaceHeight() * ZoomAmount;
 
-		Canvas->DrawTile(XPos, YPos, Width, Height, 0.0f, 0.0f, 1.0f, 1.0f, TextureDrawColor, Texture->Resource, SE_BLEND_Opaque);
+		Canvas->DrawTile(XPos, YPos, Width, Height, 0.0f, 0.0f, 1.0f, 1.0f, TextureDrawColor, Texture->Resource, bUseTranslucentBlend);
 	}
 
 	// Overlay the selection rectangles
