Fixing compile error (no editor)

diff --git a/Engine/Plugins/Experimental/Paper2D/Source/Paper2D/Private/PaperTileMapRenderSceneProxy.cpp b/Engine/Plugins/Experimental/Paper2D/Source/Paper2D/Private/PaperTileMapRenderSceneProxy.cpp
--- a/Engine/Plugins/Experimental/Paper2D/Source/Paper2D/Private/PaperTileMapRenderSceneProxy.cpp
+++ b/Engine/Plugins/Experimental/Paper2D/Source/Paper2D/Private/PaperTileMapRenderSceneProxy.cpp
@@ -157,12 +157,12 @@ void FPaperTileMapRenderSceneProxy::DrawDynamicElements(FPrimitiveDrawInterface*
 			{
 				continue;
 			}
-
+#if WITH_EDITORONLY_DATA
 			if (Layer->bHiddenInEditor)
 			{
 				continue;
 			}
-
+#endif
 			if (Layer->bCollisionLayer && View->bIsGameView)
 			{
 				continue;
