VS2012 fix for initializer lists in PaperTiledImporterFactory.cpp

diff --git a/Engine/Plugins/Experimental/Paper2D/Source/PaperTiledImporter/Private/PaperTiledImporterFactory.cpp b/Engine/Plugins/Experimental/Paper2D/Source/PaperTiledImporter/Private/PaperTiledImporterFactory.cpp
--- a/Engine/Plugins/Experimental/Paper2D/Source/PaperTiledImporter/Private/PaperTiledImporterFactory.cpp
+++ b/Engine/Plugins/Experimental/Paper2D/Source/PaperTiledImporter/Private/PaperTiledImporterFactory.cpp
@@ -420,11 +420,11 @@ void UPaperTiledImporterFactory::ParseGlobalInfoFromJSON(TSharedPtr<FJsonObject>
 
 	// Parse all of the integer fields
 	FRequiredIntField IntFields[] = {
-		{ OutParsedInfo.FileVersion, TEXT("version") },
-		{ OutParsedInfo.Width, TEXT("width") },
-		{ OutParsedInfo.Height, TEXT("height") },
-		{ OutParsedInfo.TileWidth, TEXT("tilewidth") },
-		{ OutParsedInfo.TileHeight, TEXT("tileheight") }
+		FRequiredIntField( OutParsedInfo.FileVersion, TEXT("version") ),
+		FRequiredIntField( OutParsedInfo.Width, TEXT("width") ),
+		FRequiredIntField( OutParsedInfo.Height, TEXT("height") ),
+		FRequiredIntField( OutParsedInfo.TileWidth, TEXT("tilewidth") ),
+		FRequiredIntField( OutParsedInfo.TileHeight, TEXT("tileheight") )
 	};
 	bSuccessfullyParsed = bSuccessfullyParsed && ParseIntegerFields(IntFields, ARRAY_COUNT(IntFields), Tree, NameForErrors, bSilent);
 
@@ -513,13 +513,13 @@ void FTileSetFromTiled::ParseFromJSON(TSharedPtr<FJsonObject> Tree, const FStrin
 
 	// Parse all of the integer fields
 	FRequiredIntField IntFields[] = {
-		{ FirstGID, TEXT("firstgid"), 1 },
-		{ ImageWidth, TEXT("imagewidth"), 1 },
-		{ ImageHeight, TEXT("imageheight"), 1 },
-		{ Margin, TEXT("margin"), 0 },
-		{ Spacing, TEXT("spacing"), 0 },
-		{ TileWidth, TEXT("tilewidth"), 1 },
-		{ TileHeight, TEXT("tileheight"), 1 }
+		FRequiredIntField( FirstGID, TEXT("firstgid"), 1 ),
+		FRequiredIntField( ImageWidth, TEXT("imagewidth"), 1 ),
+		FRequiredIntField( ImageHeight, TEXT("imageheight"), 1 ),
+		FRequiredIntField( Margin, TEXT("margin"), 0 ),
+		FRequiredIntField( Spacing, TEXT("spacing"), 0 ),
+		FRequiredIntField( TileWidth, TEXT("tilewidth"), 1 ),
+		FRequiredIntField( TileHeight, TEXT("tileheight"), 1 )
 	};
 
 	bSuccessfullyParsed = bSuccessfullyParsed && ParseIntegerFields(IntFields, ARRAY_COUNT(IntFields), Tree, NameForErrors, bSilent);
@@ -584,10 +584,10 @@ void FTileLayerFromTiled::ParseFromJSON(TSharedPtr<FJsonObject> Tree, const FStr
 
 	// Parse all of the integer fields
 	FRequiredIntField IntFields[] = {
-		{ Width, TEXT("width"), 1 },
-		{ Height, TEXT("height"), 1 },
-		{ OffsetX, TEXT("x"), 0 },
-		{ OffsetY, TEXT("y"), 0 }
+		FRequiredIntField( Width, TEXT("width"), 1 ),
+		FRequiredIntField( Height, TEXT("height"), 1 ),
+		FRequiredIntField( OffsetX, TEXT("x"), 0 ),
+		FRequiredIntField( OffsetY, TEXT("y"), 0 )
 	};
 
 	bSuccessfullyParsed = bSuccessfullyParsed && ParseIntegerFields(IntFields, ARRAY_COUNT(IntFields), Tree, NameForErrors, bSilent);
