CIS fix.

diff --git a/Engine/Plugins/Experimental/Paper2D/Source/Paper2D/Private/PaperSpriteSceneProxy.cpp b/Engine/Plugins/Experimental/Paper2D/Source/Paper2D/Private/PaperSpriteSceneProxy.cpp
--- a/Engine/Plugins/Experimental/Paper2D/Source/Paper2D/Private/PaperSpriteSceneProxy.cpp
+++ b/Engine/Plugins/Experimental/Paper2D/Source/Paper2D/Private/PaperSpriteSceneProxy.cpp
@@ -2,6 +2,7 @@
 
 #include "Paper2DPrivatePCH.h"
 #include "PaperSpriteSceneProxy.h"
+#include "PhysicsEngine/BodySetup2D.h"
 
 //////////////////////////////////////////////////////////////////////////
 // FPaperSpriteSceneProxy
