Paper2D: Add some support for replicating state of flipbook components

diff --git a/Engine/Plugins/Experimental/Paper2D/Source/Paper2D/Classes/PaperFlipbookComponent.h b/Engine/Plugins/Experimental/Paper2D/Source/Paper2D/Classes/PaperFlipbookComponent.h
--- a/Engine/Plugins/Experimental/Paper2D/Source/Paper2D/Classes/PaperFlipbookComponent.h
+++ b/Engine/Plugins/Experimental/Paper2D/Source/Paper2D/Classes/PaperFlipbookComponent.h
@@ -5,33 +5,48 @@
 #include "PaperFlipbook.h"
 #include "PaperFlipbookComponent.generated.h"
 
-UCLASS(ShowCategories=(Mobility), EarlyAccessPreview, meta=(BlueprintSpawnableComponent))
+UCLASS(ShowCategories=(Mobility, ComponentReplication), EarlyAccessPreview, meta=(BlueprintSpawnableComponent))
 class PAPER2D_API UPaperFlipbookComponent : public UPrimitiveComponent
 {
 	GENERATED_UCLASS_BODY()
 
 protected:
-	UPROPERTY(Category=Sprite, EditAnywhere, meta=(DisplayThumbnail = "true"))
+	/** Flipbook currently being played */
+	UPROPERTY(Category=Sprite, EditAnywhere, meta=(DisplayThumbnail = "true"), ReplicatedUsing=OnRep_SourceFlipbook)
 	UPaperFlipbook* SourceFlipbook;
 
+	/** Material used to display the flipbook */
 	UPROPERTY(Category=Sprite, EditAnywhere)
 	UMaterialInterface* Material;
 
-	UPROPERTY(Category=Sprite, EditAnywhere, BlueprintReadWrite)
+	/** Current play rate of the flipbook */
+	UPROPERTY(Category=Sprite, EditAnywhere, Replicated)
 	float PlayRate;
 
 	/** Whether the flipbook should loop when it reaches the end, or stop */
-	UPROPERTY()
+	UPROPERTY(Replicated)
 	uint32 bLooping:1;
 
 	/** If playback should move the current position backwards instead of forwards */
-	UPROPERTY()
+	UPROPERTY(Replicated)
 	uint32 bReversePlayback:1;
 
 	/** Are we currently playing (moving Position) */
-	UPROPERTY()
+	UPROPERTY(Replicated)
 	uint32 bPlaying:1;
 
+	/** Current position in the timeline */
+	UPROPERTY(Replicated)
+	float AccumulatedTime;
+
+	/** Last frame index calculated */
+	UPROPERTY()
+	int32 CachedFrameIndex;
+
+	/** Vertex color to apply to the frames */
+	UPROPERTY(BlueprintReadOnly, Interp, Category=Sprite)
+	FLinearColor SpriteColor;
+
 public:
 	/** Change the flipbook used by this instance. */
 	UFUNCTION(BlueprintCallable, Category="Sprite")
@@ -118,19 +133,9 @@ public:
 	float GetFlipbookFramerate() const;
 
 protected:
-	/** Current position in the timeline */
-	UPROPERTY()
-	float AccumulatedTime;
-
-	/** Last frame index calculated */
-	UPROPERTY()
-	int32 CachedFrameIndex;
-
-	/* Vertex color to apply to the frames */
-	UPROPERTY(BlueprintReadOnly, Interp, Category=Sprite)
-	FLinearColor SpriteColor;
+	UFUNCTION()
+	void OnRep_SourceFlipbook(class UPaperFlipbook* OldFlipbook);
 
-protected:
 	void CalculateCurrentFrame();
 	UPaperSprite* GetSpriteAtCachedIndex() const;
 
diff --git a/Engine/Plugins/Experimental/Paper2D/Source/Paper2D/Private/PaperFlipbookComponent.cpp b/Engine/Plugins/Experimental/Paper2D/Source/Paper2D/Private/PaperFlipbookComponent.cpp
--- a/Engine/Plugins/Experimental/Paper2D/Source/Paper2D/Private/PaperFlipbookComponent.cpp
+++ b/Engine/Plugins/Experimental/Paper2D/Source/Paper2D/Private/PaperFlipbookComponent.cpp
@@ -3,6 +3,7 @@
 #include "Paper2DPrivatePCH.h"
 #include "PaperFlipbookSceneProxy.h"
 #include "PaperFlipbookComponent.h"
+#include "Runtime/Engine/Public/Net/UnrealNetwork.h"
 
 //////////////////////////////////////////////////////////////////////////
 // UPaperFlipbookComponent
@@ -189,6 +190,25 @@ void UPaperFlipbookComponent::TickFlipbook(float DeltaTime)
 #endif
 }
 
+void UPaperFlipbookComponent::GetLifetimeReplicatedProps(TArray< FLifetimeProperty > & OutLifetimeProps) const
+{
+	Super::GetLifetimeReplicatedProps(OutLifetimeProps);
+
+	DOREPLIFETIME(UPaperFlipbookComponent, SourceFlipbook);
+}
+
+void UPaperFlipbookComponent::OnRep_SourceFlipbook(class UPaperFlipbook* OldFlipbook)
+{
+	if (OldFlipbook != SourceFlipbook)
+	{
+		// Force SetFlipbook to change the animation (by default it won't change)
+		UPaperFlipbook* NewFlipbook = SourceFlipbook;
+		SourceFlipbook = nullptr;
+
+		SetFlipbook(NewFlipbook);
+	}
+}
+
 void UPaperFlipbookComponent::TickComponent(float DeltaTime, enum ELevelTick TickType, FActorComponentTickFunction* ThisTickFunction)
 {
 	// Advance time
