Paper2D: Convert UPaperSpriteComponent to derive from UMeshComponent

diff --git a/Engine/Plugins/Experimental/Paper2D/Source/Paper2D/Classes/PaperSpriteComponent.h b/Engine/Plugins/Experimental/Paper2D/Source/Paper2D/Classes/PaperSpriteComponent.h
--- a/Engine/Plugins/Experimental/Paper2D/Source/Paper2D/Classes/PaperSpriteComponent.h
+++ b/Engine/Plugins/Experimental/Paper2D/Source/Paper2D/Classes/PaperSpriteComponent.h
@@ -7,7 +7,7 @@
 #include "PaperSpriteComponent.generated.h"
 
 UCLASS(MinimalAPI, ShowCategories=(Mobility), EarlyAccessPreview, meta=(BlueprintSpawnableComponent))
-class UPaperSpriteComponent : public UPrimitiveComponent
+class UPaperSpriteComponent : public UMeshComponent
 {
 	GENERATED_UCLASS_BODY()
 
@@ -16,9 +16,9 @@ protected:
 	UPROPERTY(Category=Sprite, EditAnywhere, BlueprintReadOnly, meta=(DisplayThumbnail = "true"))
 	UPaperSprite* SourceSprite;
 
-	// The material override for this sprite component (if any)
-	UPROPERTY(Category=Sprite, EditAnywhere, BlueprintReadOnly)
-	UMaterialInterface* MaterialOverride;
+	// DEPRECATED in 4.4: The material override for this sprite component (if any); replaced by the Materials array inherited from UMeshComponent
+	UPROPERTY()
+	UMaterialInterface* MaterialOverride_DEPRECATED;
 
 	// The color of the sprite (passed to the sprite material as a vertex color)
 	UPROPERTY(BlueprintReadOnly, Interp, Category=Sprite)
@@ -44,6 +44,10 @@ public:
 	// UObject interface
 #if WITH_EDITOR
 	virtual void PostEditChangeProperty(FPropertyChangedEvent& PropertyChangedEvent) override;
+#endif
+#if WITH_EDITORONLY_DATA
+	virtual void Serialize(FArchive& Ar) override;
+	virtual void PostLoad() override;
 #endif
 	// End of UObject interface
 
@@ -62,6 +66,11 @@ public:
 	virtual FPrimitiveSceneProxy* CreateSceneProxy() override;
 	virtual FBoxSphereBounds CalcBounds(const FTransform & LocalToWorld) const override;
 	virtual class UBodySetup* GetBodySetup() override;
+	virtual void GetUsedTextures(TArray<UTexture*>& OutTextures, EMaterialQualityLevel::Type QualityLevel) override;
+	virtual UMaterialInterface* GetMaterial(int32 MaterialIndex) const override;
+	virtual void GetUsedMaterials(TArray<UMaterialInterface*>& OutMaterials) const override;
+	virtual void GetStreamingTextureInfo(TArray<FStreamingTexturePrimitiveInfo>& OutStreamingTextures) const override;
+	virtual int32 GetNumMaterials() const override;
 	// End of UPrimitiveComponent interface
 
 protected:
diff --git a/Engine/Plugins/Experimental/Paper2D/Source/Paper2D/Private/PaperCustomVersion.h b/Engine/Plugins/Experimental/Paper2D/Source/Paper2D/Private/PaperCustomVersion.h
--- a/Engine/Plugins/Experimental/Paper2D/Source/Paper2D/Private/PaperCustomVersion.h
+++ b/Engine/Plugins/Experimental/Paper2D/Source/Paper2D/Private/PaperCustomVersion.h
@@ -31,6 +31,9 @@ struct FPaperCustomVersion
 		// Fixed undo issues in several asset types due to missing RF_Transactional flag
 		AddTransactionalToClasses = 7,
 
+		// Converted UPaperSpriteComponent to be derived from UMeshComponent
+		ConvertPaperSpriteComponentToBeMeshComponent = 8,
+
 		// -----<new versions can be added above this line>-------------------------------------------------
 		VersionPlusOne,
 		LatestVersion = VersionPlusOne - 1
diff --git a/Engine/Plugins/Experimental/Paper2D/Source/Paper2D/Private/PaperSpriteComponent.cpp b/Engine/Plugins/Experimental/Paper2D/Source/Paper2D/Private/PaperSpriteComponent.cpp
--- a/Engine/Plugins/Experimental/Paper2D/Source/Paper2D/Private/PaperSpriteComponent.cpp
+++ b/Engine/Plugins/Experimental/Paper2D/Source/Paper2D/Private/PaperSpriteComponent.cpp
@@ -3,7 +3,9 @@
 #include "Paper2DPrivatePCH.h"
 #include "PaperSpriteSceneProxy.h"
 #include "PaperSpriteComponent.h"
-#include "PhysicsEngine/BodySetup2D.h"
+#include "PaperCustomVersion.h"
+#include "Runtime/Engine/Classes/PhysicsEngine/BodySetup2D.h"
+#include "Runtime/Engine/Public/ContentStreaming.h"
 
 //////////////////////////////////////////////////////////////////////////
 // UPaperSpriteComponent
@@ -13,9 +15,13 @@ UPaperSpriteComponent::UPaperSpriteComponent(const FPostConstructInitializePrope
 {
 	SetCollisionProfileName(UCollisionProfile::BlockAllDynamic_ProfileName);
 
-	MaterialOverride = nullptr;
+	MaterialOverride_DEPRECATED = nullptr;
 
 	SpriteColor = FLinearColor::White;
+
+	CastShadow = false;
+	bUseAsOccluder = false;
+	bCanEverAffectNavigation = false;
 }
 
 #if WITH_EDITOR
@@ -27,6 +33,30 @@ void UPaperSpriteComponent::PostEditChangeProperty(FPropertyChangedEvent& Proper
 }
 #endif
 
+#if WITH_EDITORONLY_DATA
+void UPaperSpriteComponent::Serialize(FArchive& Ar)
+{
+	Super::Serialize(Ar);
+
+	Ar.UsingCustomVersion(FPaperCustomVersion::GUID);
+}
+
+void UPaperSpriteComponent::PostLoad()
+{
+	Super::PostLoad();
+
+	const int32 PaperVer = GetLinkerCustomVersion(FPaperCustomVersion::GUID);
+
+	if (PaperVer < FPaperCustomVersion::ConvertPaperSpriteComponentToBeMeshComponent)
+	{
+		if (MaterialOverride_DEPRECATED != nullptr)
+		{
+			SetMaterial(0, MaterialOverride_DEPRECATED);
+		}
+	}
+}
+#endif
+
 FPrimitiveSceneProxy* UPaperSpriteComponent::CreateSceneProxy()
 {
 	FPaperSpriteSceneProxy* NewProxy = new FPaperSpriteSceneProxy(this);
@@ -157,6 +187,10 @@ bool UPaperSpriteComponent::SetSprite(class UPaperSprite* NewSprite)
 			// Update physics representation right away
 			RecreatePhysicsState();
 
+			// Notify the streaming system. Don't use Update(), because this may be the first time the mesh has been set
+			// and the component may have to be added to the streaming system for the first time.
+			IStreamingManager::Get().NotifyPrimitiveAttached(this, DPT_Spawned);
+
 			// Since we have new mesh, we need to update bounds
 			UpdateBounds();
 
@@ -167,6 +201,52 @@ bool UPaperSpriteComponent::SetSprite(class UPaperSprite* NewSprite)
 	return false;
 }
 
+void UPaperSpriteComponent::GetUsedTextures(TArray<UTexture*>& OutTextures, EMaterialQualityLevel::Type QualityLevel)
+{
+	// Get any textures referenced by our materials
+	Super::GetUsedTextures(OutTextures, QualityLevel);
+
+	// Get the texture referenced by the sprite
+	if (SourceSprite != nullptr)
+	{
+		if (UTexture* BakedTexture = SourceSprite->GetBakedTexture())
+		{
+			OutTextures.AddUnique(BakedTexture);
+		}
+	}
+}
+
+UMaterialInterface* UPaperSpriteComponent::GetMaterial(int32 MaterialIndex) const
+{
+	if (Materials.IsValidIndex(MaterialIndex) && (Materials[MaterialIndex] != nullptr))
+	{
+		return Materials[MaterialIndex];
+	}
+	else if (SourceSprite != nullptr)
+	{
+		//@TODO: PAPER2D: Need to support this for multi-material support
+		return SourceSprite->GetDefaultMaterial();
+	}
+
+	return nullptr;
+}
+
+void UPaperSpriteComponent::GetUsedMaterials(TArray<UMaterialInterface*>& OutMaterials) const
+{
+	return Super::GetUsedMaterials(OutMaterials);
+}
+
+void UPaperSpriteComponent::GetStreamingTextureInfo(TArray<FStreamingTexturePrimitiveInfo>& OutStreamingTextures) const
+{
+	//@TODO: PAPER2D: Need to support this for proper texture streaming
+	return Super::GetStreamingTextureInfo(OutStreamingTextures);
+}
+
+int32 UPaperSpriteComponent::GetNumMaterials() const
+{
+	return FMath::Max<int32>(Materials.Num(), 1);
+}
+
 UPaperSprite* UPaperSpriteComponent::GetSprite()
 {
 	return SourceSprite;
diff --git a/Engine/Plugins/Experimental/Paper2D/Source/Paper2D/Private/PaperSpriteSceneProxy.cpp b/Engine/Plugins/Experimental/Paper2D/Source/Paper2D/Private/PaperSpriteSceneProxy.cpp
--- a/Engine/Plugins/Experimental/Paper2D/Source/Paper2D/Private/PaperSpriteSceneProxy.cpp
+++ b/Engine/Plugins/Experimental/Paper2D/Source/Paper2D/Private/PaperSpriteSceneProxy.cpp
@@ -12,22 +12,10 @@ FPaperSpriteSceneProxy::FPaperSpriteSceneProxy(const UPaperSpriteComponent* InCo
 	, SourceSprite(nullptr)
 {
 	WireframeColor = InComponent->GetWireframeColor();
-	SourceSprite = InComponent->SourceSprite; //@TODO: This is totally not threadsafe, and won't keep up to date if the actor's sprite changes, etc....
-	
-	if (SourceSprite)
-	{
-		Material = SourceSprite->GetDefaultMaterial();
-	}
-		
-	if (InComponent->MaterialOverride)
-	{
-		Material = InComponent->MaterialOverride;
-	}
+	Material = InComponent->GetMaterial(0);
+	MaterialRelevance = InComponent->GetMaterialRelevance();
 
-	if (Material)
-	{
-		MaterialRelevance = Material->GetRelevance();
-	}
+	SourceSprite = InComponent->SourceSprite; //@TODO: This is totally not threadsafe, and won't keep up to date if the actor's sprite changes, etc....
 }
 
 void FPaperSpriteSceneProxy::DrawDynamicElements(FPrimitiveDrawInterface* PDI, const FSceneView* View)
