Paper2D: Non-unity fixes

diff --git a/Engine/Plugins/Experimental/Paper2D/Source/Paper2DEditor/Private/FlipbookEditor/FlipbookEditorCommands.h b/Engine/Plugins/Experimental/Paper2D/Source/Paper2DEditor/Private/FlipbookEditor/FlipbookEditorCommands.h
--- a/Engine/Plugins/Experimental/Paper2D/Source/Paper2DEditor/Private/FlipbookEditor/FlipbookEditorCommands.h
+++ b/Engine/Plugins/Experimental/Paper2D/Source/Paper2DEditor/Private/FlipbookEditor/FlipbookEditorCommands.h
@@ -2,6 +2,8 @@
 
 #pragma once
 
+#include "PaperStyle.h"
+
 class FFlipbookEditorCommands : public TCommands<FFlipbookEditorCommands>
 {
 public:
diff --git a/Engine/Plugins/Experimental/Paper2D/Source/Paper2DEditor/Private/PaperEditorCommands.h b/Engine/Plugins/Experimental/Paper2D/Source/Paper2DEditor/Private/PaperEditorCommands.h
--- a/Engine/Plugins/Experimental/Paper2D/Source/Paper2DEditor/Private/PaperEditorCommands.h
+++ b/Engine/Plugins/Experimental/Paper2D/Source/Paper2DEditor/Private/PaperEditorCommands.h
@@ -2,6 +2,8 @@
 
 #pragma once
 
+#include "PaperStyle.h"
+
 class FPaperEditorCommands : public TCommands<FPaperEditorCommands>
 {
 public:
diff --git a/Engine/Plugins/Experimental/Paper2D/Source/Paper2DEditor/Private/SpriteEditor/SpriteEditorCommands.h b/Engine/Plugins/Experimental/Paper2D/Source/Paper2DEditor/Private/SpriteEditor/SpriteEditorCommands.h
--- a/Engine/Plugins/Experimental/Paper2D/Source/Paper2DEditor/Private/SpriteEditor/SpriteEditorCommands.h
+++ b/Engine/Plugins/Experimental/Paper2D/Source/Paper2DEditor/Private/SpriteEditor/SpriteEditorCommands.h
@@ -2,6 +2,8 @@
 
 #pragma once
 
+#include "PaperStyle.h"
+
 class FSpriteEditorCommands : public TCommands<FSpriteEditorCommands>
 {
 public:
diff --git a/Engine/Plugins/Experimental/Paper2D/Source/Paper2DEditor/Private/TileMapEditing/TileMapEditorCommands.h b/Engine/Plugins/Experimental/Paper2D/Source/Paper2DEditor/Private/TileMapEditing/TileMapEditorCommands.h
--- a/Engine/Plugins/Experimental/Paper2D/Source/Paper2DEditor/Private/TileMapEditing/TileMapEditorCommands.h
+++ b/Engine/Plugins/Experimental/Paper2D/Source/Paper2DEditor/Private/TileMapEditing/TileMapEditorCommands.h
@@ -2,6 +2,8 @@
 
 #pragma once
 
+#include "PaperStyle.h"
+
 class FTileMapEditorCommands : public TCommands<FTileMapEditorCommands>
 {
 public:
