Paper2D: Tile map improvements
- Add limited support for different projection modes on tile maps (orthogonal, isometric, and staggered isometric; only orthogonal works with in-editor editing)
- Add adjustable Spacing and Margin values on tile sets (Spacing is a gap per tile in the texture, and Margin is a gap around the border)
- Add adjustable Z offset per tile in X and Y, as well as per-layer, which permits a fix for z-fighting when using overlapping isometric tiles
- Add ability to drag-drop a tile set to make a new tile map actor (just sets the tile width and tile height to good defaults, nothing else)
- Move part of the geometry generation to a one-time step on the game thread instead of per-frame (WIP)
- Edits to tile sets and tile maps now refresh placed tile map instances

diff --git a/Engine/Plugins/Experimental/Paper2D/Source/Paper2D/Classes/PaperTileMap.h b/Engine/Plugins/Experimental/Paper2D/Source/Paper2D/Classes/PaperTileMap.h
--- a/Engine/Plugins/Experimental/Paper2D/Source/Paper2D/Classes/PaperTileMap.h
+++ b/Engine/Plugins/Experimental/Paper2D/Source/Paper2D/Classes/PaperTileMap.h
@@ -6,6 +6,24 @@
 #include "PaperSprite.h"
 #include "PaperTileMap.generated.h"
 
+// The different kinds of projection modes supported
+UENUM()
+namespace ETileMapProjectionMode
+{
+	enum Type
+	{
+		// 
+		Orthogonal,
+
+		//
+		IsometricDiamond,
+
+		//
+		IsometricStaggered
+	};
+}
+
+
 UCLASS()
 class PAPER2D_API UPaperTileMap : public UDataAsset //@TODO: Just to make it easy to spawn for now
 {
@@ -28,12 +46,22 @@ class PAPER2D_API UPaperTileMap : public UDataAsset //@TODO: Just to make it eas
 	int32 TileHeight;
 
 	// Pixels per Unreal Unit (pixels per cm)
-	UPROPERTY(Category = Setup, EditAnywhere)
+	UPROPERTY(Category=Setup, EditAnywhere)
 	float PixelsPerUnit;
 
+	// 
+	UPROPERTY(Category = Setup, EditAnywhere)
+	float SeparationPerTileX;
+
+	UPROPERTY(Category = Setup, EditAnywhere)
+	float SeparationPerTileY;
+	
+	UPROPERTY(Category=Setup, EditAnywhere)
+	float SeparationPerLayer;
+
 	// Default tile set to use for new layers
 	UPROPERTY(Category=Setup, EditAnywhere)
-	UPaperTileSet* DefaultLayerTileSet;
+	class UPaperTileSet* DefaultLayerTileSet;
 
 	// Test material
 	UPROPERTY(Category=Setup, EditAnywhere)
@@ -41,12 +69,16 @@ class PAPER2D_API UPaperTileMap : public UDataAsset //@TODO: Just to make it eas
 
 	// The list of layers
 	UPROPERTY(Category=Debug, VisibleAnywhere, EditInline) //@TODO: Remove this
-	TArray<UPaperTileLayer*> TileLayers;
+	TArray<class UPaperTileLayer*> TileLayers;
 
 	// Collision domain (no collision, 2D, or 3D)
 	UPROPERTY(Category=Collision, EditAnywhere)
 	TEnumAsByte<ESpriteCollisionMode::Type> SpriteCollisionDomain;
 
+	// Tile map type
+	UPROPERTY(Category = Setup, EditAnywhere)
+	TEnumAsByte<ETileMapProjectionMode::Type> ProjectionMode;
+
 	// Baked physics data.
 	UPROPERTY()
 	class UBodySetup* BodySetup;
diff --git a/Engine/Plugins/Experimental/Paper2D/Source/Paper2D/Classes/PaperTileMapRenderComponent.h b/Engine/Plugins/Experimental/Paper2D/Source/Paper2D/Classes/PaperTileMapRenderComponent.h
--- a/Engine/Plugins/Experimental/Paper2D/Source/Paper2D/Classes/PaperTileMapRenderComponent.h
+++ b/Engine/Plugins/Experimental/Paper2D/Source/Paper2D/Classes/PaperTileMapRenderComponent.h
@@ -45,6 +45,8 @@ public:
 protected:
 	friend class FPaperTileMapRenderSceneProxy;
 
+	void RebuildRenderData(class FPaperTileMapRenderSceneProxy* Proxy);
+
 public:
 	// UObject interface
 	virtual void PostInitProperties() override;
diff --git a/Engine/Plugins/Experimental/Paper2D/Source/Paper2D/Classes/PaperTileSet.h b/Engine/Plugins/Experimental/Paper2D/Source/Paper2D/Classes/PaperTileSet.h
--- a/Engine/Plugins/Experimental/Paper2D/Source/Paper2D/Classes/PaperTileSet.h
+++ b/Engine/Plugins/Experimental/Paper2D/Source/Paper2D/Classes/PaperTileSet.h
@@ -35,6 +35,12 @@ class PAPER2D_API UPaperTileSet : public UDataAsset //@TODO: Just to make it eas
 	UPROPERTY(Category=TileSet, EditAnywhere, meta=(UIMin=0, ClampMin=0))
 	FIntPoint DrawingOffset;
 
+	// UObject interface
+#if WITH_EDITOR
+	virtual void PostEditChangeProperty(FPropertyChangedEvent& PropertyChangedEvent) override;
+#endif
+	// End of UObject interface
+
 public:
 	int32 GetTileCount() const;
 	int32 GetTileCountX() const;
diff --git a/Engine/Plugins/Experimental/Paper2D/Source/Paper2D/Private/PaperRenderSceneProxy.cpp b/Engine/Plugins/Experimental/Paper2D/Source/Paper2D/Private/PaperRenderSceneProxy.cpp
--- a/Engine/Plugins/Experimental/Paper2D/Source/Paper2D/Private/PaperRenderSceneProxy.cpp
+++ b/Engine/Plugins/Experimental/Paper2D/Source/Paper2D/Private/PaperRenderSceneProxy.cpp
@@ -237,6 +237,74 @@ void FPaperRenderSceneProxy::DrawDynamicElements_RichMesh(FPrimitiveDrawInterfac
 	if (Material != nullptr)
 	{
 		DrawBatch(PDI, View, bUseOverrideColor, OverrideColor, Material, BatchedSprites);
+		DrawNewBatches(PDI, View, bUseOverrideColor, OverrideColor);
+	}
+}
+
+void FPaperRenderSceneProxy::DrawNewBatches(FPrimitiveDrawInterface* PDI, const FSceneView* View, bool bUseOverrideColor, const FLinearColor& OverrideColor)
+{
+	//@TODO: Doesn't support OverrideColor yet
+	if (BatchedSections.Num() == 0)
+	{
+		return;
+	}
+
+	const uint8 DPG = GetDepthPriorityGroup(View);
+	
+	FVertexFactory* VertexFactory = GetPaperSpriteVertexFactory();
+
+	FMeshBatch Mesh;
+
+	Mesh.UseDynamicData = true;
+	Mesh.DynamicVertexData = BatchedVertices.GetTypedData();
+	Mesh.DynamicVertexStride = sizeof(FPaperSpriteVertex);
+
+	Mesh.VertexFactory = VertexFactory;
+	Mesh.LCI = NULL;
+	Mesh.ReverseCulling = IsLocalToWorldDeterminantNegative() ? true : false;
+	Mesh.CastShadow = false;
+	Mesh.DepthPriorityGroup = DPG;
+	Mesh.Type = PT_TriangleList;
+	Mesh.bDisableBackfaceCulling = true;
+
+	// Set up the FMeshBatchElement.
+	FMeshBatchElement& BatchElement = Mesh.Elements[0];
+	BatchElement.IndexBuffer = NULL;
+	BatchElement.DynamicIndexData = NULL;
+	BatchElement.DynamicIndexStride = 0;
+	BatchElement.PrimitiveUniformBuffer = GetUniformBuffer();
+
+	const bool bIsWireframeView = View->Family->EngineShowFlags.Wireframe;
+
+	for (const FSpriteRenderSection& Batch : BatchedSections)
+	{
+		FMaterialRenderProxy* ParentMaterialProxy = Batch.Material->GetRenderProxy((View->Family->EngineShowFlags.Selection) && IsSelected(), IsHovered());
+
+		FTexture* TextureResource = (Batch.Texture != nullptr) ? Batch.Texture->Resource : nullptr;
+		if ((TextureResource != nullptr) && (Batch.NumVertices > 0))
+		{
+			// Create a texture override material proxy and register it as a dynamic resource so that it won't be deleted until the rendering thread has finished with it
+			FTextureOverrideRenderProxy* TextureOverrideMaterialProxy = new FTextureOverrideRenderProxy(ParentMaterialProxy, Batch.Texture, TEXT("SpriteTexture"));
+			PDI->RegisterDynamicResource(TextureOverrideMaterialProxy);
+
+			Mesh.MaterialRenderProxy = TextureOverrideMaterialProxy;
+
+			BatchElement.FirstIndex = Batch.VertexOffset;
+			BatchElement.MinVertexIndex = Batch.VertexOffset;
+			BatchElement.MaxVertexIndex = Batch.VertexOffset + Batch.NumVertices;
+			BatchElement.NumPrimitives = Batch.NumVertices / 3;
+
+			DrawRichMesh(
+				PDI, 
+				Mesh,
+				WireframeColor,
+				FLinearColor(1.0f, 1.0f, 1.0f),//LevelColor,
+				FLinearColor(1.0f, 1.0f, 1.0f),//PropertyColor,
+				this,
+				IsSelected(),
+				bIsWireframeView
+				);
+		}
 	}
 }
 
diff --git a/Engine/Plugins/Experimental/Paper2D/Source/Paper2D/Private/PaperRenderSceneProxy.h b/Engine/Plugins/Experimental/Paper2D/Source/Paper2D/Private/PaperRenderSceneProxy.h
--- a/Engine/Plugins/Experimental/Paper2D/Source/Paper2D/Private/PaperRenderSceneProxy.h
+++ b/Engine/Plugins/Experimental/Paper2D/Source/Paper2D/Private/PaperRenderSceneProxy.h
@@ -32,6 +32,51 @@ struct FPaperSpriteVertex
 	static FPackedNormal PackedNormalZ;
 };
 
+//////////////////////////////////////////////////////////////////////////
+// FSpriteRenderSection
+
+struct FSpriteRenderSection
+{
+	int32 VertexOffset;
+	int32 NumVertices;
+
+	UMaterialInterface* Material;
+	UTexture2D* Texture;
+
+	FSpriteRenderSection()
+		: VertexOffset(INDEX_NONE)
+		, NumVertices(0)
+	{
+	}
+
+	template <typename SourceArrayType>
+	void AddTriangles(const FSpriteDrawCallRecord& Record, SourceArrayType& Vertices)
+	{
+		if (NumTriangles == 0)
+		{
+			VertexOffset = VertexBuffer.Num();
+			Texture = Record.Texture;
+		}
+		else
+		{
+			checkSlow((VertexOffset + NumVertices) == Vertices.Num());
+		}
+
+		NumVertices += Record.RenderVerts.Num();
+		Vertices.Reserve(NumVertices);
+
+		for (int32 SVI = 0; SVI < Record.RenderVerts.Num(); ++SVI)
+		{
+			const FVector4& SourceVert = Record.RenderVerts[SVI];
+
+			const FVector Pos((PaperAxisX * SourceVert.X) + (PaperAxisY * SourceVert.Y) + Record.Destination);
+			const FVector2D UV(SourceVert.Z, SourceVert.W);
+
+			new (Vertices) FPaperSpriteVertex(Pos, UV, Record.Color);
+		}
+	}
+};
+
 //////////////////////////////////////////////////////////////////////////
 // FPaperRenderSceneProxy
 
@@ -61,7 +106,7 @@ protected:
 	virtual void DrawDynamicElements_RichMesh(FPrimitiveDrawInterface* PDI, const FSceneView* View, bool bUseOverrideColor, const FLinearColor& OverrideColor);
 
 	void DrawBatch(FPrimitiveDrawInterface* PDI, const FSceneView* View, bool bUseOverrideColor, const FLinearColor& OverrideColor, class UMaterialInterface* BatchMaterial, const TArray<FSpriteDrawCallRecord>& Batch);
-
+	void DrawNewBatches(FPrimitiveDrawInterface* PDI, const FSceneView* View, bool bUseOverrideColor, const FLinearColor& OverrideColor);
 
 	bool IsCollisionView(const FSceneView* View, bool& bDrawSimpleCollision, bool& bDrawComplexCollision) const;
 
@@ -70,6 +115,9 @@ protected:
 	FVertexFactory* GetPaperSpriteVertexFactory() const;
 
 protected:
+	TArray<FPaperSpriteVertex> BatchedVertices;
+	TArray<FSpriteRenderSection> BatchedSections;
+
 	TArray<FSpriteDrawCallRecord> BatchedSprites;
 	class UMaterialInterface* Material;
 	FVector Origin;
diff --git a/Engine/Plugins/Experimental/Paper2D/Source/Paper2D/Private/PaperTileLayer.cpp b/Engine/Plugins/Experimental/Paper2D/Source/Paper2D/Private/PaperTileLayer.cpp
--- a/Engine/Plugins/Experimental/Paper2D/Source/Paper2D/Private/PaperTileLayer.cpp
+++ b/Engine/Plugins/Experimental/Paper2D/Source/Paper2D/Private/PaperTileLayer.cpp
@@ -12,6 +12,12 @@ UPaperTileLayer::UPaperTileLayer(const FPostConstructInitializeProperties& PCIP)
 	LayerWidth = 4;
 	LayerHeight = 4;
 
+#if WITH_EDITORONLY_DATA
+	LayerOpacity = 1.0f;
+
+	bHiddenInEditor = false;
+#endif
+
 	DestructiveAllocateMap(LayerWidth, LayerHeight);
 }
 
diff --git a/Engine/Plugins/Experimental/Paper2D/Source/Paper2D/Private/PaperTileMap.cpp b/Engine/Plugins/Experimental/Paper2D/Source/Paper2D/Private/PaperTileMap.cpp
--- a/Engine/Plugins/Experimental/Paper2D/Source/Paper2D/Private/PaperTileMap.cpp
+++ b/Engine/Plugins/Experimental/Paper2D/Source/Paper2D/Private/PaperTileMap.cpp
@@ -15,6 +15,9 @@ UPaperTileMap::UPaperTileMap(const FPostConstructInitializeProperties& PCIP)
 	TileWidth = 32;
 	TileHeight = 32;
 	PixelsPerUnit = 1.0f;
+	SeparationPerTileX = 0.0f;
+	SeparationPerTileY = 0.0f;
+	SeparationPerLayer = 64.0f;
 	SpriteCollisionDomain = ESpriteCollisionMode::None;
 
 	static ConstructorHelpers::FObjectFinder<UMaterial> DefaultMaterial(TEXT("/Paper2D/DefaultSpriteMaterial"));
@@ -22,10 +25,17 @@ UPaperTileMap::UPaperTileMap(const FPostConstructInitializeProperties& PCIP)
 }
 
 #if WITH_EDITOR
+
+#include "PaperTileMapRenderComponent.h"
+#include "ComponentReregisterContext.h"
+
 void UPaperTileMap::PostEditChangeProperty(FPropertyChangedEvent& PropertyChangedEvent)
 {
 	const FName PropertyName = (PropertyChangedEvent.Property != NULL) ? PropertyChangedEvent.Property->GetFName() : NAME_None;
 
+	//@TODO: Determine when these are really needed, as they're seriously expensive!
+	TComponentReregisterContext<UPaperTileMapRenderComponent> ReregisterStaticComponents;
+
 	if ((PropertyName == GET_MEMBER_NAME_CHECKED(UPaperTileMap, MapWidth)) || (PropertyName == GET_MEMBER_NAME_CHECKED(UPaperTileMap, MapHeight)))
 	{
 		MapWidth = FMath::Max(MapWidth, 1);
diff --git a/Engine/Plugins/Experimental/Paper2D/Source/Paper2D/Private/PaperTileMapRenderComponent.cpp b/Engine/Plugins/Experimental/Paper2D/Source/Paper2D/Private/PaperTileMapRenderComponent.cpp
--- a/Engine/Plugins/Experimental/Paper2D/Source/Paper2D/Private/PaperTileMapRenderComponent.cpp
+++ b/Engine/Plugins/Experimental/Paper2D/Source/Paper2D/Private/PaperTileMapRenderComponent.cpp
@@ -26,7 +26,10 @@ UPaperTileMapRenderComponent::UPaperTileMapRenderComponent(const FPostConstructI
 
 FPrimitiveSceneProxy* UPaperTileMapRenderComponent::CreateSceneProxy()
 {
-	return new FPaperTileMapRenderSceneProxy(this);
+	FPaperTileMapRenderSceneProxy* Proxy = new FPaperTileMapRenderSceneProxy(this);
+	RebuildRenderData(Proxy);
+
+	return Proxy;
 }
 
 void UPaperTileMapRenderComponent::PostInitProperties()
@@ -137,4 +140,156 @@ const UObject* UPaperTileMapRenderComponent::AdditionalStatObject() const
 	return nullptr;
 }
 
+void UPaperTileMapRenderComponent::RebuildRenderData(FPaperTileMapRenderSceneProxy* Proxy)
+{
+	TArray<FSpriteDrawCallRecord> BatchedSprites;
+	
+	if (TileMap == nullptr)
+	{
+		return;
+	}
+
+	for (int32 Z = 0; Z < TileMap->TileLayers.Num(); ++Z)
+	{
+		UPaperTileLayer* Layer = TileMap->TileLayers[Z];
+
+		if (Layer == nullptr)
+		{
+			continue;
+		}
+
+		FLinearColor DrawColor = FLinearColor::White;
+#if WITH_EDITORONLY_DATA
+		if (Layer->bHiddenInEditor)
+		{
+			continue;
+		}
+
+		DrawColor.A = Layer->LayerOpacity;
+#endif
+
+		for (int32 Y = 0; Y < TileMap->MapHeight; ++Y)
+		{
+			// In pixels
+			FVector WorldOffset = FVector::ZeroVector;
+			FVector WorldStepX;
+			FVector WorldStepY;
+			FVector TileSetOffset = FVector::ZeroVector;
+			switch (TileMap->ProjectionMode)
+			{
+			case ETileMapProjectionMode::Orthogonal:
+			default:
+				WorldOffset = (TileMap->TileWidth * PaperAxisX * 0.5f) + (TileMap->TileHeight * PaperAxisY * 0.5f);
+				WorldStepX = PaperAxisX * TileMap->TileWidth;
+				WorldStepY = -PaperAxisY * TileMap->TileHeight;
+				break;
+			case ETileMapProjectionMode::IsometricDiamond:
+				WorldOffset = (TileMap->MapWidth * TileMap->TileWidth * 0.5f * PaperAxisX) - (Y * TileMap->TileHeight * 0.5f * PaperAxisY);
+				WorldStepX = (TileMap->TileWidth * PaperAxisX * 0.5f) - (TileMap->TileHeight * PaperAxisY * 0.5f);
+				WorldStepY = (TileMap->TileWidth * PaperAxisX * -0.5f) - (TileMap->TileHeight * PaperAxisY * 0.5f);
+				break;
+			case ETileMapProjectionMode::IsometricStaggered:
+				WorldOffset = ((Y & 1) * 0.5f * TileMap->TileWidth * PaperAxisX) + (0.5f * TileMap->TileHeight * PaperAxisY);
+				WorldStepX = PaperAxisX * TileMap->TileWidth;
+				WorldStepY = -PaperAxisY * TileMap->TileHeight;
+				break;
+			}
+
+			for (int32 X = 0; X < TileMap->MapWidth; ++X)
+			{
+				const FPaperTileInfo TileInfo = Layer->GetCell(X, Y);
+
+				// do stuff
+				const float TotalSeparation = (TileMap->SeparationPerLayer * Z) + (TileMap->SeparationPerTileX * X) + (TileMap->SeparationPerTileY * Y);
+				FVector WorldPos = (WorldStepX * X) + (WorldStepY * Y) + WorldOffset;
+				WorldPos += TotalSeparation * PaperAxisZ;
+
+				const int32 TileWidth = TileMap->TileWidth;
+				const int32 TileHeight = TileMap->TileHeight;
+
+				UTexture2D* LastSourceTexture = nullptr;
+				FVector2D InverseTextureSize(1.0f, 1.0f);
+				FVector2D SourceDimensionsUV(1.0f, 1.0f);
+
+				FVector2D TileSeparation(TileMap->TileWidth, TileMap->TileHeight);
+				FVector2D TileSizeXY(TileMap->TileWidth, TileMap->TileHeight);
+
+				{
+					UTexture2D* SourceTexture = nullptr;
+
+					FVector2D SourceUV = FVector2D::ZeroVector;
+					if (Layer->bCollisionLayer)
+					{
+						if (TileInfo.PackedTileIndex == 0)
+						{
+							continue;
+						}
+						SourceTexture = UCanvas::StaticClass()->GetDefaultObject<UCanvas>()->DefaultTexture;
+					}
+					else
+					{
+						if (TileInfo.TileSet == nullptr)
+						{
+							continue;
+						}
+
+						if (!TileInfo.TileSet->GetTileUV(TileInfo.PackedTileIndex, /*out*/ SourceUV))
+						{
+							continue;
+						}
+
+						SourceTexture = TileInfo.TileSet->TileSheet;
+						if (SourceTexture == nullptr)
+						{
+							continue;
+						}
+					}
+
+					if (SourceTexture != LastSourceTexture)
+					{
+						InverseTextureSize = FVector2D(1.0f / SourceTexture->GetSizeX(), 1.0f / SourceTexture->GetSizeY());
+
+						if (TileInfo.TileSet != nullptr)
+						{
+							SourceDimensionsUV = FVector2D(TileInfo.TileSet->TileWidth * InverseTextureSize.X, TileInfo.TileSet->TileHeight * InverseTextureSize.Y);
+							TileSeparation = FVector2D(TileInfo.TileSet->TileWidth, TileInfo.TileSet->TileHeight);
+							TileSetOffset = (TileInfo.TileSet->DrawingOffset.X * PaperAxisX) + (TileInfo.TileSet->DrawingOffset.Y * PaperAxisY);
+						}
+						else
+						{
+							SourceDimensionsUV = FVector2D(TileWidth * InverseTextureSize.X, TileHeight * InverseTextureSize.Y);
+							TileSeparation = FVector2D(TileWidth, TileHeight);
+							TileSetOffset = FVector::ZeroVector;
+						}
+						LastSourceTexture = SourceTexture;
+					}
+
+					SourceUV.X *= InverseTextureSize.X;
+					SourceUV.Y *= InverseTextureSize.Y;
+
+					FSpriteDrawCallRecord& NewTile = *(new (BatchedSprites) FSpriteDrawCallRecord());
+					NewTile.Destination = WorldPos + TileSetOffset;
+					NewTile.Texture = SourceTexture;
+					NewTile.Color = DrawColor;
+
+					const FVector4 BottomLeft(0.0f, 0.0f, SourceUV.X, SourceUV.Y + SourceDimensionsUV.Y);
+					const FVector4 BottomRight(TileSeparation.X, 0.0f, SourceUV.X + SourceDimensionsUV.X, SourceUV.Y + SourceDimensionsUV.Y);
+					const FVector4 TopRight(TileSeparation.X, TileSeparation.Y, SourceUV.X + SourceDimensionsUV.X, SourceUV.Y);
+					const FVector4 TopLeft(0.0f, TileSeparation.Y, SourceUV.X, SourceUV.Y);
+
+					new (NewTile.RenderVerts) FVector4(BottomLeft);
+					new (NewTile.RenderVerts) FVector4(TopRight);
+					new (NewTile.RenderVerts) FVector4(BottomRight);
+
+					new (NewTile.RenderVerts) FVector4(BottomLeft);
+					new (NewTile.RenderVerts) FVector4(TopLeft);
+					new (NewTile.RenderVerts) FVector4(TopRight);
+				}
+			}
+		}
+	}
+
+	Proxy->SetBatchesHack(BatchedSprites);
+}
+
 #undef LOCTEXT_NAMESPACE
\ No newline at end of file
diff --git a/Engine/Plugins/Experimental/Paper2D/Source/Paper2D/Private/PaperTileMapRenderSceneProxy.cpp b/Engine/Plugins/Experimental/Paper2D/Source/Paper2D/Private/PaperTileMapRenderSceneProxy.cpp
--- a/Engine/Plugins/Experimental/Paper2D/Source/Paper2D/Private/PaperTileMapRenderSceneProxy.cpp
+++ b/Engine/Plugins/Experimental/Paper2D/Source/Paper2D/Private/PaperTileMapRenderSceneProxy.cpp
@@ -26,8 +26,6 @@ FPaperTileMapRenderSceneProxy::FPaperTileMapRenderSceneProxy(const UPaperTileMap
 
 void FPaperTileMapRenderSceneProxy::DrawDynamicElements(FPrimitiveDrawInterface* PDI, const FSceneView* View)
 {
-	BatchedSprites.Empty();
-
 	// Draw the tile maps
 	//@TODO: RenderThread race condition
 	if (TileMap != NULL)
@@ -51,7 +49,7 @@ void FPaperTileMapRenderSceneProxy::DrawDynamicElements(FPrimitiveDrawInterface*
 				else
 				{
 					// Make a material for drawing solid collision stuff
-					const UMaterial* LevelColorationMaterial = View->Family->EngineShowFlags.Lighting 
+					const UMaterial* LevelColorationMaterial = View->Family->EngineShowFlags.Lighting
 						? GEngine->ShadedLevelColorationLitMaterial : GEngine->ShadedLevelColorationUnlitMaterial;
 
 					const FColoredMaterialRenderProxy CollisionMaterialInstance(
@@ -66,13 +64,13 @@ void FPaperTileMapRenderSceneProxy::DrawDynamicElements(FPrimitiveDrawInterface*
 
 					// In old wireframe collision mode, always draw the wireframe highlighted (selected or not).
 					bool bDrawWireSelected = IsSelected();
-					if(View->Family->EngineShowFlags.Collision)
+					if (View->Family->EngineShowFlags.Collision)
 					{
 						bDrawWireSelected = true;
 					}
 
 					// Differentiate the color based on bBlockNonZeroExtent.  Helps greatly with skimming a level for optimization opportunities.
-					FColor CollisionColor = FColor(157,149,223,255);
+					FColor CollisionColor = FColor(157, 149, 223, 255);
 
 					const bool bPerHullColor = false;
 					const bool bDrawSimpleSolid = false;
@@ -138,127 +136,12 @@ void FPaperTileMapRenderSceneProxy::DrawDynamicElements(FPrimitiveDrawInterface*
 			// Create a local space bounding box
 			const FVector TopLeft(-TW*0.5f, -2.0f, TH*0.5f);
 			const FVector Dimensions(TileMap->MapWidth * TW, 4.0f, -TileMap->MapHeight * TH);
-			const FBox OutlineBox(TopLeft, TopLeft+Dimensions);
-		
+			const FBox OutlineBox(TopLeft, TopLeft + Dimensions);
+
 			// Draw it		
 			DrawWireBox(PDI, OutlineBox.TransformBy(LocalToWorld), FLinearColor::White, DPG); //@TODO: Paper: Doesn't handle rotation well - need to use DrawOrientedWireBox
 		}
 #endif
-
-		const float LayerSeparation = 64.0f;
-
-		const int32 TileSeparationX = TileMap->TileWidth;
-		const int32 TileSeparationY = TileMap->TileHeight;
-
-		for (int32 LayerIndex = 0; LayerIndex < TileMap->TileLayers.Num(); ++LayerIndex)
-		{
-			UPaperTileLayer* Layer = TileMap->TileLayers[LayerIndex];
-			if (Layer == NULL)
-			{
-				continue;
-			}
-#if WITH_EDITORONLY_DATA
-			if (Layer->bHiddenInEditor)
-			{
-				continue;
-			}
-#endif
-			if (Layer->bCollisionLayer && View->bIsGameView)
-			{
-				continue;
-			}
-
-			ensure(Layer->LayerWidth == TileMap->MapWidth);
-			ensure(Layer->LayerHeight == TileMap->MapHeight);
-
-			const int32 TileWidth = TileMap->TileWidth;
-			const int32 TileHeight = TileMap->TileHeight;
-
-			UTexture2D* LastSourceTexture = nullptr;
-			FVector2D InverseTextureSize(1.0f, 1.0f);
-			FVector2D SourceDimensionsUV(1.0f, 1.0f);
-
-			FVector2D TileSizeXY(TileMap->TileWidth, TileMap->TileHeight);
-			FLinearColor DrawColor = FLinearColor::White;
-			EBlendMode BlendMode = BLEND_Translucent;
-			if (Layer->bCollisionLayer)
-			{
-				DrawColor = FLinearColor::Blue;
-				BlendMode = BLEND_Opaque;
-			}
-
-			for (int32 Y = 0; Y < Layer->LayerHeight; ++Y)
-			{
-				for (int32 X = 0; X < Layer->LayerWidth; ++X)
-				{
-					const FPaperTileInfo TileInfo = Layer->GetCell(X, Y);
-					UTexture2D* SourceTexture = nullptr;
-
-					FVector2D SourceUV = FVector2D::ZeroVector;
-					if (Layer->bCollisionLayer)
-					{
-						if (TileInfo.PackedTileIndex == 0)
-						{
-							continue;
-						}
-						SourceTexture = UCanvas::StaticClass()->GetDefaultObject<UCanvas>()->DefaultTexture;
-					}
-					else
-					{
-						if (TileInfo.TileSet == nullptr)
-						{
-							continue;
-						}
-
-						if (!TileInfo.TileSet->GetTileUV(TileInfo.PackedTileIndex, /*out*/ SourceUV))
-						{
-							continue;
-						}
-
-						SourceTexture = TileInfo.TileSet->TileSheet;
-						if (SourceTexture == nullptr)
-						{
-							continue;
-						}
-					}
-
-
-					if (SourceTexture != LastSourceTexture)
-					{
-						InverseTextureSize = FVector2D(1.0f / SourceTexture->GetSizeX(), 1.0f / SourceTexture->GetSizeY());
-						SourceDimensionsUV = FVector2D(TileWidth * InverseTextureSize.X, TileHeight * InverseTextureSize.Y);
-						LastSourceTexture = SourceTexture;
-					}
-
-
-					SourceUV.X *= InverseTextureSize.X;
-					SourceUV.Y *= InverseTextureSize.Y;
-
-					const FVector DrawPos(TileSeparationX * (X - 0.5f) * PaperAxisX + TileSeparationY * -(Y + 0.5f) * PaperAxisY + ((LayerIndex * LayerSeparation) * PaperAxisZ));
-
-					FSpriteDrawCallRecord& NewTile = *(new (BatchedSprites) FSpriteDrawCallRecord());
-					NewTile.Destination = DrawPos;
-					NewTile.Texture = SourceTexture;
-					NewTile.Color = DrawColor;
-					
-					const float W = TileSeparationX;
-					const float H = TileSeparationY;
-
-					const FVector4 BottomLeft(0.0f, 0.0f, SourceUV.X, SourceUV.Y + SourceDimensionsUV.Y);
-					const FVector4 BottomRight(W, 0.0f, SourceUV.X + SourceDimensionsUV.X, SourceUV.Y + SourceDimensionsUV.Y);
-					const FVector4 TopRight(W, H, SourceUV.X + SourceDimensionsUV.X, SourceUV.Y);
-					const FVector4 TopLeft(0.0f, H, SourceUV.X, SourceUV.Y);
-
-					new (NewTile.RenderVerts) FVector4(BottomLeft);
-					new (NewTile.RenderVerts) FVector4(TopRight);
-					new (NewTile.RenderVerts) FVector4(BottomRight);
-
-					new (NewTile.RenderVerts) FVector4(BottomLeft);
-					new (NewTile.RenderVerts) FVector4(TopLeft);
-					new (NewTile.RenderVerts) FVector4(TopRight);
-				}
-			}
-		}
 	}
 
 	// Draw all of the queued up sprites
diff --git a/Engine/Plugins/Experimental/Paper2D/Source/Paper2D/Private/PaperTileMapRenderSceneProxy.h b/Engine/Plugins/Experimental/Paper2D/Source/Paper2D/Private/PaperTileMapRenderSceneProxy.h
--- a/Engine/Plugins/Experimental/Paper2D/Source/Paper2D/Private/PaperTileMapRenderSceneProxy.h
+++ b/Engine/Plugins/Experimental/Paper2D/Source/Paper2D/Private/PaperTileMapRenderSceneProxy.h
@@ -16,6 +16,10 @@ public:
 	virtual void DrawDynamicElements(FPrimitiveDrawInterface* PDI, const FSceneView* View) override;
 	// End of FPrimitiveSceneProxy interface.
 
+	void SetBatchesHack(TArray<FSpriteDrawCallRecord>& InBatchedSprites)
+	{
+		BatchedSprites = InBatchedSprites;
+	}
 protected:
 	//@TODO: Not thread safe
 	const UPaperTileMap* TileMap;
diff --git a/Engine/Plugins/Experimental/Paper2D/Source/Paper2D/Private/PaperTileSet.cpp b/Engine/Plugins/Experimental/Paper2D/Source/Paper2D/Private/PaperTileSet.cpp
--- a/Engine/Plugins/Experimental/Paper2D/Source/Paper2D/Private/PaperTileSet.cpp
+++ b/Engine/Plugins/Experimental/Paper2D/Source/Paper2D/Private/PaperTileSet.cpp
@@ -20,8 +20,8 @@ int32 UPaperTileSet::GetTileCount() const
 		const int32 TextureWidth = TileSheet->GetSizeX();
 		const int32 TextureHeight = TileSheet->GetSizeY();
 
-		const int32 CellsX = TextureWidth / TileWidth;
-		const int32 CellsY = TextureHeight / TileHeight;
+		const int32 CellsX = (TextureWidth - (Margin * 2) + Spacing) / (TileWidth + Spacing);
+		const int32 CellsY = (TextureHeight - (Margin * 2) + Spacing) / (TileHeight + Spacing);
 
 		return CellsX * CellsY;
 	}
@@ -37,7 +37,7 @@ int32 UPaperTileSet::GetTileCountX() const
 	{
 		checkSlow(TileWidth > 0);
 		const int32 TextureWidth = TileSheet->GetSizeX();
-		const int32 CellsX = TextureWidth / TileWidth;
+		const int32 CellsX = (TextureWidth - (Margin * 2) + Spacing) / (TileWidth + Spacing);
 		return CellsX;
 	}
 	else
@@ -52,7 +52,7 @@ int32 UPaperTileSet::GetTileCountY() const
 	{
 		checkSlow(TileHeight > 0);
 		const int32 TextureHeight = TileSheet->GetSizeY();
-		const int32 CellsY = TextureHeight / TileHeight;
+		const int32 CellsY = (TextureHeight - (Margin * 2) + Spacing) / (TileHeight + Spacing);
 		return CellsY;
 	}
 	else
@@ -67,11 +67,11 @@ bool UPaperTileSet::GetTileUV(int32 TileIndex, /*out*/ FVector2D& Out_TileUV) co
 	if (TileSheet != NULL)
 	{
 		checkSlow((TileWidth > 0) && (TileHeight > 0));
-		const int32 TextureWidth = TileSheet->GetSizeX();
-		const int32 TextureHeight = TileSheet->GetSizeY();
+		const int32 TextureWidth = TileSheet->GetSizeX() - (Margin * 2) + Spacing;
+		const int32 TextureHeight = TileSheet->GetSizeY() - (Margin * 2) + Spacing;
 
-		const int32 CellsX = TextureWidth / TileWidth;
-		const int32 CellsY = TextureHeight / TileHeight;
+		const int32 CellsX = TextureWidth / (TileWidth + Spacing);
+		const int32 CellsY = TextureHeight / (TileHeight + Spacing);
 
 		const int32 NumCells = CellsX * CellsY;
 
@@ -84,8 +84,8 @@ bool UPaperTileSet::GetTileUV(int32 TileIndex, /*out*/ FVector2D& Out_TileUV) co
 			const int32 X = TileIndex % CellsX;
 			const int32 Y = TileIndex / CellsX;
 
-			Out_TileUV.X = X * TileWidth;
-			Out_TileUV.Y = Y * TileHeight;
+			Out_TileUV.X = X * (TileWidth + Spacing) + Margin;
+			Out_TileUV.Y = Y * (TileHeight + Spacing) + Margin;
 			return true;
 		}
 	}
@@ -94,3 +94,25 @@ bool UPaperTileSet::GetTileUV(int32 TileIndex, /*out*/ FVector2D& Out_TileUV) co
 		return false;
 	}
 }
+
+#if WITH_EDITOR
+
+#include "PaperTileMapRenderComponent.h"
+#include "ComponentReregisterContext.h"
+
+void UPaperTileSet::PostEditChangeProperty(FPropertyChangedEvent& PropertyChangedEvent)
+{
+	const FName PropertyName = (PropertyChangedEvent.Property != NULL) ? PropertyChangedEvent.Property->GetFName() : NAME_None;
+
+	Margin = FMath::Max<int32>(Margin, 0);
+	Spacing = FMath::Max<int32>(Spacing, 0);
+
+	TileWidth = FMath::Max<int32>(TileWidth, 1);
+	TileHeight = FMath::Max<int32>(TileHeight, 1);
+
+	//@TODO: Determine when these are really needed, as they're seriously expensive!
+	TComponentReregisterContext<UPaperTileMapRenderComponent> ReregisterStaticComponents;
+
+	Super::PostEditChangeProperty(PropertyChangedEvent);
+}
+#endif
diff --git a/Engine/Plugins/Experimental/Paper2D/Source/Paper2DEditor/Private/TileMapEditing/EdModeTileMap.cpp b/Engine/Plugins/Experimental/Paper2D/Source/Paper2DEditor/Private/TileMapEditing/EdModeTileMap.cpp
--- a/Engine/Plugins/Experimental/Paper2D/Source/Paper2DEditor/Private/TileMapEditing/EdModeTileMap.cpp
+++ b/Engine/Plugins/Experimental/Paper2D/Source/Paper2DEditor/Private/TileMapEditing/EdModeTileMap.cpp
@@ -480,7 +480,7 @@ bool FEdModeTileMap::PaintTiles(const FViewportCursorLocation& Ray)
 			}
 		}
 
-		if ( bChangedSomething && GetActiveLayerPaintingMode() == ETileMapLayerPaintingMode::CollisionLayers )
+		if (bChangedSomething)
 		{
 			TileMap->PostEditChange();
 		}
@@ -543,7 +543,7 @@ bool FEdModeTileMap::EraseTiles(const FViewportCursorLocation& Ray)
 			}
 		}
 
-		if ( bChangedSomething && GetActiveLayerPaintingMode() == ETileMapLayerPaintingMode::CollisionLayers )
+		if (bChangedSomething)
 		{
 			TileMap->PostEditChange();
 		}
@@ -682,7 +682,7 @@ bool FEdModeTileMap::FloodFillTiles(const FViewportCursorLocation& Ray)
 			}
 		}
 
-		if ( bChangedSomething && GetActiveLayerPaintingMode() == ETileMapLayerPaintingMode::CollisionLayers )
+		if (bChangedSomething)
 		{
 			TileMap->PostEditChange();
 		}
diff --git a/Engine/Plugins/Experimental/Paper2D/Source/Paper2DEditor/Private/TileMapEditing/TileMapActorFactory.cpp b/Engine/Plugins/Experimental/Paper2D/Source/Paper2DEditor/Private/TileMapEditing/TileMapActorFactory.cpp
--- a/Engine/Plugins/Experimental/Paper2D/Source/Paper2DEditor/Private/TileMapEditing/TileMapActorFactory.cpp
+++ b/Engine/Plugins/Experimental/Paper2D/Source/Paper2DEditor/Private/TileMapEditing/TileMapActorFactory.cpp
@@ -15,31 +15,51 @@ UTileMapActorFactory::UTileMapActorFactory(const FPostConstructInitializePropert
 
 void UTileMapActorFactory::PostSpawnActor(UObject* Asset, AActor* NewActor)
 {
+	APaperTileMapActor* TypedActor = CastChecked<APaperTileMapActor>(NewActor);
+	UPaperTileMapRenderComponent* RenderComponent = TypedActor->RenderComponent;
+	check(RenderComponent);
+
 	if (UPaperTileMap* TileMap = Cast<UPaperTileMap>(Asset))
 	{
 		GEditor->SetActorLabelUnique(NewActor, TileMap->GetName());
 
-		APaperTileMapActor* TypedActor = CastChecked<APaperTileMapActor>(NewActor);
-		UPaperTileMapRenderComponent* RenderComponent = TypedActor->RenderComponent;
-		check(RenderComponent);
-
 		RenderComponent->UnregisterComponent();
 		RenderComponent->TileMap = TileMap;
 		RenderComponent->RegisterComponent();
 	}
+	else if (UPaperTileSet* TileSet = Cast<UPaperTileSet>(Asset))
+	{
+		GEditor->SetActorLabelUnique(NewActor, TileSet->GetName());
+
+		if (RenderComponent->TileMap != nullptr)
+		{
+			RenderComponent->UnregisterComponent();
+			RenderComponent->TileMap->TileWidth = TileSet->TileWidth;
+			RenderComponent->TileMap->TileHeight = TileSet->TileHeight;
+			RenderComponent->RegisterComponent();
+		}
+	}
 }
 
 void UTileMapActorFactory::PostCreateBlueprint(UObject* Asset, AActor* CDO)
 {
-	if (UPaperTileMap* TileMap = Cast<UPaperTileMap>(Asset))
+	if (APaperTileMapActor* TypedActor = Cast<APaperTileMapActor>(CDO))
 	{
-		if (APaperTileMapActor* TypedActor = Cast<APaperTileMapActor>(CDO))
-		{
-			UPaperTileMapRenderComponent* RenderComponent = TypedActor->RenderComponent;
-			check(RenderComponent);
+		UPaperTileMapRenderComponent* RenderComponent = TypedActor->RenderComponent;
+		check(RenderComponent);
 
+		if (UPaperTileMap* TileMap = Cast<UPaperTileMap>(Asset))
+		{
 			RenderComponent->TileMap = TileMap;
 		}
+		else if (UPaperTileSet* TileSet = Cast<UPaperTileSet>(Asset))
+		{
+			if (RenderComponent->TileMap != nullptr)
+			{
+				RenderComponent->TileMap->TileWidth = TileSet->TileWidth;
+				RenderComponent->TileMap->TileHeight = TileSet->TileHeight;
+			}
+		}
 	}
 }
 
@@ -47,14 +67,22 @@ bool UTileMapActorFactory::CanCreateActorFrom(const FAssetData& AssetData, FText
 {
 	if (GetDefault<UPaperRuntimeSettings>()->bEnableTileMapEditing)
 	{
-		if (AssetData.IsValid() && AssetData.GetClass()->IsChildOf(UPaperTileMap::StaticClass()))
+		if (AssetData.IsValid())
 		{
-			return true;
+			UClass* AssetClass = AssetData.GetClass();
+			if ((AssetClass != nullptr) && (AssetClass->IsChildOf(UPaperTileMap::StaticClass()) || AssetClass->IsChildOf(UPaperTileSet::StaticClass())))
+			{
+				return true;
+			}
+			else
+			{
+				OutErrorMsg = NSLOCTEXT("Paper2D", "CanCreateActorFrom_NoTileMap", "No tile map was specified.");
+				return false;
+			}
 		}
 		else
 		{
-			OutErrorMsg = NSLOCTEXT("Paper2D", "CanCreateActorFrom_NoTileMap", "No tile map was specified.");
-			return false;
+			return true;
 		}
 	}
 	else
diff --git a/Engine/Plugins/Experimental/Paper2D/Source/PaperTiledImporter/Classes/PaperTiledImporterFactory.h b/Engine/Plugins/Experimental/Paper2D/Source/PaperTiledImporter/Classes/PaperTiledImporterFactory.h
--- a/Engine/Plugins/Experimental/Paper2D/Source/PaperTiledImporter/Classes/PaperTiledImporterFactory.h
+++ b/Engine/Plugins/Experimental/Paper2D/Source/PaperTiledImporter/Classes/PaperTiledImporterFactory.h
@@ -2,6 +2,7 @@
 
 #pragma once
 
+#include "PaperTileMap.h"
 #include "PaperTileLayer.h"
 #include "PaperTiledImporterFactory.generated.h"
 
@@ -11,7 +12,8 @@ namespace ETiledOrientation
 	{
 		Unknown,
 		Orthogonal,
-		Isometric
+		Isometric,
+		Staggered
 	};
 }
 
@@ -120,6 +122,7 @@ struct FTileMapFromTiled
 	//map of properties
 
 	FPaperTileInfo ConvertTileGIDToPaper2D(int32 GID) const;
+	ETileMapProjectionMode::Type GetOrientationType() const;
 
 	FTileMapFromTiled();
 	bool IsValid() const;
diff --git a/Engine/Plugins/Experimental/Paper2D/Source/PaperTiledImporter/Private/PaperTiledImporterFactory.cpp b/Engine/Plugins/Experimental/Paper2D/Source/PaperTiledImporter/Private/PaperTiledImporterFactory.cpp
--- a/Engine/Plugins/Experimental/Paper2D/Source/PaperTiledImporter/Private/PaperTiledImporterFactory.cpp
+++ b/Engine/Plugins/Experimental/Paper2D/Source/PaperTiledImporter/Private/PaperTiledImporterFactory.cpp
@@ -197,7 +197,10 @@ UObject* UPaperTiledImporterFactory::FactoryCreateText(UClass* InClass, UObject*
 		Result->MapHeight = GlobalInfo.Height;
 		Result->TileWidth = GlobalInfo.TileWidth;
 		Result->TileHeight = GlobalInfo.TileHeight;
-		//@TODO: Need to handle Orientation
+		Result->SeparationPerTileX = 0.0f;
+		Result->SeparationPerTileY = 0.0f;
+		Result->SeparationPerLayer = -1.0f;
+		Result->ProjectionMode = GlobalInfo.GetOrientationType();
 
 		// Create the tile sets
 		for (const FTileSetFromTiled& TileSetData : GlobalInfo.TileSets)
@@ -438,6 +441,10 @@ void UPaperTiledImporterFactory::ParseGlobalInfoFromJSON(TSharedPtr<FJsonObject>
 	{
 		OutParsedInfo.Orientation = ETiledOrientation::Isometric;
 	}
+	else if (OrientationModeStr == TEXT("staggered"))
+	{
+		OutParsedInfo.Orientation = ETiledOrientation::Staggered;
+	}
 	else
 	{
 		TILED_IMPORT_ERROR(TEXT("Failed to parse '%s'.  Invalid value for '%s' (%s but expected 'orthogonal' or 'isometric')"), *NameForErrors, TEXT("orientation"), *OrientationModeStr);
@@ -491,6 +498,20 @@ FPaperTileInfo FTileMapFromTiled::ConvertTileGIDToPaper2D(int32 GID) const
 	return Result;
 }
 
+ETileMapProjectionMode::Type FTileMapFromTiled::GetOrientationType() const
+{
+	switch (Orientation)
+	{
+	case ETiledOrientation::Isometric:
+		return ETileMapProjectionMode::IsometricDiamond;
+	case ETiledOrientation::Staggered:
+		return ETileMapProjectionMode::IsometricStaggered;
+	case ETiledOrientation::Orthogonal:
+	default:
+		return ETileMapProjectionMode::Orthogonal;
+	};
+}
+
 //////////////////////////////////////////////////////////////////////////
 // FTileSetFromTiled
 
@@ -604,7 +625,6 @@ void FTileLayerFromTiled::ParseFromJSON(TSharedPtr<FJsonObject> Tree, const FStr
 
 	if (!Tree->TryGetStringField(TEXT("name"), /*out*/ Name))
 	{
-		//@TODO: Figure out what to do with the layer type!
 		TILED_IMPORT_ERROR(TEXT("Failed to parse '%s'.  Expected a layer name"), *NameForErrors);
 		bSuccessfullyParsed = false;
 	}
