Paper2D: Rename texture context menu command 'Configure For Sprites' to 'Configure For Retro Sprites' to better reflect the function it performs

diff --git a/Engine/Plugins/Experimental/Paper2D/Source/Paper2DEditor/Private/ContentBrowserExtensions/ContentBrowserExtensions.cpp b/Engine/Plugins/Experimental/Paper2D/Source/Paper2DEditor/Private/ContentBrowserExtensions/ContentBrowserExtensions.cpp
--- a/Engine/Plugins/Experimental/Paper2D/Source/Paper2DEditor/Private/ContentBrowserExtensions/ContentBrowserExtensions.cpp
+++ b/Engine/Plugins/Experimental/Paper2D/Source/Paper2DEditor/Private/ContentBrowserExtensions/ContentBrowserExtensions.cpp
@@ -204,8 +204,8 @@ public:
 		FUIAction Action_ConfigureTexturesForSprites(
 			FExecuteAction::CreateStatic(&FPaperContentBrowserExtensions_Impl::ExecuteSelectedContentFunctor, StaticCastSharedPtr<FContentBrowserSelectedAssetExtensionBase>(ConfigureFunctor)));
 		Builder.AddMenuEntry(
-			LOCTEXT("CB_Extension_Texture_ConfigureTextureForSprites", "Configure For Sprites"), 
-			LOCTEXT("CB_Extension_Texture_ConfigureTextureForSprites_Tooltip", "Sets compression settings and sampling modes to good defaults for sprites"),
+			LOCTEXT("CB_Extension_Texture_ConfigureTextureForSprites", "Configure For Retro Sprites"), 
+			LOCTEXT("CB_Extension_Texture_ConfigureTextureForSprites_Tooltip", "Sets compression settings and sampling modes to good defaults for retro sprites (nearest filtering, uncompressed, etc...)"),
 			FSlateIcon(),
 			Action_ConfigureTexturesForSprites,
 			NAME_None,
