Paper2D: Add ConditionalPostLoad call for associated textures to avoid load-time ordering issues when opening older sprites

diff --git a/Engine/Plugins/Experimental/Paper2D/Source/Paper2D/Private/PaperSprite.cpp b/Engine/Plugins/Experimental/Paper2D/Source/Paper2D/Private/PaperSprite.cpp
--- a/Engine/Plugins/Experimental/Paper2D/Source/Paper2D/Private/PaperSprite.cpp
+++ b/Engine/Plugins/Experimental/Paper2D/Source/Paper2D/Private/PaperSprite.cpp
@@ -580,11 +580,13 @@ void UPaperSprite::RebuildRenderData()
 	FVector2D TextureSize(1.0f, 1.0f);
 	if (EffectiveTexture)
 	{
-		EffectiveTexture->FinishCachePlatformData();
+		EffectiveTexture->ConditionalPostLoad();
 		const int32 TextureWidth = EffectiveTexture->GetSizeX();
 		const int32 TextureHeight = EffectiveTexture->GetSizeY();
-		check((TextureWidth > 0) && (TextureHeight > 0));
-		TextureSize = FVector2D(TextureWidth, TextureHeight);
+		if (ensure((TextureWidth > 0) && (TextureHeight > 0)))
+		{
+			TextureSize = FVector2D(TextureWidth, TextureHeight);
+		}
 	}
 	const float InverseWidth = 1.0f / TextureSize.X;
 	const float InverseHeight = 1.0f / TextureSize.Y;
