Paper2D: Commented out WIP vertex buffer changes

diff --git a/Engine/Plugins/Experimental/Paper2D/Source/Paper2D/Private/PaperRenderSceneProxy.h b/Engine/Plugins/Experimental/Paper2D/Source/Paper2D/Private/PaperRenderSceneProxy.h
--- a/Engine/Plugins/Experimental/Paper2D/Source/Paper2D/Private/PaperRenderSceneProxy.h
+++ b/Engine/Plugins/Experimental/Paper2D/Source/Paper2D/Private/PaperRenderSceneProxy.h
@@ -49,12 +49,13 @@ struct FSpriteRenderSection
 	{
 	}
 
+#if 0
 	template <typename SourceArrayType>
 	void AddTriangles(const FSpriteDrawCallRecord& Record, SourceArrayType& Vertices)
 	{
-		if (NumTriangles == 0)
+		if (NumVertices == 0)
 		{
-			VertexOffset = VertexBuffer.Num();
+			VertexOffset = Vertices.Num();
 			Texture = Record.Texture;
 		}
 		else
@@ -75,6 +76,7 @@ struct FSpriteRenderSection
 			new (Vertices) FPaperSpriteVertex(Pos, UV, Record.Color);
 		}
 	}
+#endif
 };
 
 //////////////////////////////////////////////////////////////////////////
