Paper2D: Add support for dicing sprite render geometry, separating out translucent/masked quads from opaque ones (WIP)
- Sprites now have an alternate material and the ability to have a second render section defined using that material (only used for diced geometry right now, with a convention that the alternate material is the opaque one in that case)
- Adds a new ESpritePolygonMode::Diced value, valid only for render geometry (treated like SourceBoundingBox for collision)
- Adds opportunistic triangle merging for both collision and render geometry when there are multiple source polygons

diff --git a/Engine/Plugins/Experimental/Paper2D/Source/Paper2D/Classes/PaperSprite.h b/Engine/Plugins/Experimental/Paper2D/Source/Paper2D/Classes/PaperSprite.h
--- a/Engine/Plugins/Experimental/Paper2D/Source/Paper2D/Classes/PaperSprite.h
+++ b/Engine/Plugins/Experimental/Paper2D/Source/Paper2D/Classes/PaperSprite.h
@@ -67,10 +67,14 @@ protected:
 	UPROPERTY()
 	UTexture2D* BakedSourceTexture;
 
-	// The material to use on a sprite instance if not overridden
+	// The material to use on a sprite instance if not overridden (this is the default material when only one is being used, and is the translucent/masked material for Diced render geometry, slot 0)
 	UPROPERTY(Category=Sprite, EditAnywhere, BlueprintReadOnly)
 	UMaterialInterface* DefaultMaterial;
 
+	// The alternate material to use on a sprite instance if not overridden (this is only used for Diced render geometry, and will be the opaque material in that case, slot 1)
+	UPROPERTY(Category=Sprite, EditAnywhere, BlueprintReadOnly)
+	UMaterialInterface* AlternateMaterial;
+
 	// List of sockets on this sprite
 	UPROPERTY(Category=Sprite, EditAnywhere)
 	TArray<FPaperSpriteSocket> Sockets;
@@ -117,6 +121,10 @@ protected:
 #endif
 
 public:
+	// The point at which the alternate material takes over in the baked render data (or INDEX_NONE)
+	UPROPERTY()
+	int32 AlternateMaterialSplitIndex;
+
 	// Baked render data (triangle vertices, stored as XY UV tuples)
 	//   XY is the XZ position in world space, relative to the pivot
 	//   UV is normalized (0..1)
@@ -188,6 +196,15 @@ public:
 	// Return the default material for this sprite
 	UMaterialInterface* GetDefaultMaterial() const { return DefaultMaterial; }
 
+	// Return the alternate material for this sprite
+	UMaterialInterface* GetAlternateMaterial() const { return AlternateMaterial; }
+
+	// Returns either the default material (index 0) or alternate material (index 1)
+	UMaterialInterface* GetMaterial(int32 MaterialIndex) const;
+
+	// Returns the number of materials (1 or 2, depending on if there is alternate geometry)
+	int32 GetNumMaterials() const;
+
 	// Returns the render bounds of this sprite
 	FBoxSphereBounds GetRenderBounds() const;
 
diff --git a/Engine/Plugins/Experimental/Paper2D/Source/Paper2D/Classes/SpriteEditorOnlyTypes.h b/Engine/Plugins/Experimental/Paper2D/Source/Paper2D/Classes/SpriteEditorOnlyTypes.h
--- a/Engine/Plugins/Experimental/Paper2D/Source/Paper2D/Classes/SpriteEditorOnlyTypes.h
+++ b/Engine/Plugins/Experimental/Paper2D/Source/Paper2D/Classes/SpriteEditorOnlyTypes.h
@@ -41,6 +41,32 @@ namespace ESpritePolygonMode
 
 		// Fully custom geometry; edited by hand
 		FullyCustom,
+
+		// Diced (split up into smaller squares, including only non-empty ones in the final geometry)
+		Diced
+	};
+}
+
+// Method of specifying polygons for a sprite's render or collision data
+UENUM()
+namespace ESpriteSubdivisionMode
+{
+	enum Type
+	{
+		// Don't subdivide this sprite
+		NotSubdivided,
+
+		// Subdivide the sprite, 
+		SubdivideAutomatic,
+
+		// Bounding box that can be edited explicitly
+		SubdivideManually,
+
+		// Shrink-wrapped geometry
+		ShrinkWrapped,
+
+		// Fully custom geometry; edited by hand
+		FullyCustom,
 	};
 }
 
@@ -49,13 +75,21 @@ struct FSpritePolygonCollection
 {
 	GENERATED_USTRUCT_BODY()
 
-	UPROPERTY(Category=PolygonData, EditAnywhere)
+	UPROPERTY(Category=PolygonData, EditAnywhere, AdvancedDisplay)
 	TArray<FSpritePolygon> Polygons;
 
 	// The geometry type
 	UPROPERTY(Category=PolygonData, EditAnywhere)
 	TEnumAsByte<ESpritePolygonMode::Type> GeometryType;
 
+	// Size of a single subdivision (in pixels) in X (for Diced mode)
+	UPROPERTY(Category=PolygonData, EditAnywhere)
+	int32 PixelsPerSubdivisionX;
+
+	// Size of a single subdivision (in pixels) in Y (for Diced mode)
+	UPROPERTY(Category=PolygonData, EditAnywhere)
+	int32 PixelsPerSubdivisionY;
+
 	// Experimental: Hint to the triangulation routine that extra vertices should be preserved
 	UPROPERTY(Category=PolygonData, EditAnywhere, AdvancedDisplay)
 	bool bAvoidVertexMerging;
@@ -70,12 +104,25 @@ struct FSpritePolygonCollection
 
 	FSpritePolygonCollection()
 		: GeometryType(ESpritePolygonMode::TightBoundingBox)
+		, PixelsPerSubdivisionX(32)
+		, PixelsPerSubdivisionY(32)
 		, bAvoidVertexMerging(false)
 		, AlphaThreshold(0.0f)
 		, SimplifyEpsilon(2.0f)
 	{
 	}
 
+	void AddRectanglePolygon(FVector2D Position, FVector2D Size)
+	{
+		FSpritePolygon& Poly = *new (Polygons) FSpritePolygon();
+		new (Poly.Vertices) FVector2D(Position.X, Position.Y);
+		new (Poly.Vertices) FVector2D(Position.X + Size.X, Position.Y);
+		new (Poly.Vertices) FVector2D(Position.X + Size.X, Position.Y + Size.Y);
+		new (Poly.Vertices) FVector2D(Position.X, Position.Y + Size.Y);
+		Poly.BoxSize = Size;
+		Poly.BoxPosition = Position;
+	}
+
 	void Reset()
 	{
 		Polygons.Empty();
diff --git a/Engine/Plugins/Experimental/Paper2D/Source/Paper2D/Private/AlphaBitmap.h b/Engine/Plugins/Experimental/Paper2D/Source/Paper2D/Private/AlphaBitmap.h
new file mode 100644
--- /dev/null
+++ b/Engine/Plugins/Experimental/Paper2D/Source/Paper2D/Private/AlphaBitmap.h
@@ -0,0 +1,330 @@
+// Copyright 1998-2014 Epic Games, Inc. All Rights Reserved.
+
+#pragma once
+
+//////////////////////////////////////////////////////////////////////////
+// FBitmap
+
+struct FAlphaBitmap
+{
+	FAlphaBitmap(UTexture2D* SourceTexture, uint8 DefaultValue = 0)
+		: Width(0)
+		, Height(0)
+		, DefaultValue(DefaultValue)
+	{
+		ExtractFromTexture(SourceTexture);
+	}
+
+	void ExtractFromTexture(UTexture2D* SourceTexture)
+	{
+		// use the source art if it exists
+		FTextureSource* TextureSource = NULL;
+		if ((SourceTexture != NULL) && SourceTexture->Source.IsValid())
+		{
+			switch (SourceTexture->Source.GetFormat())
+			{
+			case TSF_G8:
+			case TSF_BGRA8:
+				TextureSource = &(SourceTexture->Source);
+				break;
+			default:
+				break;
+			};
+		}
+
+		if (TextureSource != NULL)
+		{
+			TArray<uint8> TextureRawData;
+			TextureSource->GetMipData(TextureRawData, 0);
+			int32 BytesPerPixel = TextureSource->GetBytesPerPixel();
+			ETextureSourceFormat PixelFormat = TextureSource->GetFormat();
+			Width = TextureSource->GetSizeX();
+			Height = TextureSource->GetSizeY();
+			RawData.SetNumZeroed(Width * Height);
+
+			for (int Y = 0; Y < Height; ++Y)
+			{
+				for (int X = 0; X < Width; ++X)
+				{
+					int32 PixelByteOffset = (X + Y * Width) * BytesPerPixel;
+					const uint8* PixelPtr = TextureRawData.GetTypedData() + PixelByteOffset;
+					FColor Color;
+					if (PixelFormat == TSF_BGRA8)
+					{
+						Color = *((FColor*)PixelPtr);
+					}
+					else
+					{
+						checkSlow(PixelFormat == TSF_G8);
+						const uint8 Intensity = *PixelPtr;
+						Color = FColor(Intensity, Intensity, Intensity, Intensity);
+					}
+
+					RawData[Y * Width + X] = Color.A;
+				}
+			}
+		}
+	}
+
+	// Create an empty bitmap
+	FAlphaBitmap(int Width, int Height, int DefaultValue = 0)
+		: Width(Width)
+		, Height(Height)
+		, DefaultValue(DefaultValue)
+	{
+		RawData.SetNumZeroed(Width * Height);
+		ClearToDefaultValue();
+	}
+
+	bool IsValid() const
+	{
+		return (Width != 0) && (Height != 0);
+	}
+
+	uint8 GetPixel(int X, int Y) const
+	{
+		if ((X >= 0) && (X < Width) && (Y >= 0) && (Y < Height))
+		{
+			return RawData[Y * Width + X];
+		}
+		else
+		{
+			return DefaultValue;
+		}
+	}
+
+	void SetPixel(int X, int Y, uint8 Value)
+	{
+		if ((X >= 0) && (X < Width) && (Y >= 0) && (Y < Height))
+		{
+			RawData[Y * Width + X] = Value;
+		}
+	}
+
+	void ClearToDefaultValue()
+	{
+		int RawDataCount = Width * Height;
+		for (int RawDataIndex = 0; RawDataIndex < RawDataCount; ++RawDataIndex)
+		{
+			RawData[RawDataIndex] = DefaultValue;
+		}
+	}
+
+	bool IsColumnEqual(int X, int Y0, int Y1, int Target) const
+	{
+		for (int Y = Y0; Y <= Y1; ++Y)
+		{
+			if (GetPixel(X, Y) != Target)
+			{
+				return false;
+			}
+		}
+		return true;
+	}
+
+	bool IsRowEqual(int X0, int X1, int Y, int Target) const
+	{
+		for (int X = X0; X <= X1; ++X)
+		{
+			if (GetPixel(X, Y) != Target)
+			{
+				return false;
+			}
+		}
+		return true;
+	}
+
+	// Is the rectangle empty (X0..X1, Y0..Y1 inclusive)
+	bool IsRegionEqual(int32 X0, int32 Y0, int32 X1, int32 Y1, int32 Target)
+	{
+		bool bEqual = true;
+		for (int32 Y = Y0; (Y < Y1) && bEqual; ++Y)
+		{
+			bEqual = bEqual && IsRowEqual(X0, X1, Y, Target);
+		}
+		return bEqual;
+	}
+
+	bool IsColumnEmpty(int X, int Y0, int Y1) const
+	{
+		return IsColumnEqual(X, Y0, Y1, 0);
+	}
+
+	bool IsRowEmpty(int X0, int X1, int Y)
+	{
+		return IsRowEqual(X0, X1, Y, 0);
+	}
+
+	// Is the rectangle empty (X0..X1, Y0..Y1 inclusive)
+	bool IsRegionEmpty(int32 X0, int32 Y0, int32 X1, int32 Y1)
+	{
+		return IsRegionEqual(X0, Y0, X1, Y1, 0);
+	}
+
+	// Returns the tight bounding box around pixels that are not 0
+	void GetTightBounds(FIntPoint& OutOrigin, FIntPoint& OutDimension)
+	{
+		OutOrigin.X = 0;
+		OutOrigin.Y = 0;
+		OutDimension.X = Width;
+		OutDimension.Y = Height;
+		TightenBounds(OutOrigin, OutDimension);
+	}
+
+	// Returns the tight bounding box around pixels that are not 0
+	void TightenBounds(FIntPoint& InOutOrigin, FIntPoint& InOutDimension)
+	{
+		int Top = InOutOrigin.Y;
+		int Bottom = InOutOrigin.Y + InOutDimension.Y - 1;
+		int Left = InOutOrigin.X;
+		int Right = InOutOrigin.X + InOutDimension.X - 1;
+
+		while (Top < Bottom && IsRowEmpty(Left, Right, Top))
+		{
+			++Top;
+		}
+		while (Bottom >= Top && IsRowEmpty(Left, Right, Bottom))
+		{
+			--Bottom;
+		}
+		while (Left < Right && IsColumnEmpty(Left, Top, Bottom))
+		{
+			++Left;
+		}
+		while (Right >= Left && IsColumnEmpty(Right, Top, Bottom))
+		{
+			--Right;
+		}
+
+		InOutOrigin.X = Left;
+		InOutOrigin.Y = Top;
+		InOutDimension.X = Right - Left + 1;
+		InOutDimension.Y = Bottom - Top + 1;
+	}
+
+	// Sets pixels to 1 marking the rectangle outline
+	void DrawRectOutline(int StartX, int StartY, int Width, int Height)
+	{
+		const int X0 = StartX;
+		const int Y0 = StartY;
+		const int X1 = StartX + Width - 1;
+		const int Y1 = StartY + Height - 1;
+		for (int Y = Y0; Y <= Y1; ++Y)
+		{
+			SetPixel(X0, Y, 1);
+			SetPixel(X1, Y, 1);
+		}
+		for (int X = X0; X <= X1; ++X)
+		{
+			SetPixel(X, Y0, 1);
+			SetPixel(X, Y1, 1);
+		}
+	}
+
+	void FillRect(int StartX, int StartY, int Width, int Height)
+	{
+		const int X0 = StartX;
+		const int Y0 = StartY;
+		const int X1 = StartX + Width - 1;
+		const int Y1 = StartY + Height - 1;
+		for (int Y = Y0; Y <= Y1; ++Y)
+		{
+			for (int X = X0; X <= X1; ++X)
+			{
+				SetPixel(X, Y, 1);
+			}
+		}
+	}
+
+	// Wind through the bitmap from StartX, StartY to find the closest hit point
+	bool FoundClosestValidPoint(int StartX, int StartY, const int MaxAllowedSearchDistance, FIntPoint& OutHitPoint)
+	{
+		// Constrain the point within the image bounds
+		StartX = FMath::Clamp(StartX, 0, Width - 1);
+		StartY = FMath::Clamp(StartY, 0, Height - 1);
+		OutHitPoint.X = StartX;
+		OutHitPoint.Y = StartY;
+
+		// Should probably calculate a better max based on StartX and StartY
+		int RequiredSearchDistance = (FMath::Max(Width, Height) + 1) / 2;
+		int MaxSearchDistance = FMath::Min(RequiredSearchDistance, MaxAllowedSearchDistance);
+		int SearchDistance = 0;
+		while (SearchDistance < MaxSearchDistance)
+		{
+			int X0 = FMath::Max(StartX - SearchDistance, 0);
+			int X1 = FMath::Min(StartX + SearchDistance, Width - 1);
+			int Y0 = FMath::Max(StartY - SearchDistance, 0);
+			int Y1 = FMath::Min(StartY + SearchDistance, Height - 1);
+			// Search along the rectangular edge
+			for (int Y = Y0; Y <= Y1; ++Y)
+			{
+				if (GetPixel(X0, Y) != 0) 
+				{
+					OutHitPoint.X = X0;
+					OutHitPoint.Y = Y;
+					return true;
+				}
+				if (GetPixel(X1, Y) != 0) 
+				{
+					OutHitPoint.X = X1;
+					OutHitPoint.Y = Y;
+					return true;
+				}
+			}
+			for (int X = X0; X <= X1; ++X)
+			{
+				if (GetPixel(X, Y0) != 0) 
+				{
+					OutHitPoint.X = X;
+					OutHitPoint.Y = Y0;
+					return true;
+				}
+				if (GetPixel(X, Y1) != 0) 
+				{
+					OutHitPoint.X = X;
+					OutHitPoint.Y = Y1;
+					return true;
+				}
+			}
+			SearchDistance += 1;
+		}
+		return false;
+	}
+
+	void ThresholdImage(int32 AlphaThreshold)
+	{
+		for (uint8& Pixel : RawData)
+		{
+			if (Pixel > AlphaThreshold)
+			{
+				Pixel = 1;
+			}
+			else
+			{
+				Pixel = DefaultValue;
+			}
+		}
+	}
+
+	// Flushes values smaller or equal to LowAlphaThreshold to 0, and values greater or equal to HighAlphaThreshold to 255
+	void ThresholdImageBothWays(int32 LowAlphaThreshold, int32 HighAlphaThreshold)
+	{
+		for (uint8& Pixel : RawData)
+		{
+			if (Pixel <= LowAlphaThreshold)
+			{
+				Pixel = 0;
+			}
+			else if (Pixel >= HighAlphaThreshold)
+			{
+				Pixel = 255;
+			}
+		}
+	}
+
+public:
+	TArray<uint8> RawData;
+	int32 Width;
+	int32 Height;
+	uint8 DefaultValue;
+};
diff --git a/Engine/Plugins/Experimental/Paper2D/Source/Paper2D/Private/BitmapUtils.h b/Engine/Plugins/Experimental/Paper2D/Source/Paper2D/Private/BitmapUtils.h
--- a/Engine/Plugins/Experimental/Paper2D/Source/Paper2D/Private/BitmapUtils.h
+++ b/Engine/Plugins/Experimental/Paper2D/Source/Paper2D/Private/BitmapUtils.h
@@ -2,173 +2,23 @@
 
 #pragma once
 
+#include "AlphaBitmap.h"
+
 //////////////////////////////////////////////////////////////////////////
 // FBitmap
 
-struct FBitmap
+struct FBitmap : public FAlphaBitmap
 {
-	FBitmap(UTexture2D* SourceTexture, int AlphaThreshold, uint8 DefaultValue = 0)
-		: Width(0)
-		, Height(0)
-		, DefaultValue(DefaultValue)
+	FBitmap(UTexture2D* SourceTexture, int32 AlphaThreshold, uint8 DefaultValue = 0)
+		: FAlphaBitmap(SourceTexture, DefaultValue)
 	{
-		ExtractFromTexture(SourceTexture, AlphaThreshold);
-	}
-
-	void ExtractFromTexture(UTexture2D* SourceTexture, int AlphaThreshold)
-	{
-		// use the source art if it exists
-		FTextureSource* TextureSource = NULL;
-		if ((SourceTexture != NULL) && SourceTexture->Source.IsValid())
-		{
-			switch (SourceTexture->Source.GetFormat())
-			{
-			case TSF_G8:
-			case TSF_BGRA8:
-				TextureSource = &(SourceTexture->Source);
-				break;
-			default:
-				break;
-			};
-		}
-
-		if (TextureSource != NULL)
-		{
-			TArray<uint8> TextureRawData;
-			TextureSource->GetMipData(TextureRawData, 0);
-			int32 BytesPerPixel = TextureSource->GetBytesPerPixel();
-			ETextureSourceFormat PixelFormat = TextureSource->GetFormat();
-			Width = TextureSource->GetSizeX();
-			Height = TextureSource->GetSizeY();
-			RawData.SetNumZeroed(Width * Height);
-
-			for (int Y = 0; Y < Height; ++Y)
-			{
-				for (int X = 0; X < Width; ++X)
-				{
-					int32 PixelByteOffset = (X + Y * Width) * BytesPerPixel;
-					const uint8* PixelPtr = TextureRawData.GetTypedData() + PixelByteOffset;
-					FColor Color;
-					if (PixelFormat == TSF_BGRA8)
-					{
-						Color = *((FColor*)PixelPtr);
-					}
-					else
-					{
-						checkSlow(PixelFormat == TSF_G8);
-						const uint8 Intensity = *PixelPtr;
-						Color = FColor(Intensity, Intensity, Intensity, Intensity);
-					}
-
-					if (Color.A > AlphaThreshold)
-					{
-						RawData[Y * Width + X] = 1;
-					}
-					else
-					{
-						RawData[Y * Width + X] = DefaultValue;
-					}
-				}
-			}
-		}
+		ThresholdImage(AlphaThreshold);
 	}
 
 	// Create an empty bitmap
 	FBitmap(int Width, int Height, int DefaultValue = 0)
-		: Width(Width)
-		, Height(Height)
-		, DefaultValue(DefaultValue)
+		: FAlphaBitmap(Width, Height, DefaultValue)
 	{
-		RawData.SetNumZeroed(Width * Height);
-		ClearToDefaultValue();
-	}
-
-	bool IsValid() const
-	{
-		return Width != 0 && Height != 0;
-	}
-
-	uint8 GetPixel(int X, int Y)
-	{
-		if (X >= 0 && X < Width && Y >= 0 && Y < Height)
-		{
-			return RawData[Y * Width + X];
-		}
-		else
-		{
-			return DefaultValue;
-		}
-	}
-
-	void SetPixel(int X, int Y, uint8 Value)
-	{
-		if (X >= 0 && X < Width && Y >= 0 && Y < Height)
-		{
-			RawData[Y * Width + X] = Value;
-		}
-	}
-
-	void ClearToDefaultValue()
-	{
-		int RawDataCount = Width * Height;
-		for (int RawDataIndex = 0; RawDataIndex < RawDataCount; ++RawDataIndex)
-		{
-			RawData[RawDataIndex] = DefaultValue;
-		}
-	}
-
-	bool IsColumnEmpty(int X, int Y0, int Y1)
-	{
-		for (int Y = Y0; Y <= Y1; ++Y)
-		{
-			if (GetPixel(X, Y) != 0)
-			{
-				return false;
-			}
-		}
-		return true;
-	}
-
-	bool IsRowEmpty(int X0, int X1, int Y)
-	{
-		for (int X = X0; X <= X1; ++X)
-		{
-			if (GetPixel(X, Y) != 0)
-			{
-				return false;
-			}
-		}
-		return true;
-	}
-
-	// Returns the tight bounding box around pixels that are not 0
-	void GetTightBounds(FIntPoint& Origin, FIntPoint& Dimension)
-	{
-		int Top = 0;
-		int Bottom = Height - 1;
-		int Left = 0;
-		int Right = Width - 1;
-		while (Top < Bottom && IsRowEmpty(Left, Right, Top))
-		{
-			++Top;
-		}
-		while (Bottom >= Top && IsRowEmpty(Left, Right, Bottom))
-		{
-			--Bottom;
-		}
-		while (Left < Right && IsColumnEmpty(Left, Top, Bottom))
-		{
-			++Left;
-		}
-		while (Right >= Left && IsColumnEmpty(Right, Top, Bottom))
-		{
-			--Right;
-		}
-
-		Origin.X = Left;
-		Origin.Y = Top;
-		Dimension.X = Right - Left + 1;
-		Dimension.Y = Bottom - Top + 1;
 	}
 
 	// Performs a flood fill on the target bitmap, with the boundary defined
@@ -250,40 +100,6 @@ struct FBitmap
 		return false;
 	}
 
-	// Sets pixels to 1 marking the rectangle outline
-	void DrawRectOutline(int StartX, int StartY, int Width, int Height)
-	{
-		const int X0 = StartX;
-		const int Y0 = StartY;
-		const int X1 = StartX + Width - 1;
-		const int Y1 = StartY + Height - 1;
-		for (int Y = Y0; Y <= Y1; ++Y)
-		{
-			SetPixel(X0, Y, 1);
-			SetPixel(X1, Y, 1);
-		}
-		for (int X = X0; X <= X1; ++X)
-		{
-			SetPixel(X, Y0, 1);
-			SetPixel(X, Y1, 1);
-		}
-	}
-
-	void FillRect(int StartX, int StartY, int Width, int Height)
-	{
-		const int X0 = StartX;
-		const int Y0 = StartY;
-		const int X1 = StartX + Width - 1;
-		const int Y1 = StartY + Height - 1;
-		for (int Y = Y0; Y <= Y1; ++Y)
-		{
-			for (int X = X0; X <= X1; ++X)
-			{
-				SetPixel(X, Y, 1);
-			}
-		}
-	}
-
 	// Finds the rect of the contour of the shape clicked on, extended by rectangles to support
 	// separated but intersecting islands.
 	bool HasConnectedRect(const int X, const int Y, bool Extend, FIntPoint& OutOrigin, FIntPoint& OutDimension)
@@ -330,61 +146,6 @@ struct FBitmap
 		return true;
 	}
 
-	// Wind through the bitmap from StartX, StartY to find the closest hit point
-	bool FoundClosestValidPoint(int StartX, int StartY, const int MaxAllowedSearchDistance, FIntPoint& OutHitPoint)
-	{
-		// Constrain the point within the image bounds
-		StartX = FMath::Clamp(StartX, 0, Width - 1);
-		StartY = FMath::Clamp(StartY, 0, Height - 1);
-		OutHitPoint.X = StartX;
-		OutHitPoint.Y = StartY;
-
-		// Should probably calculate a better max based on StartX and StartY
-		int RequiredSearchDistance = (FMath::Max(Width, Height) + 1) / 2;
-		int MaxSearchDistance = FMath::Min(RequiredSearchDistance, MaxAllowedSearchDistance);
-		int SearchDistance = 0;
-		while (SearchDistance < MaxSearchDistance)
-		{
-			int X0 = FMath::Max(StartX - SearchDistance, 0);
-			int X1 = FMath::Min(StartX + SearchDistance, Width - 1);
-			int Y0 = FMath::Max(StartY - SearchDistance, 0);
-			int Y1 = FMath::Min(StartY + SearchDistance, Height - 1);
-			// Search along the rectangular edge
-			for (int Y = Y0; Y <= Y1; ++Y)
-			{
-				if (GetPixel(X0, Y) != 0) 
-				{
-					OutHitPoint.X = X0;
-					OutHitPoint.Y = Y;
-					return true;
-				}
-				if (GetPixel(X1, Y) != 0) 
-				{
-					OutHitPoint.X = X1;
-					OutHitPoint.Y = Y;
-					return true;
-				}
-			}
-			for (int X = X0; X <= X1; ++X)
-			{
-				if (GetPixel(X, Y0) != 0) 
-				{
-					OutHitPoint.X = X;
-					OutHitPoint.Y = Y0;
-					return true;
-				}
-				if (GetPixel(X, Y1) != 0) 
-				{
-					OutHitPoint.X = X;
-					OutHitPoint.Y = Y1;
-					return true;
-				}
-			}
-			SearchDistance += 1;
-		}
-		return false;
-	}
-
 	// Detects all valid rects in this bitmap
 	void ExtractRects(TArray<FIntRect>& OutRects)
 	{
@@ -434,11 +195,4 @@ struct FBitmap
 			}
 		}
 	}
-
-	TArray<uint8> RawData;
-	int32 Width;
-	int32 Height;
-	uint8 DefaultValue;
 };
-
-
diff --git a/Engine/Plugins/Experimental/Paper2D/Source/Paper2D/Private/PaperSprite.cpp b/Engine/Plugins/Experimental/Paper2D/Source/Paper2D/Private/PaperSprite.cpp
--- a/Engine/Plugins/Experimental/Paper2D/Source/Paper2D/Private/PaperSprite.cpp
+++ b/Engine/Plugins/Experimental/Paper2D/Source/Paper2D/Private/PaperSprite.cpp
@@ -213,7 +213,7 @@ void UPaperSprite::ExtractSourceRegionFromTexturePoint(const FVector2D& SourcePo
 	FIntPoint ClosestValidPoint;
 
 	FBitmap Bitmap(SourceTexture, 0, 0);
-	if (Bitmap.IsValid() && Bitmap.FoundClosestValidPoint(SourceIntPoint.X, SourceIntPoint.Y, 10, /*out*/ClosestValidPoint))
+	if (Bitmap.IsValid() && Bitmap.FoundClosestValidPoint(SourceIntPoint.X, SourceIntPoint.Y, 10, /*out*/ ClosestValidPoint))
 	{
 		FIntPoint Origin;
 		FIntPoint Dimension;
@@ -254,6 +254,8 @@ UPaperSprite::UPaperSprite(const FPostConstructInitializeProperties& PCIP)
 {
 	// Default to using physics
 	SpriteCollisionDomain = ESpriteCollisionMode::Use3DPhysics;
+	
+	AlternateMaterialSplitIndex = INDEX_NONE;
 
 #if WITH_EDITORONLY_DATA
 	PivotMode = ESpritePivotMode::Center_Center;
@@ -264,8 +266,11 @@ UPaperSprite::UPaperSprite(const FPostConstructInitializeProperties& PCIP)
 	PixelsPerUnrealUnit = 2.56f;
 #endif
 
-	static ConstructorHelpers::FObjectFinder<UMaterial> DefaultMaterialRef(TEXT("/Paper2D/DefaultSpriteMaterial.DefaultSpriteMaterial"));
-	DefaultMaterial = DefaultMaterialRef.Object;
+	static ConstructorHelpers::FObjectFinder<UMaterialInterface> MaskedMaterialRef(TEXT("/Paper2D/MaskedUnlitSpriteMaterial.MaskedUnlitSpriteMaterial"));
+	DefaultMaterial = MaskedMaterialRef.Object;
+
+	static ConstructorHelpers::FObjectFinder<UMaterialInterface> OpaqueMaterialRef(TEXT("/Paper2D/OpaqueUnlitSpriteMaterial.OpaqueUnlitSpriteMaterial"));
+	AlternateMaterial = OpaqueMaterialRef.Object;
 }
 
 #if WITH_EDITORONLY_DATA
@@ -276,6 +281,14 @@ void UPaperSprite::PostEditChangeProperty(FPropertyChangedEvent& PropertyChanged
 		PixelsPerUnrealUnit = 1.0f;
 	}
 
+	if (CollisionGeometry.GeometryType == ESpritePolygonMode::Diced)
+	{
+		// Disallow dicing on collision geometry for now
+		CollisionGeometry.GeometryType == ESpritePolygonMode::SourceBoundingBox;
+	}
+	RenderGeometry.PixelsPerSubdivisionX = FMath::Max(RenderGeometry.PixelsPerSubdivisionX, 4);
+	RenderGeometry.PixelsPerSubdivisionY = FMath::Max(RenderGeometry.PixelsPerSubdivisionY, 4);
+
 	SourceDimension.X = FMath::Max(SourceDimension.X, 0.0f);
 	SourceDimension.Y = FMath::Max(SourceDimension.Y, 0.0f);
 
@@ -379,6 +392,7 @@ void UPaperSprite::RebuildCollisionData()
 		BodySetup->CollisionTraceFlag = CTF_UseSimpleAsComplex;
 		switch (CollisionGeometry.GeometryType)
 		{
+		case ESpritePolygonMode::Diced:
 		case ESpritePolygonMode::SourceBoundingBox:
 			BuildBoundingBoxCollisionData(/*bUseTightBounds=*/ false);
 			break;
@@ -549,8 +563,11 @@ void UPaperSprite::BuildBoundingBoxCollisionData(bool bUseTightBounds)
 
 void UPaperSprite::RebuildRenderData()
 {
+	FSpritePolygonCollection AlternateGeometry;
+
 	switch (RenderGeometry.GeometryType)
 	{
+	case ESpritePolygonMode::Diced:
 	case ESpritePolygonMode::SourceBoundingBox:
 		CreatePolygonFromBoundingBox(RenderGeometry, /*bUseTightBounds=*/ false);
 		break;
@@ -570,10 +587,6 @@ void UPaperSprite::RebuildRenderData()
 		check(false); // unknown mode
 	};
 
-	// Triangulate the render geometry
-	TArray<FVector2D> TriangluatedPoints;
-	Triangulate(RenderGeometry, /*out*/ TriangluatedPoints);
-
 	// Determine the texture size
 	UTexture2D* EffectiveTexture = GetBakedTexture();
 
@@ -596,6 +609,79 @@ void UPaperSprite::RebuildRenderData()
 
 	const float UnitsPerPixel = GetUnrealUnitsPerPixel();
 
+	if ((RenderGeometry.GeometryType == ESpritePolygonMode::Diced) && (EffectiveTexture != nullptr))
+	{
+		const int32 AlphaThresholdInt = FMath::Clamp<int32>(RenderGeometry.AlphaThreshold * 255, 0, 255);
+		FAlphaBitmap SourceBitmap(EffectiveTexture);
+		SourceBitmap.ThresholdImageBothWays(AlphaThresholdInt, 255);
+
+		bool bSeparateOpaqueSections = true;
+
+		// Dice up the source geometry and sort into translucent and opaque sections
+		RenderGeometry.Polygons.Empty();
+
+		const int32 X0 = (int32)SourceUV.X;
+		const int32 Y0 = (int32)SourceUV.Y;
+		const int32 X1 = (int32)(SourceUV.X + SourceDimension.X);
+		const int32 Y1 = (int32)(SourceUV.Y + SourceDimension.Y);
+
+		for (int32 Y = Y0; Y < Y1; Y += RenderGeometry.PixelsPerSubdivisionY)
+		{
+			const int32 TileHeight = FMath::Min(RenderGeometry.PixelsPerSubdivisionY, Y1 - Y);
+
+			for (int32 X = X0; X < X1; X += RenderGeometry.PixelsPerSubdivisionX)
+			{
+				const int32 TileWidth = FMath::Min(RenderGeometry.PixelsPerSubdivisionX, X1 - X);
+
+				if (!SourceBitmap.IsRegionEmpty(X, Y, X + TileWidth - 1, Y + TileHeight - 1))
+				{
+					FIntPoint Origin(X, Y);
+					FIntPoint Dimension(TileWidth, TileHeight);
+					
+					SourceBitmap.TightenBounds(Origin, Dimension);
+
+					bool bOpaqueSection = false;
+					if (bSeparateOpaqueSections)
+					{
+						if (SourceBitmap.IsRegionEqual(Origin.X, Origin.Y, Origin.X + Dimension.X - 1, Origin.Y + Dimension.Y - 1, 255))
+						{
+							bOpaqueSection = true;
+						}
+					}
+
+					if (bOpaqueSection)
+					{
+						AlternateGeometry.AddRectanglePolygon(Origin, Dimension);
+					}
+					else
+					{
+						RenderGeometry.AddRectanglePolygon(Origin, Dimension);
+					}
+				}
+			}
+		}
+	}
+
+	// Triangulate the render geometry
+	TArray<FVector2D> TriangluatedPoints;
+	Triangulate(RenderGeometry, /*out*/ TriangluatedPoints);
+
+	// Triangulate the alternate render geometry, if present
+	if (AlternateGeometry.Polygons.Num() > 0)
+	{
+		TArray<FVector2D> AlternateTriangluatedPoints;
+		Triangulate(AlternateGeometry, /*out*/ AlternateTriangluatedPoints);
+
+		AlternateMaterialSplitIndex = TriangluatedPoints.Num();
+		TriangluatedPoints.Append(AlternateTriangluatedPoints);
+		RenderGeometry.Polygons.Append(AlternateGeometry.Polygons);
+	}
+	else
+	{
+		AlternateMaterialSplitIndex = INDEX_NONE;
+	}
+
+	// Bake the verts
 	BakedRenderData.Empty(TriangluatedPoints.Num());
 	for (int32 PointIndex = 0; PointIndex < TriangluatedPoints.Num(); ++PointIndex)
 	{
@@ -862,19 +948,15 @@ void UPaperSprite::CreatePolygonFromBoundingBox(FSpritePolygonCollection& GeomOw
 
 	// Put the bounding box into the geometry array
 	GeomOwner.Polygons.Empty();
-	FSpritePolygon& Poly = *new (GeomOwner.Polygons) FSpritePolygon();
-	new (Poly.Vertices) FVector2D(BoxPosition.X, BoxPosition.Y);
-	new (Poly.Vertices) FVector2D(BoxPosition.X + BoxSize.X, BoxPosition.Y);
-	new (Poly.Vertices) FVector2D(BoxPosition.X + BoxSize.X, BoxPosition.Y + BoxSize.Y);
-	new (Poly.Vertices) FVector2D(BoxPosition.X, BoxPosition.Y + BoxSize.Y);
-	Poly.BoxSize = BoxSize;
-	Poly.BoxPosition = BoxPosition;
+	GeomOwner.AddRectanglePolygon(BoxPosition, BoxSize);
 }
 
 void UPaperSprite::Triangulate(const FSpritePolygonCollection& Source, TArray<FVector2D>& Target)
 {
 	Target.Empty();
 
+	TArray<FClipSMTriangle> AllGeneratedTriangles;
+
 	for (int32 PolygonIndex = 0; PolygonIndex < Source.Polygons.Num(); ++PolygonIndex)
 	{
 		const FSpritePolygon& SourcePoly = Source.Polygons[PolygonIndex];
@@ -882,7 +964,7 @@ void UPaperSprite::Triangulate(const FSpritePolygonCollection& Source, TArray<FV
 		{
 			// Convert our format into one the triangulation library supports
 			FClipSMPolygon ClipPolygon(0);
-			for (int32 VertexIndex = 0; VertexIndex < SourcePoly.Vertices.Num() ; ++VertexIndex)
+			for (int32 VertexIndex = 0; VertexIndex < SourcePoly.Vertices.Num(); ++VertexIndex)
 			{
 				FClipSMVertex* ClipVertex = new (ClipPolygon.Vertices) FClipSMVertex;
 				FMemory::Memzero(ClipVertex, sizeof(FClipSMVertex));
@@ -897,18 +979,23 @@ void UPaperSprite::Triangulate(const FSpritePolygonCollection& Source, TArray<FV
 			TArray<FClipSMTriangle> GeneratedTriangles;
 			if (TriangulatePoly(/*out*/ GeneratedTriangles, ClipPolygon, Source.bAvoidVertexMerging))
 			{
-				// Convert the triangles back to our 2D data structure
-				for (int32 TriangleIndex = 0; TriangleIndex < GeneratedTriangles.Num(); ++TriangleIndex)
-				{
-					const FClipSMTriangle& Triangle = GeneratedTriangles[TriangleIndex];
-
-					new (Target) FVector2D(Triangle.Vertices[0].Pos.X, Triangle.Vertices[0].Pos.Z);
-					new (Target) FVector2D(Triangle.Vertices[1].Pos.X, Triangle.Vertices[1].Pos.Z);
-					new (Target) FVector2D(Triangle.Vertices[2].Pos.X, Triangle.Vertices[2].Pos.Z);
-				}
+				AllGeneratedTriangles.Append(GeneratedTriangles);
 			}
 		}
 	}
+
+	if (!Source.bAvoidVertexMerging && (Source.Polygons.Num() > 1) && (AllGeneratedTriangles.Num() > 1))
+	{
+		RemoveRedundantTriangles(AllGeneratedTriangles);
+	}
+
+	// Convert the triangles back to our 2D data structure
+	for (const FClipSMTriangle& Triangle : AllGeneratedTriangles)
+	{
+		new (Target) FVector2D(Triangle.Vertices[0].Pos.X, Triangle.Vertices[0].Pos.Z);
+		new (Target) FVector2D(Triangle.Vertices[1].Pos.X, Triangle.Vertices[1].Pos.Z);
+		new (Target) FVector2D(Triangle.Vertices[2].Pos.X, Triangle.Vertices[2].Pos.Z);
+	}
 }
 
 void UPaperSprite::ExtractRectsFromTexture(UTexture2D* Texture, TArray<FIntRect>& OutRects)
@@ -1166,3 +1253,24 @@ UTexture2D* UPaperSprite::GetBakedTexture() const
 {
 	return (BakedSourceTexture != nullptr) ? BakedSourceTexture : SourceTexture;
 }
+
+UMaterialInterface* UPaperSprite::GetMaterial(int32 MaterialIndex) const
+{
+	if (MaterialIndex == 0)
+	{
+		return GetDefaultMaterial();
+	}
+	else if (MaterialIndex == 1)
+	{
+		return GetAlternateMaterial();
+	}
+	else
+	{
+		return nullptr;
+	}
+}
+
+int32 UPaperSprite::GetNumMaterials() const
+{
+	return (AlternateMaterialSplitIndex != INDEX_NONE) ? 2 : 1;
+}
diff --git a/Engine/Plugins/Experimental/Paper2D/Source/Paper2D/Private/PaperSpriteComponent.cpp b/Engine/Plugins/Experimental/Paper2D/Source/Paper2D/Private/PaperSpriteComponent.cpp
--- a/Engine/Plugins/Experimental/Paper2D/Source/Paper2D/Private/PaperSpriteComponent.cpp
+++ b/Engine/Plugins/Experimental/Paper2D/Source/Paper2D/Private/PaperSpriteComponent.cpp
@@ -60,10 +60,13 @@ void UPaperSpriteComponent::PostLoad()
 FPrimitiveSceneProxy* UPaperSpriteComponent::CreateSceneProxy()
 {
 	FPaperSpriteSceneProxy* NewProxy = new FPaperSpriteSceneProxy(this);
-	FSpriteDrawCallRecord DrawCall;
-	DrawCall.BuildFromSprite(SourceSprite);
-	DrawCall.Color = SpriteColor;
-	NewProxy->SetDrawCall_RenderThread(DrawCall);
+	if (SourceSprite != nullptr)
+	{
+		FSpriteDrawCallRecord DrawCall;
+		DrawCall.BuildFromSprite(SourceSprite);
+		DrawCall.Color = SpriteColor;
+		NewProxy->SetSprite_RenderThread(DrawCall, SourceSprite->AlternateMaterialSplitIndex);
+	}
 	return NewProxy;
 }
 
@@ -103,13 +106,15 @@ void UPaperSpriteComponent::SendRenderDynamicData_Concurrent()
 		FSpriteDrawCallRecord DrawCall;
 		DrawCall.BuildFromSprite(SourceSprite);
 		DrawCall.Color = SpriteColor;
+		int32 SplitIndex = (SourceSprite != nullptr) ? SourceSprite->AlternateMaterialSplitIndex : INDEX_NONE;
 
-		ENQUEUE_UNIQUE_RENDER_COMMAND_TWOPARAMETER(
+		ENQUEUE_UNIQUE_RENDER_COMMAND_THREEPARAMETER(
 				FSendPaperSpriteComponentDynamicData,
-				FPaperRenderSceneProxy*,InSceneProxy,(FPaperRenderSceneProxy*)SceneProxy,
+				FPaperSpriteSceneProxy*,InSceneProxy,(FPaperSpriteSceneProxy*)SceneProxy,
 				FSpriteDrawCallRecord,InSpriteToSend,DrawCall,
+				int32,InSplitIndex,SplitIndex,
 			{
-				InSceneProxy->SetDrawCall_RenderThread(InSpriteToSend);
+				InSceneProxy->SetSprite_RenderThread(InSpriteToSend, InSplitIndex);
 			});
 	}
 }
@@ -224,8 +229,7 @@ UMaterialInterface* UPaperSpriteComponent::GetMaterial(int32 MaterialIndex) cons
 	}
 	else if (SourceSprite != nullptr)
 	{
-		//@TODO: PAPER2D: Need to support this for multi-material support
-		return SourceSprite->GetDefaultMaterial();
+		return SourceSprite->GetMaterial(MaterialIndex);
 	}
 
 	return nullptr;
@@ -244,7 +248,14 @@ void UPaperSpriteComponent::GetStreamingTextureInfo(TArray<FStreamingTexturePrim
 
 int32 UPaperSpriteComponent::GetNumMaterials() const
 {
-	return FMath::Max<int32>(Materials.Num(), 1);
+	if (SourceSprite != nullptr)
+	{
+		return FMath::Max<int32>(Materials.Num(), SourceSprite->GetNumMaterials());
+	}
+	else
+	{
+		return FMath::Max<int32>(Materials.Num(), 1);
+	}
 }
 
 UPaperSprite* UPaperSpriteComponent::GetSprite()
diff --git a/Engine/Plugins/Experimental/Paper2D/Source/Paper2D/Private/PaperSpriteSceneProxy.cpp b/Engine/Plugins/Experimental/Paper2D/Source/Paper2D/Private/PaperSpriteSceneProxy.cpp
--- a/Engine/Plugins/Experimental/Paper2D/Source/Paper2D/Private/PaperSpriteSceneProxy.cpp
+++ b/Engine/Plugins/Experimental/Paper2D/Source/Paper2D/Private/PaperSpriteSceneProxy.cpp
@@ -13,6 +13,8 @@ FPaperSpriteSceneProxy::FPaperSpriteSceneProxy(const UPaperSpriteComponent* InCo
 {
 	WireframeColor = InComponent->GetWireframeColor();
 	Material = InComponent->GetMaterial(0);
+	AlternateMaterial = InComponent->GetMaterial(1);
+	MaterialSplitIndex = INDEX_NONE;
 	MaterialRelevance = InComponent->GetMaterialRelevance();
 
 	SourceSprite = InComponent->SourceSprite; //@TODO: This is totally not threadsafe, and won't keep up to date if the actor's sprite changes, etc....
@@ -73,3 +75,43 @@ void FPaperSpriteSceneProxy::DrawDynamicElements(FPrimitiveDrawInterface* PDI, c
 
 	FPaperRenderSceneProxy::DrawDynamicElements(PDI, View);
 }
+
+void FPaperSpriteSceneProxy::SetSprite_RenderThread(const FSpriteDrawCallRecord& NewDynamicData, int32 SplitIndex)
+{
+	BatchedSprites.Empty();
+	AlternateBatchedSprites.Empty();
+
+	if (SplitIndex != INDEX_NONE)
+	{
+		FSpriteDrawCallRecord& Record = *new (BatchedSprites) FSpriteDrawCallRecord;
+		FSpriteDrawCallRecord& AltRecord = *new (AlternateBatchedSprites) FSpriteDrawCallRecord;
+		
+		Record.Color = NewDynamicData.Color;
+		Record.Destination = NewDynamicData.Destination;
+		Record.Texture = NewDynamicData.Texture;
+		Record.RenderVerts.Append(NewDynamicData.RenderVerts.GetTypedData(), SplitIndex);
+
+		AltRecord.Color = NewDynamicData.Color;
+		AltRecord.Destination = NewDynamicData.Destination;
+		AltRecord.Texture = NewDynamicData.Texture;
+		AltRecord.RenderVerts.Append(NewDynamicData.RenderVerts.GetTypedData() + SplitIndex, NewDynamicData.RenderVerts.Num() - SplitIndex);
+	}
+	else
+	{
+		FSpriteDrawCallRecord& Record = *new (BatchedSprites) FSpriteDrawCallRecord;
+		Record = NewDynamicData;
+	}
+}
+
+void FPaperSpriteSceneProxy::DrawDynamicElements_RichMesh(FPrimitiveDrawInterface* PDI, const FSceneView* View, bool bUseOverrideColor, const FLinearColor& OverrideColor)
+{
+	if (Material != nullptr)
+	{
+		DrawBatch(PDI, View, bUseOverrideColor, OverrideColor, Material, BatchedSprites);
+	}
+
+	if ((AlternateMaterial != nullptr) && (AlternateBatchedSprites.Num() > 0))
+	{
+		DrawBatch(PDI, View, bUseOverrideColor, OverrideColor, AlternateMaterial, AlternateBatchedSprites);
+	}
+}
diff --git a/Engine/Plugins/Experimental/Paper2D/Source/Paper2D/Private/PaperSpriteSceneProxy.h b/Engine/Plugins/Experimental/Paper2D/Source/Paper2D/Private/PaperSpriteSceneProxy.h
--- a/Engine/Plugins/Experimental/Paper2D/Source/Paper2D/Private/PaperSpriteSceneProxy.h
+++ b/Engine/Plugins/Experimental/Paper2D/Source/Paper2D/Private/PaperSpriteSceneProxy.h
@@ -16,6 +16,16 @@ public:
 	virtual void DrawDynamicElements(FPrimitiveDrawInterface* PDI, const FSceneView* View) override;
 	// End of FPrimitiveSceneProxy interface
 
+	void SetSprite_RenderThread(const FSpriteDrawCallRecord& NewDynamicData, int32 SplitIndex);
+
 protected:
+
+	// FPaperRenderSceneProxy interface
+	virtual void DrawDynamicElements_RichMesh(FPrimitiveDrawInterface* PDI, const FSceneView* View, bool bUseOverrideColor, const FLinearColor& OverrideColor) override;
+	// End of FPaperRenderSceneProxy interface
+
+	UMaterialInterface* AlternateMaterial;
+	int32 MaterialSplitIndex;
 	const UPaperSprite* SourceSprite;
+	TArray<FSpriteDrawCallRecord> AlternateBatchedSprites;
 };
diff --git a/Engine/Plugins/Experimental/Paper2D/Source/Paper2DEditor/Private/SpriteEditor/SpriteEditorViewportClient.cpp b/Engine/Plugins/Experimental/Paper2D/Source/Paper2DEditor/Private/SpriteEditor/SpriteEditorViewportClient.cpp
--- a/Engine/Plugins/Experimental/Paper2D/Source/Paper2DEditor/Private/SpriteEditor/SpriteEditorViewportClient.cpp
+++ b/Engine/Plugins/Experimental/Paper2D/Source/Paper2DEditor/Private/SpriteEditor/SpriteEditorViewportClient.cpp
@@ -319,7 +319,7 @@ void FSpriteEditorViewportClient::DrawGeometry(FViewport& InViewport, FSceneView
 void FSpriteEditorViewportClient::DrawGeometryStats(FViewport& InViewport, FSceneView& View, FCanvas& Canvas, const FSpritePolygonCollection& Geometry, bool bIsRenderGeometry, int32& YPos)
 {
 	// Draw the type of geometry we're displaying stats for
-	const FText GeometryName = bIsRenderGeometry ? LOCTEXT("RenderGeometry", "Render Geometry") : LOCTEXT("CollisionGeometry", "Collision Geometry");
+	const FText GeometryName = bIsRenderGeometry ? LOCTEXT("RenderGeometry", "Render Geometry (source)") : LOCTEXT("CollisionGeometry", "Collision Geometry (source)");
 
 	FCanvasTextItem TextItem(FVector2D(6, YPos), GeometryName, GEngine->GetSmallFont(), FLinearColor::White);
 	TextItem.EnableShadow(FLinearColor::Black);
@@ -343,8 +343,92 @@ void FSpriteEditorViewportClient::DrawGeometryStats(FViewport& InViewport, FScen
 	TextItem.Text = FText::Format(LOCTEXT("VerticesCount", "Verts: {0}"), FText::AsNumber(NumVerts));
 	TextItem.Draw(&Canvas);
 	TextItem.Position.Y += 18.0f;
-	
-	
+
+	YPos = (int32)TextItem.Position.Y;
+}
+
+void AttributeTrianglesByMaterialType(int32 NumTriangles, UMaterialInterface* Material, int32& NumOpaqueTriangles, int32& NumMaskedTriangles, int32& NumTranslucentTriangles)
+{
+	if (Material != nullptr)
+	{
+		switch (Material->GetBlendMode())
+		{
+		case EBlendMode::BLEND_Opaque:
+			NumOpaqueTriangles += NumTriangles;
+			break;
+		case EBlendMode::BLEND_Translucent:
+		case EBlendMode::BLEND_Additive:
+		case EBlendMode::BLEND_Modulate:
+			NumTranslucentTriangles += NumTriangles;
+			break;
+		case EBlendMode::BLEND_Masked:
+			NumMaskedTriangles += NumTriangles;
+			break;
+		}
+	}
+}
+
+void FSpriteEditorViewportClient::DrawRenderStats(FViewport& InViewport, FSceneView& View, FCanvas& Canvas, class UPaperSprite* Sprite, int32& YPos)
+{
+	FCanvasTextItem TextItem(FVector2D(6, YPos), LOCTEXT("RenderGeomBaked", "Render Geometry (baked)"), GEngine->GetSmallFont(), FLinearColor::White);
+	TextItem.EnableShadow(FLinearColor::Black);
+
+	TextItem.Draw(&Canvas);
+	TextItem.Position += FVector2D(6.0f, 18.0f);
+
+	int32 NumSections = (Sprite->AlternateMaterialSplitIndex != INDEX_NONE) ? 2 : 1;
+	if (NumSections > 1)
+	{
+		TextItem.Text = FText::Format(LOCTEXT("SectionCount", "Sections: {0}"), FText::AsNumber(NumSections));
+		TextItem.Draw(&Canvas);
+		TextItem.Position.Y += 18.0f;
+	}
+
+	int32 NumOpaqueTriangles = 0;
+	int32 NumMaskedTriangles = 0;
+	int32 NumTranslucentTriangles = 0;
+
+	int32 NumVerts = Sprite->BakedRenderData.Num();
+	int32 DefaultTriangles = 0;
+	int32 AlternateTriangles = 0;
+	if (Sprite->AlternateMaterialSplitIndex != INDEX_NONE)
+	{
+		DefaultTriangles = Sprite->AlternateMaterialSplitIndex / 3;
+		AlternateTriangles = (NumVerts - Sprite->AlternateMaterialSplitIndex) / 3;
+	}
+	else
+	{
+		DefaultTriangles = NumVerts / 3;
+	}
+
+	AttributeTrianglesByMaterialType(DefaultTriangles, Sprite->GetDefaultMaterial(), /*inout*/ NumOpaqueTriangles, /*inout*/ NumMaskedTriangles, /*inout*/ NumTranslucentTriangles);
+	AttributeTrianglesByMaterialType(AlternateTriangles, Sprite->GetAlternateMaterial(), /*inout*/ NumOpaqueTriangles, /*inout*/ NumMaskedTriangles, /*inout*/ NumTranslucentTriangles);
+
+	// Draw the number of polygons
+	if (NumOpaqueTriangles > 0)
+	{
+		TextItem.Text = FText::Format(LOCTEXT("OpaqueTriangleCount", "Triangles: {0} (opaque)"), FText::AsNumber(NumOpaqueTriangles));
+		TextItem.Draw(&Canvas);
+		TextItem.Position.Y += 18.0f;
+	}
+
+	if (NumMaskedTriangles > 0)
+	{
+		TextItem.Text = FText::Format(LOCTEXT("MaskedTriangleCount", "Triangles: {0} (masked)"), FText::AsNumber(NumMaskedTriangles));
+		TextItem.Draw(&Canvas);
+		TextItem.Position.Y += 18.0f;
+	}
+
+	if (NumTranslucentTriangles > 0)
+	{
+		TextItem.Text = FText::Format(LOCTEXT("TranslucentTriangleCount", "Triangles: {0} (translucent)"), FText::AsNumber(NumTranslucentTriangles));
+		TextItem.Draw(&Canvas);
+		TextItem.Position.Y += 18.0f;
+	}
+
+	TextItem.Text = FText::Format(LOCTEXT("BakedVertexCount", "Verts: {0}"), FText::AsNumber(Sprite->BakedRenderData.Num()));
+	TextItem.Draw(&Canvas);
+	TextItem.Position.Y += 18.0f;
 
 	YPos = (int32)TextItem.Position.Y;
 }
@@ -485,6 +569,7 @@ void FSpriteEditorViewportClient::DrawCanvas(FViewport& Viewport, FSceneView& Vi
 		// As well as all of the geometry
 		DrawGeometryStats(Viewport, View, Canvas, Sprite->CollisionGeometry, false, /*inout*/ YPos);
 		DrawGeometryStats(Viewport, View, Canvas, Sprite->RenderGeometry, true, /*inout*/ YPos);
+		DrawRenderStats(Viewport, View, Canvas, Sprite, /*inout*/ YPos);
 
 		// And bounds
 		DrawBoundsAsText(Viewport, View, Canvas, /*inout*/ YPos);
@@ -520,6 +605,7 @@ void FSpriteEditorViewportClient::DrawCanvas(FViewport& Viewport, FSceneView& Vi
 			const FLinearColor RenderGeomColor(1.0f, 0.2f, 0.0f, 1.0f);
 			DrawGeometry(Viewport, View, Canvas, Sprite->RenderGeometry, RenderGeomColor, true);
 			DrawGeometryStats(Viewport, View, Canvas, Sprite->RenderGeometry, true, /*inout*/ YPos);
+			DrawRenderStats(Viewport, View, Canvas, Sprite, /*inout*/ YPos);
 
 			// And bounds
 			DrawBoundsAsText(Viewport, View, Canvas, /*inout*/ YPos);
diff --git a/Engine/Plugins/Experimental/Paper2D/Source/Paper2DEditor/Private/SpriteEditor/SpriteEditorViewportClient.h b/Engine/Plugins/Experimental/Paper2D/Source/Paper2DEditor/Private/SpriteEditor/SpriteEditorViewportClient.h
--- a/Engine/Plugins/Experimental/Paper2D/Source/Paper2DEditor/Private/SpriteEditor/SpriteEditorViewportClient.h
+++ b/Engine/Plugins/Experimental/Paper2D/Source/Paper2DEditor/Private/SpriteEditor/SpriteEditorViewportClient.h
@@ -168,6 +168,7 @@ private:
 	void DrawBoundsAsText(FViewport& InViewport, FSceneView& View, FCanvas& Canvas, int32& YPos);
 	void DrawGeometry(FViewport& InViewport, FSceneView& View, FCanvas& Canvas, const FSpritePolygonCollection& Geometry, const FLinearColor& GeometryVertexColor, bool bIsRenderGeometry);
 	void DrawGeometryStats(FViewport& InViewport, FSceneView& View, FCanvas& Canvas, const FSpritePolygonCollection& Geometry, bool bIsRenderGeometry, int32& YPos);
+	void DrawRenderStats(FViewport& InViewport, FSceneView& View, FCanvas& Canvas, class UPaperSprite* Sprite, int32& YPos);
 	void DrawSockets(const FSceneView* View, FPrimitiveDrawInterface* PDI);
 	void DrawSocketNames(FViewport& InViewport, FSceneView& View, FCanvas& Canvas);
 	void DrawSourceRegion(FViewport& InViewport, FSceneView& View, FCanvas& Canvas, const FLinearColor& GeometryVertexColor, bool bIsRenderGeometry);
