Paper2D: Clarified SetFlipbook behavior in a comment

diff --git a/Engine/Plugins/Experimental/Paper2D/Source/Paper2D/Classes/PaperFlipbookComponent.h b/Engine/Plugins/Experimental/Paper2D/Source/Paper2D/Classes/PaperFlipbookComponent.h
--- a/Engine/Plugins/Experimental/Paper2D/Source/Paper2D/Classes/PaperFlipbookComponent.h
+++ b/Engine/Plugins/Experimental/Paper2D/Source/Paper2D/Classes/PaperFlipbookComponent.h
@@ -48,7 +48,7 @@ protected:
 	FLinearColor SpriteColor;
 
 public:
-	/** Change the flipbook used by this instance. */
+	/** Change the flipbook used by this instance (will reset the play time to 0 if it is a new flipbook). */
 	UFUNCTION(BlueprintCallable, Category="Sprite")
 	virtual bool SetFlipbook(class UPaperFlipbook* NewFlipbook);
 
