Header cleanup: Level.h, Brush.h

diff --git a/Engine/Plugins/Experimental/Paper2D/Source/Paper2D/Private/PaperRenderSceneProxy.cpp b/Engine/Plugins/Experimental/Paper2D/Source/Paper2D/Private/PaperRenderSceneProxy.cpp
--- a/Engine/Plugins/Experimental/Paper2D/Source/Paper2D/Private/PaperRenderSceneProxy.cpp
+++ b/Engine/Plugins/Experimental/Paper2D/Source/Paper2D/Private/PaperRenderSceneProxy.cpp
@@ -6,6 +6,7 @@
 #include "Rendering/PaperBatchManager.h"
 #include "Rendering/PaperBatchSceneProxy.h"
 #include "PhysicsEngine/BodySetup2D.h"
+#include "LocalVertexFactory.h"
 
 //////////////////////////////////////////////////////////////////////////
 // FPaperSpriteVertex
diff --git a/Engine/Source/Runtime/Engine/Classes/Engine/Brush.h b/Engine/Source/Runtime/Engine/Classes/Engine/Brush.h
--- a/Engine/Source/Runtime/Engine/Classes/Engine/Brush.h
+++ b/Engine/Source/Runtime/Engine/Classes/Engine/Brush.h
@@ -6,7 +6,7 @@
 
 #pragma once
 #include "Model.h"
-#include "Engine/Level.h"
+#include "GameFramework/Actor.h"
 #include "Brush.generated.h"
 
 //-----------------------------------------------------------------------------
diff --git a/Engine/Source/Runtime/Engine/Classes/Engine/Level.h b/Engine/Source/Runtime/Engine/Classes/Engine/Level.h
--- a/Engine/Source/Runtime/Engine/Classes/Engine/Level.h
+++ b/Engine/Source/Runtime/Engine/Classes/Engine/Level.h
@@ -5,10 +5,8 @@
 //
 
 #pragma once
-#include "PhysicsEngine/BodyInstance.h"
 #include "Engine/LevelBase.h"
 #include "Engine/World.h"
-#include "Model.h"
 #include "Level.generated.h"
 
 class ALevelBounds;
diff --git a/Engine/Source/Runtime/Engine/Private/ArrowComponent.cpp b/Engine/Source/Runtime/Engine/Private/ArrowComponent.cpp
--- a/Engine/Source/Runtime/Engine/Private/ArrowComponent.cpp
+++ b/Engine/Source/Runtime/Engine/Private/ArrowComponent.cpp
@@ -2,6 +2,7 @@
 
 #include "EnginePrivate.h"
 #include "DynamicMeshBuilder.h"
+#include "LocalVertexFactory.h"
 
 #define DEFAULT_SCREEN_SIZE	(0.0025f)
 #define ARROW_SCALE			(80.0f)
diff --git a/Engine/Source/Runtime/Engine/Private/DynamicMeshBuilder.cpp b/Engine/Source/Runtime/Engine/Private/DynamicMeshBuilder.cpp
--- a/Engine/Source/Runtime/Engine/Private/DynamicMeshBuilder.cpp
+++ b/Engine/Source/Runtime/Engine/Private/DynamicMeshBuilder.cpp
@@ -7,6 +7,7 @@
 #include "EnginePrivate.h"
 #include "DynamicMeshBuilder.h"
 #include "ResourcePool.h"
+#include "LocalVertexFactory.h"
 
 class FGlobalDynamicMeshPoolPolicy
 {
diff --git a/Engine/Source/Runtime/Engine/Private/InstancedStaticMesh.cpp b/Engine/Source/Runtime/Engine/Private/InstancedStaticMesh.cpp
--- a/Engine/Source/Runtime/Engine/Private/InstancedStaticMesh.cpp
+++ b/Engine/Source/Runtime/Engine/Private/InstancedStaticMesh.cpp
@@ -34,6 +34,7 @@
 #include "SpeedTreeWind.h"
 #include "ComponentInstanceDataCache.h"
 #include "InstancedFoliage.h"
+#include "VertexFactory.h"
 
 #if WITH_PHYSX
 #include "PhysicsEngine/PhysXSupport.h"
diff --git a/Engine/Source/Runtime/Engine/Private/LocalVertexFactory.cpp b/Engine/Source/Runtime/Engine/Private/LocalVertexFactory.cpp
--- a/Engine/Source/Runtime/Engine/Private/LocalVertexFactory.cpp
+++ b/Engine/Source/Runtime/Engine/Private/LocalVertexFactory.cpp
@@ -7,6 +7,7 @@
 #include "EnginePrivate.h"
 #include "SpeedTreeWind.h"
 #include "ShaderParameterUtils.h"
+#include "LocalVertexFactory.h"
 
 class FSpeedTreeVertexFactoryShaderParameters : public FVertexFactoryShaderParameters
 {
diff --git a/Engine/Source/Runtime/Engine/Private/Materials/MaterialShared.cpp b/Engine/Source/Runtime/Engine/Private/Materials/MaterialShared.cpp
--- a/Engine/Source/Runtime/Engine/Private/Materials/MaterialShared.cpp
+++ b/Engine/Source/Runtime/Engine/Private/Materials/MaterialShared.cpp
@@ -17,6 +17,7 @@
 #include "Developer/TargetPlatform/Public/TargetPlatform.h"
 #include "ComponentReregisterContext.h"
 #include "EngineModule.h"
+#include "LocalVertexFactory.h"
 
 DEFINE_LOG_CATEGORY(LogMaterial);
 
diff --git a/Engine/Source/Runtime/Engine/Private/SkeletalRenderCPUSkin.h b/Engine/Source/Runtime/Engine/Private/SkeletalRenderCPUSkin.h
--- a/Engine/Source/Runtime/Engine/Private/SkeletalRenderCPUSkin.h
+++ b/Engine/Source/Runtime/Engine/Private/SkeletalRenderCPUSkin.h
@@ -9,6 +9,7 @@
 #include "SkeletalRender.h"
 #include "SkeletalRenderPublic.h"
 #include "GPUSkinVertexFactory.h"
+#include "LocalVertexFactory.h"
 
 /** data for a single skinned skeletal mesh vertex */
 struct FFinalSkinVertex
diff --git a/Engine/Source/Runtime/Engine/Private/TextRenderComponent.cpp b/Engine/Source/Runtime/Engine/Private/TextRenderComponent.cpp
--- a/Engine/Source/Runtime/Engine/Private/TextRenderComponent.cpp
+++ b/Engine/Source/Runtime/Engine/Private/TextRenderComponent.cpp
@@ -3,6 +3,7 @@
 #include "EnginePrivate.h"
 #include "DynamicMeshBuilder.h"
 #include "../../../../Source/Runtime/Engine/Classes/Engine/TextRenderActor.h"
+#include "LocalVertexFactory.h"
 
 ATextRenderActor::ATextRenderActor(const class FPostConstructInitializeProperties& PCIP)
 	: Super(PCIP)
diff --git a/Engine/Source/Runtime/Engine/Private/TileRendering.cpp b/Engine/Source/Runtime/Engine/Private/TileRendering.cpp
--- a/Engine/Source/Runtime/Engine/Private/TileRendering.cpp
+++ b/Engine/Source/Runtime/Engine/Private/TileRendering.cpp
@@ -7,6 +7,7 @@
 #include "EnginePrivate.h"
 #include "TileRendering.h"
 #include "EngineModule.h"
+#include "LocalVertexFactory.h"
 
 /** 
 * vertex data for a screen quad 
diff --git a/Engine/Source/Runtime/Engine/Public/PhysicsPublic.h b/Engine/Source/Runtime/Engine/Public/PhysicsPublic.h
--- a/Engine/Source/Runtime/Engine/Public/PhysicsPublic.h
+++ b/Engine/Source/Runtime/Engine/Public/PhysicsPublic.h
@@ -8,7 +8,7 @@
 #pragma once
 
 #include "PhysxUserData.h"
-
+#include "LocalVertexFactory.h"
 /**
  * Physics stats
  */
diff --git a/Engine/Source/Runtime/Engine/Public/StaticMeshResources.h b/Engine/Source/Runtime/Engine/Public/StaticMeshResources.h
--- a/Engine/Source/Runtime/Engine/Public/StaticMeshResources.h
+++ b/Engine/Source/Runtime/Engine/Public/StaticMeshResources.h
@@ -8,6 +8,7 @@
 
 #include "RawIndexBuffer.h"
 #include "TextureLayout3d.h"
+#include "LocalVertexFactory.h"
 
 /**
  * The LOD settings to use for a group of static meshes.
