Paper2D: Added support for both variants of JSON import (hashed & array) from TexturePacker

diff --git a/Engine/Plugins/Experimental/Paper2D/Source/Paper2DEditor/Private/PaperJSONHelpers.cpp b/Engine/Plugins/Experimental/Paper2D/Source/Paper2DEditor/Private/PaperJSONHelpers.cpp
--- a/Engine/Plugins/Experimental/Paper2D/Source/Paper2DEditor/Private/PaperJSONHelpers.cpp
+++ b/Engine/Plugins/Experimental/Paper2D/Source/Paper2DEditor/Private/PaperJSONHelpers.cpp
@@ -7,6 +7,8 @@
 //////////////////////////////////////////////////////////////////////////
 // FPaperJSONHelpers
 
+static TArray< TSharedPtr<FJsonValue> > EmptyArray;
+
 FString FPaperJSONHelpers::ReadString(TSharedPtr<class FJsonObject> Item, const FString& Key, const FString& DefaultValue)
 {
 	const TSharedPtr<FJsonValue>& ValuePtr = Item->GetField<EJson::String>(Key);
@@ -25,6 +27,18 @@ TSharedPtr<class FJsonObject> FPaperJSONHelpers::ReadObject(TSharedPtr<class FJs
 	}
 }
 
+const TArray< TSharedPtr<FJsonValue> >& FPaperJSONHelpers::ReadArray(TSharedPtr<class FJsonObject> Item, const FString& Key)
+{
+	if (Item->HasTypedField<EJson::Array>(Key))
+	{
+		return Item->GetArrayField(Key);
+	}
+	else
+	{
+		return EmptyArray;
+	}
+}
+
 bool FPaperJSONHelpers::ReadBoolean(const TSharedPtr<class FJsonObject> Item, const FString& Key, bool bDefaultIfMissing)
 {
 	if (Item->HasTypedField<EJson::Boolean>(Key))
diff --git a/Engine/Plugins/Experimental/Paper2D/Source/Paper2DEditor/Public/PaperJSONHelpers.h b/Engine/Plugins/Experimental/Paper2D/Source/Paper2DEditor/Public/PaperJSONHelpers.h
--- a/Engine/Plugins/Experimental/Paper2D/Source/Paper2DEditor/Public/PaperJSONHelpers.h
+++ b/Engine/Plugins/Experimental/Paper2D/Source/Paper2DEditor/Public/PaperJSONHelpers.h
@@ -14,6 +14,9 @@ public:
 	// Returns the object named Key or NULL if it is missing or the wrong type
 	static TSharedPtr<class FJsonObject> ReadObject(TSharedPtr<class FJsonObject> Item, const FString& Key);
 
+	// Returns the array named Key or NULL if it is missing or the wrong type
+	static const TArray< TSharedPtr<FJsonValue> >& ReadArray(TSharedPtr<class FJsonObject> Item, const FString& Key);
+
 	// Returns the bool named Key or bDefaultIfMissing if it is missing or the wrong type (note: no way to determine errors!)
 	static bool ReadBoolean(const TSharedPtr<class FJsonObject> Item, const FString& Key, bool bDefaultIfMissing);
 
diff --git a/Engine/Plugins/Experimental/Paper2D/Source/PaperJsonImporter/Private/PaperJsonImporterFactory.cpp b/Engine/Plugins/Experimental/Paper2D/Source/PaperJsonImporter/Private/PaperJsonImporterFactory.cpp
--- a/Engine/Plugins/Experimental/Paper2D/Source/PaperJsonImporter/Private/PaperJsonImporterFactory.cpp
+++ b/Engine/Plugins/Experimental/Paper2D/Source/PaperJsonImporter/Private/PaperJsonImporterFactory.cpp
@@ -22,6 +22,8 @@ struct FSpriteFrame
 	FVector2D SpriteSourcePos;
 	FVector2D SpriteSourceSize;
 
+	FVector2D ImageSourceSize;
+
 	bool bTrimmed;
 	bool bRotated;
 };
@@ -47,6 +49,163 @@ FText UPaperJsonImporterFactory::GetToolTip() const
 	return NSLOCTEXT("Paper2D", "PaperJsonImporterFactoryDescription", "Sprite sheets exported from Adobe Flash");
 }
 
+static bool ParseFrame(TSharedPtr<FJsonObject> &FrameData, FSpriteFrame &OutFrame)
+{
+	bool bReadFrameSuccessfully = true;
+	// An example frame:
+	//   "frame": {"x":210,"y":10,"w":190,"h":223},
+	//   "rotated": false,
+	//   "trimmed": true,
+	//   "spriteSourceSize": {"x":0,"y":11,"w":216,"h":240},
+	//   "sourceSize": {"w":216,"h":240}
+
+	OutFrame.bRotated = FPaperJSONHelpers::ReadBoolean(FrameData, TEXT("rotated"), false);
+	OutFrame.bTrimmed = FPaperJSONHelpers::ReadBoolean(FrameData, TEXT("trimmed"), false);
+
+	// Put this warning back in until rotation / trim support is implemented
+	if (OutFrame.bRotated || !OutFrame.bTrimmed)
+	{
+		//@TODO: Should learn how to support this combo
+		UE_LOG(LogPaperJsonImporter, Warning, TEXT("Frame %s is either rotated or not trimmed, which may not be handled correctly"), *OutFrame.FrameName.ToString());
+	}
+
+	bReadFrameSuccessfully = bReadFrameSuccessfully && FPaperJSONHelpers::ReadRectangle(FrameData, "frame", /*out*/ OutFrame.SpritePosInSheet, /*out*/ OutFrame.SpriteSizeInSheet);
+
+	if (OutFrame.bTrimmed)
+	{
+		bReadFrameSuccessfully = bReadFrameSuccessfully && FPaperJSONHelpers::ReadSize(FrameData, "sourceSize", /*out*/ OutFrame.ImageSourceSize);
+		bReadFrameSuccessfully = bReadFrameSuccessfully && FPaperJSONHelpers::ReadRectangle(FrameData, "spriteSourceSize", /*out*/ OutFrame.SpriteSourcePos, /*out*/ OutFrame.SpriteSourceSize);
+	}
+	else
+	{
+		OutFrame.SpriteSourcePos = FVector2D::ZeroVector;
+		OutFrame.SpriteSourceSize = OutFrame.SpriteSizeInSheet;
+		OutFrame.ImageSourceSize = OutFrame.SpriteSizeInSheet;
+	}
+
+	// A few more prerequisites to sort out before rotation can be enabled
+	//if (OutFrame.bRotated)
+	//{
+	//	// We rotate the sprite 90 deg CCW
+	//	Swap(OutFrame.SpriteSizeInSheet.X, OutFrame.SpriteSizeInSheet.Y);
+	//}
+
+	return bReadFrameSuccessfully;
+}
+
+static bool ParseFramesFromSpriteHash(TSharedPtr<FJsonObject> ObjectBlock, TArray<FSpriteFrame>& OutSpriteFrames, TSet<FName>& FrameNames)
+{
+	GWarn->BeginSlowTask(NSLOCTEXT("Paper2D", "PaperJsonImporterFactory_ParsingSprites", "Parsing Sprite Frame"), true, true);
+	bool bLoadedSuccessfully = true;
+
+	// Parse all of the frames
+	int32 FrameCount = 0;
+	for (auto FrameIt = ObjectBlock->Values.CreateIterator(); FrameIt; ++FrameIt)
+	{
+		GWarn->StatusUpdate(FrameCount, ObjectBlock->Values.Num(), NSLOCTEXT("Paper2D", "PaperJsonImporterFactory_ParsingSprites", "Parsing Sprite Frames"));
+
+		bool bReadFrameSuccessfully = true;
+
+		FSpriteFrame Frame;
+		Frame.FrameName = *FrameIt.Key();
+
+		if (FrameNames.Contains(Frame.FrameName))
+		{
+			bReadFrameSuccessfully = false;
+		}
+		else
+		{
+			FrameNames.Add(Frame.FrameName);
+		}
+
+		TSharedPtr<FJsonValue> FrameDataAsValue = FrameIt.Value();
+		TSharedPtr<FJsonObject> FrameData;
+		if (FrameDataAsValue->Type == EJson::Object)
+		{
+			FrameData = FrameDataAsValue->AsObject();
+			bReadFrameSuccessfully = bReadFrameSuccessfully && ParseFrame(FrameData, /*out*/Frame);
+		}
+		else
+		{
+			bReadFrameSuccessfully = false;
+		}
+
+		if (bReadFrameSuccessfully)
+		{
+			OutSpriteFrames.Add(Frame);
+		}
+		else
+		{
+			UE_LOG(LogPaperJsonImporter, Warning, TEXT("Frame %s is in an unexpected format"), *Frame.FrameName.ToString());
+			bLoadedSuccessfully = false;
+		}
+
+		FrameCount++;
+	}
+
+	GWarn->EndSlowTask();
+	return bLoadedSuccessfully;
+}
+
+static bool ParseFramesFromSpriteArray(const TArray<TSharedPtr<FJsonValue>>& ArrayBlock, TArray<FSpriteFrame>& OutSpriteFrames, TSet<FName>& FrameNames)
+{
+	GWarn->BeginSlowTask(NSLOCTEXT("Paper2D", "PaperJsonImporterFactory_ParsingSprites", "Parsing Sprite Frame"), true, true);
+	bool bLoadedSuccessfully = true;
+
+	// Parse all of the frames
+	int32 FrameCount = 0;
+	for (int32 FrameCount = 0; FrameCount < ArrayBlock.Num(); ++FrameCount)
+	{
+		GWarn->StatusUpdate(FrameCount, ArrayBlock.Num(), NSLOCTEXT("Paper2D", "PaperJsonImporterFactory_ParsingSprites", "Parsing Sprite Frames"));
+		bool bReadFrameSuccessfully = true;
+		FSpriteFrame Frame;
+
+		const TSharedPtr<FJsonValue>& FrameDataAsValue = ArrayBlock[FrameCount];
+		if (FrameDataAsValue->Type == EJson::Object)
+		{
+			TSharedPtr<FJsonObject> FrameData;
+			FrameData = FrameDataAsValue->AsObject();
+
+			FString FrameFilename = FPaperJSONHelpers::ReadString(FrameData, TEXT("filename"), TEXT(""));
+			if (!FrameFilename.IsEmpty())
+			{
+				Frame.FrameName = FName(*FrameFilename); // case insensitive
+				if (FrameNames.Contains(Frame.FrameName))
+				{
+					bReadFrameSuccessfully = false;
+				}
+				else
+				{
+					FrameNames.Add(Frame.FrameName);
+				}
+
+				bReadFrameSuccessfully = bReadFrameSuccessfully && ParseFrame(FrameData, /*out*/Frame);
+			}
+			else
+			{
+				bReadFrameSuccessfully = false;
+			}
+		}
+		else
+		{
+			bReadFrameSuccessfully = false;
+		}
+
+		if (bReadFrameSuccessfully)
+		{
+			OutSpriteFrames.Add(Frame);
+		}
+		else
+		{
+			UE_LOG(LogPaperJsonImporter, Warning, TEXT("Frame %s is in an unexpected format"), *Frame.FrameName.ToString());
+			bLoadedSuccessfully = false;
+		}
+	}
+
+	GWarn->EndSlowTask();
+	return bLoadedSuccessfully;
+}
+
 UObject* UPaperJsonImporterFactory::FactoryCreateText(UClass* InClass, UObject* InParent, FName InName, EObjectFlags Flags, UObject* Context, const TCHAR* Type, const TCHAR*& Buffer, const TCHAR* BufferEnd, FFeedbackContext* Warn)
 {
 	Flags |= RF_Transactional;
@@ -164,151 +323,89 @@ UObject* UPaperJsonImporterFactory::FactoryCreateText(UClass* InClass, UObject*
 		// Parse the frames
 		if (bLoadedSuccessfully)
 		{
-			TSharedPtr<FJsonObject> FramesBlock = FPaperJSONHelpers::ReadObject(SpriteDescriptorObject, TEXT("frames"));
-			if (FramesBlock.IsValid())
+			TSharedPtr<FJsonObject> ObjectFrameBlock = FPaperJSONHelpers::ReadObject(SpriteDescriptorObject, TEXT("frames"));
+			if (ObjectFrameBlock.IsValid())
 			{
-				GWarn->BeginSlowTask(NSLOCTEXT("Paper2D", "PaperJsonImporterFactory_ParsingSprites", "Parsing Sprite Frame"), true, true);
-
-				// Parse all of the frames
-				int32 FrameCount = 0;
-				for (auto FrameIt = FramesBlock->Values.CreateIterator(); FrameIt; ++FrameIt)
+				bLoadedSuccessfully = ParseFramesFromSpriteHash(ObjectFrameBlock, /*out*/ParsedFrames, /*inout*/FrameNames);
+			}
+			else
+			{
+				// Try loading as an array
+				TArray<TSharedPtr<FJsonValue>> ArrayBlock = FPaperJSONHelpers::ReadArray(SpriteDescriptorObject, TEXT("frames"));
+				if (ArrayBlock.Num() > 0)
 				{
-					GWarn->StatusUpdate(FrameCount, FramesBlock->Values.Num(), NSLOCTEXT("Paper2D", "PaperJsonImporterFactory_ParsingSprites", "Parsing Sprite Frames"));
-
-					bool bReadFrameSuccessfully = true;
-
-					FSpriteFrame Frame;
-					Frame.FrameName = *FrameIt.Key();
-
-					if (FrameNames.Contains(Frame.FrameName))
-					{
-						bReadFrameSuccessfully = false;
-					}
-					else
-					{
-						FrameNames.Add(Frame.FrameName);
-					}
-
-					TSharedPtr<FJsonValue> FrameDataAsValue = FrameIt.Value();
-					TSharedPtr<FJsonObject> FrameData;
-					if (FrameDataAsValue->Type == EJson::Object)
-					{
-						FrameData = FrameDataAsValue->AsObject();
-
-						// An example frame:
-						//   "frame": {"x":210,"y":10,"w":190,"h":223},
-						//   "rotated": false,
-						//   "trimmed": true,
-						//   "spriteSourceSize": {"x":0,"y":11,"w":216,"h":240},
-						//   "sourceSize": {"w":216,"h":240}
-
-						Frame.bRotated = FPaperJSONHelpers::ReadBoolean(FrameData, TEXT("rotated"), false);
-						Frame.bTrimmed = FPaperJSONHelpers::ReadBoolean(FrameData, TEXT("trimmed"), false);
-
-						if (Frame.bRotated || !Frame.bTrimmed)
-						{
-							//@TODO: Should learn how to support this combo
-							UE_LOG(LogPaperJsonImporter, Warning, TEXT("Frame %s is either rotated or not trimmed, which may not be handled correctly"), *Frame.FrameName.ToString());
-						}
-
-						bReadFrameSuccessfully = bReadFrameSuccessfully && FPaperJSONHelpers::ReadRectangle(FrameData, "frame", /*out*/ Frame.SpritePosInSheet, /*out*/ Frame.SpriteSizeInSheet);
-
-						if (Frame.bTrimmed)
-						{
-							bReadFrameSuccessfully = bReadFrameSuccessfully && FPaperJSONHelpers::ReadRectangle(FrameData, "spriteSourceSize", /*out*/ Frame.SpriteSourcePos, /*out*/ Frame.SpriteSourceSize);
-						}
-						else
-						{
-							Frame.SpriteSourcePos = FVector2D::ZeroVector;
-							Frame.SpriteSourceSize = Frame.SpriteSizeInSheet;
-						}
-					}
-					else
-					{
-						bReadFrameSuccessfully = false;
-					}
+					bLoadedSuccessfully = ParseFramesFromSpriteArray(ArrayBlock, /*out*/ParsedFrames, /*inout*/FrameNames);
+				}
+				else
+				{
+					UE_LOG(LogPaperJsonImporter, Warning, TEXT("Failed to parse sprite descriptor file '%s'.  Missing frames block"), *NameForErrors);
+				}
+			}
 
-					if (bReadFrameSuccessfully)
-					{
-						ParsedFrames.Add(Frame);
-					}
-					else
-					{
-						UE_LOG(LogPaperJsonImporter, Warning, TEXT("Frame %s is in an unexpected format"), *Frame.FrameName.ToString());
-						bLoadedSuccessfully = false;
-					}
 
-					FrameCount++;
-				}
+			// Create an animated sprite that encompasses the frames
+			UPaperFlipbook* Flipbook = NULL;
+			if (ParsedFrames.Num() > 0)
+			{
+				Flipbook = NewNamedObject<UPaperFlipbook>(InParent, InName, Flags);
+				Result = Flipbook;
 
-				GWarn->EndSlowTask();
+				GWarn->BeginSlowTask(NSLOCTEXT("Paper2D", "PaperJsonImporterFactory_ImportingSprites", "Importing Sprite Frame"), true, true);
+				FScopedFlipbookMutator EditLock(Flipbook);
 
-				// Create an animated sprite that encompasses the frames
-				UPaperFlipbook* Flipbook = NULL;
-				if (ParsedFrames.Num() > 0)
+				// Create objects for each successfully parsed frame
+				const int32 FrameRun = 1; //@TODO: Don't make a keyframe for every single item if we can help it
+				for (int32 FrameIndex = 0; FrameIndex < ParsedFrames.Num(); ++FrameIndex)
 				{
-					Flipbook = NewNamedObject<UPaperFlipbook>(InParent, InName, Flags);
-					Result = Flipbook;
+					GWarn->StatusUpdate(FrameIndex, ParsedFrames.Num(), NSLOCTEXT("Paper2D", "PaperJsonImporterFactory_ImportingSprites", "Importing Sprite Frames"));
 
-					GWarn->BeginSlowTask(NSLOCTEXT("Paper2D", "PaperJsonImporterFactory_ImportingSprites", "Importing Sprite Frame"), true, true);
-					FScopedFlipbookMutator EditLock(Flipbook);
-
-					// Create objects for each successfully parsed frame
-					const int32 FrameRun = 1; //@TODO: Don't make a keyframe for every single item if we can help it
-					for (int32 FrameIndex = 0; FrameIndex < ParsedFrames.Num(); ++FrameIndex)
+					// Check for user canceling the import
+					if (GWarn->ReceivedUserCancel())
 					{
-						GWarn->StatusUpdate(FrameIndex, ParsedFrames.Num(), NSLOCTEXT("Paper2D", "PaperJsonImporterFactory_ImportingSprites", "Importing Sprite Frames"));
-
-						// Check for user canceling the import
-						if (GWarn->ReceivedUserCancel())
-						{
-							break;
-						}
+						break;
+					}
 
-						const FSpriteFrame& Frame = ParsedFrames[FrameIndex];
+					const FSpriteFrame& Frame = ParsedFrames[FrameIndex];
 
-						// Create a package for the frame
-						const FString TargetSubPath = LongPackagePath + TEXT("/Frames");
+					// Create a package for the frame
+					const FString TargetSubPath = LongPackagePath + TEXT("/Frames");
 
-						UObject* OuterForFrame = NULL; // @TODO: Use this if we don't want them to be individual assets - Flipbook;
+					UObject* OuterForFrame = NULL; // @TODO: Use this if we don't want them to be individual assets - Flipbook;
 
-						// Create a unique package name and asset name for the frame
-						const FString TentativePackagePath = PackageTools::SanitizePackageName(TargetSubPath + TEXT("/") + Frame.FrameName.ToString());
-						FString DefaultSuffix;
-						FString AssetName;
-						FString PackageName;
-						AssetToolsModule.Get().CreateUniqueAssetName(TentativePackagePath, /*out*/ DefaultSuffix, /*out*/ PackageName, /*out*/ AssetName);
+					// Create a unique package name and asset name for the frame
+					const FString TentativePackagePath = PackageTools::SanitizePackageName(TargetSubPath + TEXT("/") + Frame.FrameName.ToString());
+					FString DefaultSuffix;
+					FString AssetName;
+					FString PackageName;
+					AssetToolsModule.Get().CreateUniqueAssetName(TentativePackagePath, /*out*/ DefaultSuffix, /*out*/ PackageName, /*out*/ AssetName);
 
-						// Create a package for the frame
-						OuterForFrame = CreatePackage(NULL, *PackageName);
+					// Create a package for the frame
+					OuterForFrame = CreatePackage(NULL, *PackageName);
 
-						// Create a frame in the package
-						UPaperSprite* TargetSprite = NewNamedObject<UPaperSprite>(OuterForFrame, *AssetName, Flags);
-						FAssetRegistryModule::AssetCreated(TargetSprite);
+					// Create a frame in the package
+					UPaperSprite* TargetSprite = NewNamedObject<UPaperSprite>(OuterForFrame, *AssetName, Flags);
+					FAssetRegistryModule::AssetCreated(TargetSprite);
 
-						TargetSprite->Modify();
+					TargetSprite->Modify();
 
-						TargetSprite->InitializeSprite(ImportedTexture, Frame.SpritePosInSheet, Frame.SpriteSizeInSheet);
+					// TargetSprite->InitializeSprite(ImportedTexture, Frame.SpritePosInSheet, Frame.SpriteSizeInSheet, GetDefault<UPaperRuntimeSettings>()->DefaultPixelsPerUnrealUnit, Frame.bRotated);
+					TargetSprite->InitializeSprite(ImportedTexture, Frame.SpritePosInSheet, Frame.SpriteSizeInSheet, GetDefault<UPaperRuntimeSettings>()->DefaultPixelsPerUnrealUnit);
 
-						//@TODO: Need to support pivot behavior - Total guess at pivot behavior
-						//const FVector2D SizeDifference = Frame.SpriteSourceSize - Frame.SpriteSizeInSheet;
-						//TargetSprite->SpriteData.Destination = FVector(SizeDifference.X * -0.5f, 0.0f, SizeDifference.Y * -0.5f);
+					//@TODO: Need to support pivot behavior - Total guess at pivot behavior
+					//const FVector2D SizeDifference = Frame.SpriteSourceSize - Frame.SpriteSizeInSheet;
+					//TargetSprite->SpriteData.Destination = FVector(SizeDifference.X * -0.5f, 0.0f, SizeDifference.Y * -0.5f);
 
-						// Create the entry in the animation
-						FPaperFlipbookKeyFrame* DestFrame = new (EditLock.KeyFrames) FPaperFlipbookKeyFrame();
-						DestFrame->Sprite = TargetSprite;
-						DestFrame->FrameRun = FrameRun;
+					// Create the entry in the animation
+					FPaperFlipbookKeyFrame* DestFrame = new (EditLock.KeyFrames) FPaperFlipbookKeyFrame();
+					DestFrame->Sprite = TargetSprite;
+					DestFrame->FrameRun = FrameRun;
 
-						TargetSprite->PostEditChange();
-					}
-
-					GWarn->EndSlowTask();
+					TargetSprite->PostEditChange();
 				}
+
+				GWarn->EndSlowTask();
 			}
-			else
-			{
-				UE_LOG(LogPaperJsonImporter, Warning, TEXT( "Failed to parse sprite descriptor file '%s'.  Missing frames block" ), *NameForErrors);
-			}
+
 		}
 	}
 	else
