Paper2D: Add support for casting 3D shadows from sprites (FPaperRenderSceneProxy now obeys CastShadow from the source component)

diff --git a/Engine/Plugins/Experimental/Paper2D/Source/Paper2D/Private/PaperRenderSceneProxy.cpp b/Engine/Plugins/Experimental/Paper2D/Source/Paper2D/Private/PaperRenderSceneProxy.cpp
--- a/Engine/Plugins/Experimental/Paper2D/Source/Paper2D/Private/PaperRenderSceneProxy.cpp
+++ b/Engine/Plugins/Experimental/Paper2D/Source/Paper2D/Private/PaperRenderSceneProxy.cpp
@@ -159,6 +159,7 @@ FPaperRenderSceneProxy::FPaperRenderSceneProxy(const UPrimitiveComponent* InComp
 	: FPrimitiveSceneProxy(InComponent)
 	, Material(NULL)
 	, Owner(InComponent->GetOwner())
+	, bCastShadow(InComponent->CastShadow)
 	, WireframeColor(FLinearColor::White)
 	, CollisionResponse(InComponent->GetCollisionResponseToChannels())
 {
@@ -263,7 +264,7 @@ void FPaperRenderSceneProxy::DrawNewBatches(FPrimitiveDrawInterface* PDI, const
 	Mesh.VertexFactory = VertexFactory;
 	Mesh.LCI = NULL;
 	Mesh.ReverseCulling = IsLocalToWorldDeterminantNegative() ? true : false;
-	Mesh.CastShadow = false;
+	Mesh.CastShadow = bCastShadow;
 	Mesh.DepthPriorityGroup = DPG;
 	Mesh.Type = PT_TriangleList;
 	Mesh.bDisableBackfaceCulling = true;
@@ -354,7 +355,7 @@ void FPaperRenderSceneProxy::DrawBatch(FPrimitiveDrawInterface* PDI, const FScen
 			Mesh.MaterialRenderProxy = TextureOverrideMaterialProxy;
 			Mesh.LCI = NULL;
 			Mesh.ReverseCulling = IsLocalToWorldDeterminantNegative() ? true : false;
-			Mesh.CastShadow = false;
+			Mesh.CastShadow = bCastShadow;
 			Mesh.DepthPriorityGroup = DPG;
 			Mesh.Type = PT_TriangleList;
 			Mesh.bDisableBackfaceCulling = true;
diff --git a/Engine/Plugins/Experimental/Paper2D/Source/Paper2D/Private/PaperRenderSceneProxy.h b/Engine/Plugins/Experimental/Paper2D/Source/Paper2D/Private/PaperRenderSceneProxy.h
--- a/Engine/Plugins/Experimental/Paper2D/Source/Paper2D/Private/PaperRenderSceneProxy.h
+++ b/Engine/Plugins/Experimental/Paper2D/Source/Paper2D/Private/PaperRenderSceneProxy.h
@@ -125,6 +125,8 @@ protected:
 	FVector Origin;
 	AActor* Owner;
 
+	bool bCastShadow;
+
 	// The color to draw as in wireframe mode
 	FLinearColor WireframeColor;
 
