Paper2D: Add snap to pixel grid option for pivots (enabled by default on new content), and enforced snap to pixel grid for the source rectangle (always)

diff --git a/Engine/Plugins/Experimental/Paper2D/Source/Paper2D/Classes/PaperSprite.h b/Engine/Plugins/Experimental/Paper2D/Source/Paper2D/Classes/PaperSprite.h
--- a/Engine/Plugins/Experimental/Paper2D/Source/Paper2D/Classes/PaperSprite.h
+++ b/Engine/Plugins/Experimental/Paper2D/Source/Paper2D/Classes/PaperSprite.h
@@ -103,6 +103,10 @@ protected:
 	UPROPERTY(Category=Sprite, EditAnywhere)
 	FVector2D CustomPivotPoint;
 
+	// Should the pivot be snapped to a pixel boundary?
+	UPROPERTY(Category=Sprite, EditAnywhere, AdvancedDisplay)
+	bool bSnapPivotToPixelGrid;
+
 	// Custom collision geometry polygons (in texture space)
 	UPROPERTY(Category=Collision, EditAnywhere)
 	FSpritePolygonCollection CollisionGeometry;
@@ -118,6 +122,7 @@ protected:
 	// Spritesheet group that this sprite belongs to
 	UPROPERTY(Category=Sprite, EditAnywhere, AssetRegistrySearchable)
 	class UPaperSpriteAtlas* AtlasGroup;
+
 #endif
 
 public:
@@ -153,6 +158,9 @@ public:
 	//FVector ConvertTextureSpaceToWorldSpace(const FVector& SourcePoint) const;
 	FTransform GetPivotToWorld() const;
 
+	// Returns the raw pivot position (ignoring pixel snapping)
+	FVector2D GetRawPivotPosition() const;
+
 	// Returns the current pivot position in texture space
 	FVector2D GetPivotPosition() const;
 
diff --git a/Engine/Plugins/Experimental/Paper2D/Source/Paper2D/Private/PaperCustomVersion.h b/Engine/Plugins/Experimental/Paper2D/Source/Paper2D/Private/PaperCustomVersion.h
--- a/Engine/Plugins/Experimental/Paper2D/Source/Paper2D/Private/PaperCustomVersion.h
+++ b/Engine/Plugins/Experimental/Paper2D/Source/Paper2D/Private/PaperCustomVersion.h
@@ -34,6 +34,9 @@ struct FPaperCustomVersion
 		// Converted UPaperSpriteComponent to be derived from UMeshComponent
 		ConvertPaperSpriteComponentToBeMeshComponent = 8,
 
+		// Added the option to snap pivots to a pixel grid
+		AddPivotSnapToPixelGrid = 9,
+
 		// -----<new versions can be added above this line>-------------------------------------------------
 		VersionPlusOne,
 		LatestVersion = VersionPlusOne - 1
diff --git a/Engine/Plugins/Experimental/Paper2D/Source/Paper2D/Private/PaperSprite.cpp b/Engine/Plugins/Experimental/Paper2D/Source/Paper2D/Private/PaperSprite.cpp
--- a/Engine/Plugins/Experimental/Paper2D/Source/Paper2D/Private/PaperSprite.cpp
+++ b/Engine/Plugins/Experimental/Paper2D/Source/Paper2D/Private/PaperSprite.cpp
@@ -259,6 +259,7 @@ UPaperSprite::UPaperSprite(const FPostConstructInitializeProperties& PCIP)
 
 #if WITH_EDITORONLY_DATA
 	PivotMode = ESpritePivotMode::Center_Center;
+	bSnapPivotToPixelGrid = true;
 
 	CollisionGeometry.GeometryType = ESpritePolygonMode::TightBoundingBox;
 	CollisionThickness = 10.0f;
@@ -276,6 +277,14 @@ UPaperSprite::UPaperSprite(const FPostConstructInitializeProperties& PCIP)
 #if WITH_EDITORONLY_DATA
 void UPaperSprite::PostEditChangeProperty(FPropertyChangedEvent& PropertyChangedEvent)
 {
+	//@TODO: Determine when these are really needed, as they're seriously expensive!
+	TComponentReregisterContext<UPaperSpriteComponent> ReregisterStaticComponents;
+	TComponentReregisterContext<UPaperFlipbookComponent> ReregisterAnimatedComponents;
+
+	// Look for changed properties
+	const FName PropertyName = (PropertyChangedEvent.Property != NULL) ? PropertyChangedEvent.Property->GetFName() : NAME_None;
+	const FName MemberPropertyName = (PropertyChangedEvent.MemberProperty != NULL) ? PropertyChangedEvent.MemberProperty->GetFName() : NAME_None;
+
 	if (PixelsPerUnrealUnit <= 0.0f)
 	{
 		PixelsPerUnrealUnit = 1.0f;
@@ -289,22 +298,19 @@ void UPaperSprite::PostEditChangeProperty(FPropertyChangedEvent& PropertyChanged
 	RenderGeometry.PixelsPerSubdivisionX = FMath::Max(RenderGeometry.PixelsPerSubdivisionX, 4);
 	RenderGeometry.PixelsPerSubdivisionY = FMath::Max(RenderGeometry.PixelsPerSubdivisionY, 4);
 
-	SourceDimension.X = FMath::Max(SourceDimension.X, 0.0f);
-	SourceDimension.Y = FMath::Max(SourceDimension.Y, 0.0f);
-
-	//@TODO: Determine when these are really needed, as they're seriously expensive!
-	TComponentReregisterContext<UPaperSpriteComponent> ReregisterStaticComponents;
-	TComponentReregisterContext<UPaperFlipbookComponent> ReregisterAnimatedComponents;
-
-	// Update the pivot
-	if (PivotMode != ESpritePivotMode::Custom)
+	if (MemberPropertyName == GET_MEMBER_NAME_CHECKED(UPaperSprite, SourceUV))
 	{
-		CustomPivotPoint = GetPivotPosition();
+		SourceUV.X = FMath::Max(FMath::RoundToFloat(SourceUV.X), 0.0f);
+		SourceUV.Y = FMath::Max(FMath::RoundToFloat(SourceUV.Y), 0.0f);
+	}
+	else if (MemberPropertyName == GET_MEMBER_NAME_CHECKED(UPaperSprite, SourceDimension))
+	{
+		SourceDimension.X = FMath::Max(FMath::RoundToFloat(SourceDimension.X), 0.0f);
+		SourceDimension.Y = FMath::Max(FMath::RoundToFloat(SourceDimension.Y), 0.0f);
 	}
 
-	// Look for changed properties
-
-	const FName PropertyName = (PropertyChangedEvent.Property != NULL) ? PropertyChangedEvent.Property->GetFName() : NAME_None;
+	// Update the pivot (roundtripping thru the function will round to a pixel position if that option is enabled)
+	CustomPivotPoint = GetPivotPosition();
 
 	bool bRenderDataModified = false;
 	bool bCollisionDataModified = false;
@@ -1106,7 +1112,7 @@ FTransform UPaperSprite::GetPivotToWorld() const
 	return FTransform(Translation);
 }
 
-FVector2D UPaperSprite::GetPivotPosition() const
+FVector2D UPaperSprite::GetRawPivotPosition() const
 {
 	switch (PivotMode)
 	{
@@ -1136,6 +1142,19 @@ FVector2D UPaperSprite::GetPivotPosition() const
 	};
 }
 
+FVector2D UPaperSprite::GetPivotPosition() const
+{
+	FVector2D RawPivot = GetRawPivotPosition();
+
+	if (bSnapPivotToPixelGrid)
+	{
+		RawPivot.X = FMath::RoundToFloat(RawPivot.X);
+		RawPivot.Y = FMath::RoundToFloat(RawPivot.Y);
+	}
+
+	return RawPivot;
+}
+
 void UPaperSprite::GetAssetRegistryTags(TArray<FAssetRegistryTag>& OutTags) const
 {
 	Super::GetAssetRegistryTags(OutTags);
@@ -1225,6 +1244,11 @@ void UPaperSprite::PostLoad()
 		SetFlags(RF_Transactional);
 	}
 
+	if (PaperVer < FPaperCustomVersion::AddPivotSnapToPixelGrid)
+	{
+		bSnapPivotToPixelGrid = false;
+	}
+
 	if (PaperVer < FPaperCustomVersion::AddPixelsPerUnrealUnit)
 	{
 		PixelsPerUnrealUnit = 1.0f;
diff --git a/Engine/Plugins/Experimental/Paper2D/Source/Paper2DEditor/Private/SpriteEditor/SpriteEditorViewportClient.cpp b/Engine/Plugins/Experimental/Paper2D/Source/Paper2DEditor/Private/SpriteEditor/SpriteEditorViewportClient.cpp
--- a/Engine/Plugins/Experimental/Paper2D/Source/Paper2DEditor/Private/SpriteEditor/SpriteEditorViewportClient.cpp
+++ b/Engine/Plugins/Experimental/Paper2D/Source/Paper2DEditor/Private/SpriteEditor/SpriteEditorViewportClient.cpp
@@ -865,7 +865,6 @@ void FSpriteEditorViewportClient::BeginTransaction(const FText& SessionName)
 {
 	if (ScopedTransaction == NULL)
 	{
-
 		ScopedTransaction = new FScopedTransaction(SessionName);
 
 		UPaperSprite* Sprite = GetSpriteBeingEdited();
@@ -877,7 +876,18 @@ void FSpriteEditorViewportClient::EndTransaction()
 {
 	if (bManipulationDirtiedSomething)
 	{
-		RenderSpriteComponent->GetSprite()->PostEditChange();
+		UPaperSprite* Sprite = RenderSpriteComponent->GetSprite();
+
+		if (IsInSourceRegionEditMode())
+		{
+			// Snap to pixel grid at the end of the drag
+			Sprite->SourceUV.X = FMath::Max(FMath::RoundToFloat(Sprite->SourceUV.X), 0.0f);
+			Sprite->SourceUV.Y = FMath::Max(FMath::RoundToFloat(Sprite->SourceUV.Y), 0.0f);
+			Sprite->SourceDimension.X = FMath::Max(FMath::RoundToFloat(Sprite->SourceDimension.X), 0.0f);
+			Sprite->SourceDimension.Y = FMath::Max(FMath::RoundToFloat(Sprite->SourceDimension.Y), 0.0f);
+		}
+
+		Sprite->PostEditChange();
 	}
 	
 	bManipulationDirtiedSomething = false;
