Paper2D: Force wireframe color for opaque sprite sections to green (the color of GPU happiness)

diff --git a/Engine/Plugins/Experimental/Paper2D/Source/Paper2D/Private/PaperRenderSceneProxy.cpp b/Engine/Plugins/Experimental/Paper2D/Source/Paper2D/Private/PaperRenderSceneProxy.cpp
--- a/Engine/Plugins/Experimental/Paper2D/Source/Paper2D/Private/PaperRenderSceneProxy.cpp
+++ b/Engine/Plugins/Experimental/Paper2D/Source/Paper2D/Private/PaperRenderSceneProxy.cpp
@@ -280,6 +280,7 @@ void FPaperRenderSceneProxy::DrawNewBatches(FPrimitiveDrawInterface* PDI, const
 
 	for (const FSpriteRenderSection& Batch : BatchedSections)
 	{
+		const FLinearColor EffectiveWireframeColor = (Batch.Material->GetBlendMode() != BLEND_Opaque) ? WireframeColor : FLinearColor::Green;
 		FMaterialRenderProxy* ParentMaterialProxy = Batch.Material->GetRenderProxy((View->Family->EngineShowFlags.Selection) && IsSelected(), IsHovered());
 
 		FTexture* TextureResource = (Batch.Texture != nullptr) ? Batch.Texture->Resource : nullptr;
@@ -299,7 +300,7 @@ void FPaperRenderSceneProxy::DrawNewBatches(FPrimitiveDrawInterface* PDI, const
 			DrawRichMesh(
 				PDI, 
 				Mesh,
-				WireframeColor,
+				EffectiveWireframeColor,
 				FLinearColor(1.0f, 1.0f, 1.0f),//LevelColor,
 				FLinearColor(1.0f, 1.0f, 1.0f),//PropertyColor,
 				this,
@@ -316,6 +317,9 @@ void FPaperRenderSceneProxy::DrawBatch(FPrimitiveDrawInterface* PDI, const FScen
 	
 	FVertexFactory* VertexFactory = GetPaperSpriteVertexFactory();
 
+	const bool bIsWireframeView = View->Family->EngineShowFlags.Wireframe;
+	const FLinearColor EffectiveWireframeColor = (BatchMaterial->GetBlendMode() != BLEND_Opaque) ? WireframeColor : FLinearColor::Green;
+
 	for (const FSpriteDrawCallRecord& Record : Batch)
 	{
 		FTexture* TextureResource = (Record.Texture != nullptr) ? Record.Texture->Resource : nullptr;
@@ -371,12 +375,10 @@ void FPaperRenderSceneProxy::DrawBatch(FPrimitiveDrawInterface* PDI, const FScen
 			BatchElement.PrimitiveUniformBuffer = GetUniformBuffer();
 			BatchElement.NumPrimitives = Vertices.Num() / 3;
 
-			const bool bIsWireframeView = View->Family->EngineShowFlags.Wireframe;
-
 			DrawRichMesh(
 				PDI, 
 				Mesh,
-				WireframeColor,
+				EffectiveWireframeColor,
 				FLinearColor(1.0f, 1.0f, 1.0f),//LevelColor,
 				FLinearColor(1.0f, 1.0f, 1.0f),//PropertyColor,
 				this,
