Paper2D: UX improvements to the details panel customization for UPaperSprite

diff --git a/Engine/Plugins/Experimental/Paper2D/Source/Paper2DEditor/Private/SpriteEditor/SpriteDetailsCustomization.cpp b/Engine/Plugins/Experimental/Paper2D/Source/Paper2DEditor/Private/SpriteEditor/SpriteDetailsCustomization.cpp
--- a/Engine/Plugins/Experimental/Paper2D/Source/Paper2DEditor/Private/SpriteEditor/SpriteDetailsCustomization.cpp
+++ b/Engine/Plugins/Experimental/Paper2D/Source/Paper2DEditor/Private/SpriteEditor/SpriteDetailsCustomization.cpp
@@ -90,8 +90,19 @@ void FSpriteDetailsCustomization::BuildSpriteSection(IDetailCategoryBuilder& Spr
 
 void FSpriteDetailsCustomization::BuildRenderingSection(IDetailCategoryBuilder& RenderingCategory, IDetailLayoutBuilder& DetailLayout)
 {
+	// Add the collision geometry mode into the parent container (renamed)
+	const FString RenderGeometryTypePropertyPath = FString::Printf(TEXT("%s.%s"), GET_MEMBER_NAME_STRING_CHECKED(UPaperSprite, RenderGeometry), GET_MEMBER_NAME_STRING_CHECKED(FSpritePolygonCollection, GeometryType));
+	RenderingCategory.AddProperty(DetailLayout.GetProperty(*RenderGeometryTypePropertyPath))
+		.DisplayName(LOCTEXT("RenderGeometryType", "Render Geometry Type").ToString());
+
+	// Show the rendering geometry settings
 	TSharedRef<IPropertyHandle> RenderGeometry = DetailLayout.GetProperty(GET_MEMBER_NAME_CHECKED(UPaperSprite, RenderGeometry));
 	IDetailPropertyRow& RenderGeometryProperty = RenderingCategory.AddProperty(RenderGeometry);
+
+	// Add the render polygons into advanced (renamed)
+	const FString RenderGeometryPolygonsPropertyPath = FString::Printf(TEXT("%s.%s"), GET_MEMBER_NAME_STRING_CHECKED(UPaperSprite, RenderGeometry), GET_MEMBER_NAME_STRING_CHECKED(FSpritePolygonCollection, Polygons));
+	RenderingCategory.AddProperty(DetailLayout.GetProperty(*RenderGeometryPolygonsPropertyPath), EPropertyLocation::Advanced)
+		.DisplayName(LOCTEXT("RenderPolygons", "Render Polygons").ToString());
 }
 
 void FSpriteDetailsCustomization::BuildCollisionSection(IDetailCategoryBuilder& CollisionCategory, IDetailLayoutBuilder& DetailLayout)
@@ -116,6 +127,12 @@ void FSpriteDetailsCustomization::BuildCollisionSection(IDetailCategoryBuilder&
 	GenerateWarningRow(CollisionCategory, /*bExperimental=*/ false, QueryWarningFor2D, QueryTooltipFor2D, TEXT("Shared/Editors/SpriteEditor"), TEXT("Disabled2DCollisionQueriesWarning"))
 		.Visibility(WarnAbout2DQueriesBeingDisabledVisibility);
 
+	// Add the collision geometry mode into the parent container (renamed)
+	const FString CollisionGeometryTypePropertyPath = FString::Printf(TEXT("%s.%s"), GET_MEMBER_NAME_STRING_CHECKED(UPaperSprite, CollisionGeometry), GET_MEMBER_NAME_STRING_CHECKED(FSpritePolygonCollection, GeometryType));
+	CollisionCategory.AddProperty(DetailLayout.GetProperty(*CollisionGeometryTypePropertyPath))
+		.DisplayName(LOCTEXT("CollisionGeometryType", "Collision Geometry Type").ToString())
+		.Visibility(ParticipatesInPhysics);
+
 	// Show the collision geometry when not None
 	CollisionCategory.AddProperty( DetailLayout.GetProperty(GET_MEMBER_NAME_CHECKED(UPaperSprite, CollisionGeometry)) )
 		.Visibility(ParticipatesInPhysics);
@@ -123,6 +140,12 @@ void FSpriteDetailsCustomization::BuildCollisionSection(IDetailCategoryBuilder&
 	// Show the collision thickness only in 3D mode
 	CollisionCategory.AddProperty( DetailLayout.GetProperty(GET_MEMBER_NAME_CHECKED(UPaperSprite, CollisionThickness)) )
 		.Visibility(ParticipatesInPhysics3D);
+
+	// Add the collision polygons into advanced (renamed)
+	const FString CollisionGeometryPolygonsPropertyPath = FString::Printf(TEXT("%s.%s"), GET_MEMBER_NAME_STRING_CHECKED(UPaperSprite, CollisionGeometry), GET_MEMBER_NAME_STRING_CHECKED(FSpritePolygonCollection, Polygons));
+	CollisionCategory.AddProperty(DetailLayout.GetProperty(*CollisionGeometryPolygonsPropertyPath), EPropertyLocation::Advanced)
+		.DisplayName(LOCTEXT("CollisionPolygons", "Collision Polygons").ToString())
+		.Visibility(ParticipatesInPhysics);
 }
 
 EVisibility FSpriteDetailsCustomization::PhysicsModeMatches(TSharedPtr<IPropertyHandle> Property, ESpriteCollisionMode::Type DesiredMode) const
