Paper2D: Remove display of collision geometry debug text if there is no collision on the sprite

diff --git a/Engine/Plugins/Experimental/Paper2D/Source/Paper2DEditor/Private/SpriteEditor/SpriteEditorViewportClient.cpp b/Engine/Plugins/Experimental/Paper2D/Source/Paper2DEditor/Private/SpriteEditor/SpriteEditorViewportClient.cpp
--- a/Engine/Plugins/Experimental/Paper2D/Source/Paper2DEditor/Private/SpriteEditor/SpriteEditorViewportClient.cpp
+++ b/Engine/Plugins/Experimental/Paper2D/Source/Paper2DEditor/Private/SpriteEditor/SpriteEditorViewportClient.cpp
@@ -426,10 +426,6 @@ void FSpriteEditorViewportClient::DrawRenderStats(FViewport& InViewport, FSceneV
 		TextItem.Position.Y += 18.0f;
 	}
 
-	TextItem.Text = FText::Format(LOCTEXT("BakedVertexCount", "Verts: {0}"), FText::AsNumber(Sprite->BakedRenderData.Num()));
-	TextItem.Draw(&Canvas);
-	TextItem.Position.Y += 18.0f;
-
 	YPos = (int32)TextItem.Position.Y;
 }
 
@@ -566,9 +562,13 @@ void FSpriteEditorViewportClient::DrawCanvas(FViewport& Viewport, FSceneView& Vi
 			YPos += 18;
 		}
 
-		// As well as all of the geometry
-		DrawGeometryStats(Viewport, View, Canvas, Sprite->CollisionGeometry, false, /*inout*/ YPos);
-		DrawGeometryStats(Viewport, View, Canvas, Sprite->RenderGeometry, true, /*inout*/ YPos);
+		// Collision geometry (if present)
+		if (Sprite->BodySetup != nullptr)
+		{
+			DrawGeometryStats(Viewport, View, Canvas, Sprite->CollisionGeometry, false, /*inout*/ YPos);
+		}
+
+		// Baked render data
 		DrawRenderStats(Viewport, View, Canvas, Sprite, /*inout*/ YPos);
 
 		// And bounds
@@ -588,7 +588,10 @@ void FSpriteEditorViewportClient::DrawCanvas(FViewport& Viewport, FSceneView& Vi
 			// Draw the custom collision geometry
 			const FLinearColor CollisionColor(1.0f, 1.0f, 0.0f, 1.0f);
 			DrawGeometry(Viewport, View, Canvas, Sprite->CollisionGeometry, CollisionColor, false);
-			DrawGeometryStats(Viewport, View, Canvas, Sprite->CollisionGeometry, false, /*inout*/ YPos);
+			if (Sprite->BodySetup != nullptr)
+			{
+				DrawGeometryStats(Viewport, View, Canvas, Sprite->CollisionGeometry, false, /*inout*/ YPos);
+			}
 		}
 		break;
 	case ESpriteEditorMode::EditRenderingGeomMode:
