Paper2D: Add consideration for PixelsPerUnrealUnit in UPaperSprite::ConvertWorldSpaceToTextureSpace
- Fixes double-clicking in Source view when PixelsPerUnrealUnit is not 1.0

diff --git a/Engine/Plugins/Experimental/Paper2D/Source/Paper2D/Private/PaperSprite.cpp b/Engine/Plugins/Experimental/Paper2D/Source/Paper2D/Private/PaperSprite.cpp
--- a/Engine/Plugins/Experimental/Paper2D/Source/Paper2D/Private/PaperSprite.cpp
+++ b/Engine/Plugins/Experimental/Paper2D/Source/Paper2D/Private/PaperSprite.cpp
@@ -1100,7 +1100,11 @@ FVector2D UPaperSprite::ConvertWorldSpaceToTextureSpace(const FVector& WorldPoin
 	const FVector2D SourceDims = (SourceTexture != NULL) ? FVector2D(SourceTexture->GetSurfaceWidth(), SourceTexture->GetSurfaceHeight()) : FVector2D::ZeroVector;
 	const FVector ProjectionX = WorldPoint.ProjectOnTo(PaperAxisX);
 	const FVector ProjectionY = WorldPoint.ProjectOnTo(PaperAxisY);
-	return FVector2D(FMath::Sign(ProjectionX | PaperAxisX) * ProjectionX.Size(), SourceDims.Y - FMath::Sign(ProjectionY | PaperAxisY) * ProjectionY.Size());		
+
+	const float XValue = FMath::Sign(ProjectionX | PaperAxisX) * ProjectionX.Size() * PixelsPerUnrealUnit;
+	const float YValue = FMath::Sign(ProjectionY | PaperAxisY) * ProjectionY.Size() * PixelsPerUnrealUnit;
+
+	return FVector2D(XValue, SourceDims.Y - YValue);
 }
 
 FTransform UPaperSprite::GetPivotToWorld() const
