Paper2D: Add OnFinishedPlaying delegate to UPaperFlipbookComponent, called whenever a non-looping flipbook finishes playing (reaches the end or beginning, depending on play direction)

diff --git a/Engine/Plugins/Experimental/Paper2D/Source/Paper2D/Classes/PaperFlipbookComponent.h b/Engine/Plugins/Experimental/Paper2D/Source/Paper2D/Classes/PaperFlipbookComponent.h
--- a/Engine/Plugins/Experimental/Paper2D/Source/Paper2D/Classes/PaperFlipbookComponent.h
+++ b/Engine/Plugins/Experimental/Paper2D/Source/Paper2D/Classes/PaperFlipbookComponent.h
@@ -5,6 +5,9 @@
 #include "PaperFlipbook.h"
 #include "PaperFlipbookComponent.generated.h"
 
+// Event for a non-looping flipbook finishing play
+DECLARE_DYNAMIC_MULTICAST_DELEGATE(FFlipbookFinishedPlaySignature);
+
 UCLASS(ShowCategories=(Mobility, ComponentReplication), EarlyAccessPreview, meta=(BlueprintSpawnableComponent))
 class PAPER2D_API UPaperFlipbookComponent : public UPrimitiveComponent
 {
@@ -47,6 +50,11 @@ protected:
 	UPROPERTY(BlueprintReadOnly, Interp, Category=Sprite)
 	FLinearColor SpriteColor;
 
+public:
+	/** Event called whenever a non-looping flipbook finishes playing (either reaching the beginning or the end, depending on the play direction) */
+	UPROPERTY(BlueprintAssignable)
+	FFlipbookFinishedPlaySignature OnFinishedPlaying;
+
 public:
 	/** Change the flipbook used by this instance (will reset the play time to 0 if it is a new flipbook). */
 	UFUNCTION(BlueprintCallable, Category="Sprite")
diff --git a/Engine/Plugins/Experimental/Paper2D/Source/Paper2D/Private/PaperFlipbookComponent.cpp b/Engine/Plugins/Experimental/Paper2D/Source/Paper2D/Private/PaperFlipbookComponent.cpp
--- a/Engine/Plugins/Experimental/Paper2D/Source/Paper2D/Private/PaperFlipbookComponent.cpp
+++ b/Engine/Plugins/Experimental/Paper2D/Source/Paper2D/Private/PaperFlipbookComponent.cpp
@@ -180,14 +180,11 @@ void UPaperFlipbookComponent::TickFlipbook(float DeltaTime)
 		SetPlaybackPosition(NewPosition, true);
 	}
 
-#if 0
-	// Notify user that timeline finished
+	// Notify user that the flipbook finished playing
 	if (bIsFinished)
 	{
-		TimelineFinishedFunc.ExecuteIfBound();
-		TimelineFinishFuncStatic.ExecuteIfBound();
+		OnFinishedPlaying.Broadcast();
 	}
-#endif
 }
 
 void UPaperFlipbookComponent::GetLifetimeReplicatedProps(TArray< FLifetimeProperty > & OutLifetimeProps) const
