Paper2D: Add drag-drop reordering of keyframes in flipbooks

diff --git a/Engine/Plugins/Experimental/Paper2D/Source/Paper2DEditor/Private/FlipbookEditor/SFlipbookTimeline.cpp b/Engine/Plugins/Experimental/Paper2D/Source/Paper2DEditor/Private/FlipbookEditor/SFlipbookTimeline.cpp
--- a/Engine/Plugins/Experimental/Paper2D/Source/Paper2DEditor/Private/FlipbookEditor/SFlipbookTimeline.cpp
+++ b/Engine/Plugins/Experimental/Paper2D/Source/Paper2DEditor/Private/FlipbookEditor/SFlipbookTimeline.cpp
@@ -5,6 +5,7 @@
 #include "Editor/UnrealEd/Public/DragAndDrop/AssetDragDropOp.h"
 #include "Editor/UnrealEd/Public/ScopedTransaction.h"
 #include "FlipbookEditorCommands.h"
+#include "PaperStyle.h"
 
 #define LOCTEXT_NAMESPACE "FlipbookEditor"
 
@@ -26,7 +27,6 @@ void SFlipbookTimeline::Construct(const FArguments& InArgs, TSharedPtr<const FUI
 	CommandList = InCommandList;
 
 	SlateUnitsPerFrame = 32;
-	const int32 FrameHeight = 48;
 
 	ChildSlot
 	[
@@ -63,9 +63,9 @@ void SFlipbookTimeline::Construct(const FArguments& InArgs, TSharedPtr<const FUI
 				+SVerticalBox::Slot()
 				.AutoHeight()
 				[
-					SNew(SBox).HeightOverride(FrameHeight)
+					SNew(SBox).HeightOverride(FFlipbookUIConstants::FrameHeight)
 					[
-						SNew(STimelineTrack, InCommandList)
+						SNew(SFlipbookTimelineTrack, InCommandList)
 						.SlateUnitsPerFrame(SlateUnitsPerFrame)
 						.FlipbookBeingEdited(FlipbookBeingEdited)
 						.OnSelectionChanged(OnSelectionChanged)
@@ -76,6 +76,36 @@ void SFlipbookTimeline::Construct(const FArguments& InArgs, TSharedPtr<const FUI
 	];
 }
 
+void SFlipbookTimeline::OnDragEnter(const FGeometry& MyGeometry, const FDragDropEvent& DragDropEvent)
+{
+	SCompoundWidget::OnDragEnter(MyGeometry, DragDropEvent);
+
+	TSharedPtr<FDragDropOperation> Operation = DragDropEvent.GetOperation();
+	if (!Operation.IsValid())
+	{
+	}
+	else if (Operation->IsOfType<FFlipbookKeyFrameDragDropOp>())
+	{
+		const auto& FrameDragDropOp = StaticCastSharedPtr<FFlipbookKeyFrameDragDropOp>(Operation);
+		FrameDragDropOp->SetCanDropHere(true);
+	}
+}
+
+void SFlipbookTimeline::OnDragLeave(const FDragDropEvent& DragDropEvent)
+{
+	SCompoundWidget::OnDragLeave(DragDropEvent);
+
+	TSharedPtr<FDragDropOperation> Operation = DragDropEvent.GetOperation();
+	if (!Operation.IsValid())
+	{
+	}
+	else if (Operation->IsOfType<FFlipbookKeyFrameDragDropOp>())
+	{
+		const auto& FrameDragDropOp = StaticCastSharedPtr<FFlipbookKeyFrameDragDropOp>(Operation);
+		FrameDragDropOp->SetCanDropHere(false);
+	}
+}
+
 FReply SFlipbookTimeline::OnDrop(const FGeometry& MyGeometry, const FDragDropEvent& DragDropEvent)
 {
 	bool bWasDropHandled = false;
@@ -86,12 +116,19 @@ FReply SFlipbookTimeline::OnDrop(const FGeometry& MyGeometry, const FDragDropEve
 	}
 	else if (Operation->IsOfType<FAssetDragDropOp>())
 	{
-		const auto& DragDropOp = StaticCastSharedPtr<FAssetDragDropOp>(Operation);
-
-		OnAssetsDropped(*DragDropOp);
-
+		const auto& AssetDragDropOp = StaticCastSharedPtr<FAssetDragDropOp>(Operation);
+		OnAssetsDropped(*AssetDragDropOp);
 		bWasDropHandled = true;
 	}
+	else if (Operation->IsOfType<FFlipbookKeyFrameDragDropOp>())
+	{
+		const auto& FrameDragDropOp = StaticCastSharedPtr<FFlipbookKeyFrameDragDropOp>(Operation);
+		if (UPaperFlipbook* ThisFlipbook = FlipbookBeingEdited.Get())
+		{
+			FrameDragDropOp->AppendToFlipbook(ThisFlipbook);
+			bWasDropHandled = true;
+		}
+	}
 
 	return bWasDropHandled ? FReply::Handled() : FReply::Unhandled();
 }
diff --git a/Engine/Plugins/Experimental/Paper2D/Source/Paper2DEditor/Private/FlipbookEditor/SFlipbookTimeline.h b/Engine/Plugins/Experimental/Paper2D/Source/Paper2DEditor/Private/FlipbookEditor/SFlipbookTimeline.h
--- a/Engine/Plugins/Experimental/Paper2D/Source/Paper2DEditor/Private/FlipbookEditor/SFlipbookTimeline.h
+++ b/Engine/Plugins/Experimental/Paper2D/Source/Paper2DEditor/Private/FlipbookEditor/SFlipbookTimeline.h
@@ -23,6 +23,8 @@ public:
 	void Construct(const FArguments& InArgs, TSharedPtr<const FUICommandList> InCommandList);
 
 	// SWidget interface
+	virtual void OnDragEnter(const FGeometry& MyGeometry, const FDragDropEvent& DragDropEvent) override;
+	virtual void OnDragLeave(const FDragDropEvent& DragDropEvent) override;
 	virtual FReply OnDrop(const FGeometry& MyGeometry, const FDragDropEvent& DragDropEvent) override;
 	virtual int32 OnPaint(const FPaintArgs& Args, const FGeometry& AllottedGeometry, const FSlateRect& MyClippingRect, FSlateWindowElementList& OutDrawElements, int32 LayerId, const FWidgetStyle& InWidgetStyle, bool bParentEnabled) const override;
 	virtual FReply OnMouseButtonUp(const FGeometry& MyGeometry, const FPointerEvent& MouseEvent) override;
diff --git a/Engine/Plugins/Experimental/Paper2D/Source/Paper2DEditor/Private/FlipbookEditor/SFlipbookTrackHandle.h b/Engine/Plugins/Experimental/Paper2D/Source/Paper2DEditor/Private/FlipbookEditor/SFlipbookTrackHandle.h
--- a/Engine/Plugins/Experimental/Paper2D/Source/Paper2DEditor/Private/FlipbookEditor/SFlipbookTrackHandle.h
+++ b/Engine/Plugins/Experimental/Paper2D/Source/Paper2DEditor/Private/FlipbookEditor/SFlipbookTrackHandle.h
@@ -31,7 +31,7 @@ public:
 
 		SImage::Construct(
 			SImage::FArguments()
-			.Image(FEditorStyle::GetBrush("FlipbookEditor.RegionGrabHandle")));
+			.Image(FPaperStyle::Get()->GetBrush("FlipbookEditor.RegionGrabHandle")));
 	}
 
 	virtual FReply OnMouseButtonDown(const FGeometry& MyGeometry, const FPointerEvent& MouseEvent) override
diff --git a/Engine/Plugins/Experimental/Paper2D/Source/Paper2DEditor/Private/FlipbookEditor/STimelineTrack.h b/Engine/Plugins/Experimental/Paper2D/Source/Paper2DEditor/Private/FlipbookEditor/STimelineTrack.h
--- a/Engine/Plugins/Experimental/Paper2D/Source/Paper2DEditor/Private/FlipbookEditor/STimelineTrack.h
+++ b/Engine/Plugins/Experimental/Paper2D/Source/Paper2DEditor/Private/FlipbookEditor/STimelineTrack.h
@@ -1,12 +1,120 @@
 // Copyright 1998-2014 Epic Games, Inc. All Rights Reserved.
 
+namespace FFlipbookUIConstants
+{
+	const float HandleWidth = 12.0f;
+	const float FrameHeight = 48;
+	const FMargin FramePadding(0.0f, 7.0f, 0.0f, 7.0f);
+};
+
 //////////////////////////////////////////////////////////////////////////
-// STimelineTrack
+// FFlipbookKeyFrameDragDropOp
 
-class STimelineTrack : public SCompoundWidget
+class FFlipbookKeyFrameDragDropOp : public FDragDropOperation
 {
 public:
-	SLATE_BEGIN_ARGS(STimelineTrack)
+	DRAG_DROP_OPERATOR_TYPE(FFlipbookKeyFrameDragDropOp, FDragDropOperation)
+
+	float WidgetWidth;
+	FPaperFlipbookKeyFrame KeyFrameData;
+	int32 SourceFrameIndex;
+	TWeakObjectPtr<UPaperFlipbook> SourceFlipbook;
+	FScopedTransaction Transaction;
+
+	virtual TSharedPtr<SWidget> GetDefaultDecorator() const override
+	{
+		const FLinearColor BorderColor = (KeyFrameData.Sprite != nullptr) ? FLinearColor::White : FLinearColor::Black;
+
+		return SNew(SBox)
+			.WidthOverride(WidgetWidth - FFlipbookUIConstants::FramePadding.GetTotalSpaceAlong<Orient_Horizontal>())
+			.HeightOverride(FFlipbookUIConstants::FrameHeight - FFlipbookUIConstants::FramePadding.GetTotalSpaceAlong<Orient_Vertical>())
+			[
+				SNew(SBorder)
+				.BorderImage(FPaperStyle::Get()->GetBrush("FlipbookEditor.RegionBody"))
+				.BorderBackgroundColor(BorderColor)
+				[
+					SNullWidget::NullWidget
+				]
+			];
+	}
+
+	virtual void OnDragged(const class FDragDropEvent& DragDropEvent) override
+	{
+		if (CursorDecoratorWindow.IsValid())
+		{
+			CursorDecoratorWindow->MoveWindowTo(DragDropEvent.GetScreenSpacePosition());
+		}
+	}
+
+	virtual void Construct() override
+	{
+		MouseCursor = EMouseCursor::GrabHandClosed;
+		FDragDropOperation::Construct();
+	}
+
+	virtual void OnDrop(bool bDropWasHandled, const FPointerEvent& MouseEvent) override
+	{
+		if (!bDropWasHandled)
+		{
+			// Add us back to our source, the drop fizzled
+			InsertInFlipbook(SourceFlipbook.Get(), SourceFrameIndex);
+			Transaction.Cancel();
+		}
+	}
+
+	void AppendToFlipbook(UPaperFlipbook* DestinationFlipbook)
+	{
+		DestinationFlipbook->Modify();
+		FScopedFlipbookMutator EditLock(DestinationFlipbook);
+		EditLock.KeyFrames.Add(KeyFrameData);
+	}
+
+	void InsertInFlipbook(UPaperFlipbook* DestinationFlipbook, int32 Index)
+	{
+		DestinationFlipbook->Modify();
+		FScopedFlipbookMutator EditLock(DestinationFlipbook);
+		EditLock.KeyFrames.Insert(KeyFrameData, Index);
+	}
+
+	void SetCanDropHere(bool bCanDropHere)
+	{
+		MouseCursor = bCanDropHere ? EMouseCursor::TextEditBeam : EMouseCursor::SlashedCircle;
+	}
+
+	static TSharedRef<FFlipbookKeyFrameDragDropOp> New(int32 InWidth, UPaperFlipbook* InFlipbook, int32 InFrameIndex)
+	{
+		// Create the drag-drop op containing the key
+		TSharedRef<FFlipbookKeyFrameDragDropOp> Operation = MakeShareable(new FFlipbookKeyFrameDragDropOp);
+		Operation->KeyFrameData = InFlipbook->GetKeyFrameChecked(InFrameIndex);
+		Operation->SourceFrameIndex = InFrameIndex;
+		Operation->SourceFlipbook = InFlipbook;
+		Operation->WidgetWidth = InWidth;
+		Operation->Construct();
+
+		// Remove the key from the flipbook
+		{
+			InFlipbook->Modify();
+			FScopedFlipbookMutator EditLock(InFlipbook);
+			EditLock.KeyFrames.RemoveAt(InFrameIndex);
+		}
+
+		return Operation;
+	}
+
+protected:
+	FFlipbookKeyFrameDragDropOp()
+		: Transaction(LOCTEXT("MovedFramesInTimeline", "Reorder key frames"))
+	{
+	}
+};
+
+//////////////////////////////////////////////////////////////////////////
+// SFlipbookKeyframeWidget
+
+class SFlipbookKeyframeWidget : public SCompoundWidget
+{
+public:
+	SLATE_BEGIN_ARGS(SFlipbookKeyframeWidget)
 		: _SlateUnitsPerFrame(1)
 		, _FlipbookBeingEdited(nullptr)
 	{}
@@ -17,42 +125,120 @@
 
 	SLATE_END_ARGS()
 
-	void Construct(const FArguments& InArgs, TSharedPtr<const FUICommandList> InCommandList)
+	void Construct(const FArguments& InArgs, int32 InFrameIndex, TSharedPtr<const FUICommandList> InCommandList)
 	{
+		FrameIndex = InFrameIndex;
 		CommandList = InCommandList;
 		SlateUnitsPerFrame = InArgs._SlateUnitsPerFrame;
 		FlipbookBeingEdited = InArgs._FlipbookBeingEdited;
-		SelectedFrame = INDEX_NONE;
 		OnSelectionChanged = InArgs._OnSelectionChanged;
 
-		HandleWidth = 12;
-		NumKeyframesFromLastRebuild = 0;
+		// Color each region based on whether a sprite has been set or not for it
+		const auto BorderColorDelegate = [](TAttribute<UPaperFlipbook*> ThisFlipbookPtr, int32 TestIndex) -> FSlateColor
+		{
+			UPaperFlipbook* FlipbookPtr = ThisFlipbookPtr.Get();
+			const bool bFrameValid = (FlipbookPtr != nullptr) && (FlipbookPtr->GetSpriteAtFrame(TestIndex) != nullptr);
+			return bFrameValid ? FLinearColor::White : FLinearColor::Black;
+		};
 
 		ChildSlot
 		[
-			SAssignNew(MainBoxPtr, SHorizontalBox)
+			SNew(SHorizontalBox)
+			+SHorizontalBox::Slot()
+			.AutoWidth()
+			[
+				SNew(SBox)
+				.Padding(FFlipbookUIConstants::FramePadding)
+				.WidthOverride(this, &SFlipbookKeyframeWidget::GetFrameWidth)
+				[
+					SNew(SBorder)
+					.BorderImage(FPaperStyle::Get()->GetBrush("FlipbookEditor.RegionBody"))
+					.BorderBackgroundColor_Static(BorderColorDelegate, FlipbookBeingEdited, FrameIndex)
+					.OnMouseButtonUp(this, &SFlipbookKeyframeWidget::KeyframeOnMouseButtonUp)
+					.ToolTipText(this, &SFlipbookKeyframeWidget::GetKeyframeTooltip)
+					[
+						SNullWidget::NullWidget
+					]
+				]
+			]
+			+SHorizontalBox::Slot()
+			.AutoWidth()
+			[
+				SNew(SBox)
+				.WidthOverride(FFlipbookUIConstants::HandleWidth)
+				[
+					SNew(SFlipbookTrackHandle)
+					.SlateUnitsPerFrame(SlateUnitsPerFrame)
+					.FlipbookBeingEdited(FlipbookBeingEdited)
+					.KeyFrameIdx(FrameIndex)
+				]
+			]
 		];
-
-		Rebuild();
 	}
 
-	void Tick(const FGeometry& AllottedGeometry, const double InCurrentTime, const float InDeltaTime) override
+	virtual FReply OnMouseButtonDown(const FGeometry& MyGeometry, const FPointerEvent& MouseEvent) override
 	{
-		SCompoundWidget::Tick(AllottedGeometry, InCurrentTime, InDeltaTime);
+		if (MouseEvent.GetEffectingButton() == EKeys::LeftMouseButton)
+		{
+			return FReply::Handled().DetectDrag(SharedThis(this), EKeys::LeftMouseButton);
+		}
 
+		return FReply::Unhandled();
+	}
+
+	virtual FReply OnDragDetected(const FGeometry& MyGeometry, const FPointerEvent& MouseEvent) override
+	{
+		if (MouseEvent.IsMouseButtonDown(EKeys::LeftMouseButton))
+		{
+			UPaperFlipbook* Flipbook = FlipbookBeingEdited.Get();
+			if ((Flipbook != nullptr) && Flipbook->IsValidKeyFrameIndex(FrameIndex))
+			{
+				TSharedRef<FFlipbookKeyFrameDragDropOp> Operation = FFlipbookKeyFrameDragDropOp::New(
+					GetFrameWidth().Get(), Flipbook, FrameIndex);
+
+				return FReply::Handled().BeginDragDrop(Operation);
+			}
+		}
+
+		return FReply::Unhandled();
+	}
+
+	FReply OnDrop(const FGeometry& MyGeometry, const FDragDropEvent& DragDropEvent) override
+	{
+		bool bWasDropHandled = false;
+
 		UPaperFlipbook* Flipbook = FlipbookBeingEdited.Get();
-		int32 NewNumKeyframes = (Flipbook != nullptr) ? Flipbook->GetNumKeyFrames() : 0;
-		if (NewNumKeyframes != NumKeyframesFromLastRebuild)
+		if ((Flipbook != nullptr) && Flipbook->IsValidKeyFrameIndex(FrameIndex))
 		{
-			Rebuild();
+
+			TSharedPtr<FDragDropOperation> Operation = DragDropEvent.GetOperation();
+			if (!Operation.IsValid())
+			{
+			}
+			else if (Operation->IsOfType<FAssetDragDropOp>())
+			{
+				const auto& AssetDragDropOp = StaticCastSharedPtr<FAssetDragDropOp>(Operation);
+				//@TODO: Handle asset inserts
+
+			// 			OnAssetsDropped(*AssetDragDropOp);
+			// 			bWasDropHandled = true;
+			}
+			else if (Operation->IsOfType<FFlipbookKeyFrameDragDropOp>())
+			{
+				const auto& FrameDragDropOp = StaticCastSharedPtr<FFlipbookKeyFrameDragDropOp>(Operation);
+				FrameDragDropOp->InsertInFlipbook(Flipbook, FrameIndex);
+				bWasDropHandled = true;
+			}
 		}
+
+		return bWasDropHandled ? FReply::Handled() : FReply::Unhandled();
 	}
 
-	FReply KeyframeOnMouseButtonUp(const FGeometry& MyGeometry, const FPointerEvent& MouseEvent, int32 FrameIndex)
+	FReply KeyframeOnMouseButtonUp(const FGeometry& MyGeometry, const FPointerEvent& MouseEvent)
 	{
 		if (MouseEvent.GetEffectingButton() == EKeys::RightMouseButton)
 		{
-			TSharedRef<SWidget> MenuContents = GenerateContextMenu(FrameIndex);
+			TSharedRef<SWidget> MenuContents = GenerateContextMenu();
 			FSlateApplication::Get().PushMenu(AsShared(), MenuContents, MouseEvent.GetScreenSpacePosition(), FPopupTransitionEffect(FPopupTransitionEffect::ContextMenu));
 
 			return FReply::Handled();
@@ -61,18 +247,18 @@
 		return FReply::Unhandled();
 	}
 
-	FText GetKeyframeTooltip(int32 KeyFrameIdx) const
+	FText GetKeyframeTooltip() const
 	{
 		UPaperFlipbook* Flipbook = FlipbookBeingEdited.Get();
-		if ((Flipbook != nullptr) && Flipbook->IsValidKeyFrameIndex(KeyFrameIdx))
+		if ((Flipbook != nullptr) && Flipbook->IsValidKeyFrameIndex(FrameIndex))
 		{
-			const FPaperFlipbookKeyFrame& KeyFrame = Flipbook->GetKeyFrameChecked(KeyFrameIdx);
+			const FPaperFlipbookKeyFrame& KeyFrame = Flipbook->GetKeyFrameChecked(FrameIndex);
 
 			FText SpriteLine = (KeyFrame.Sprite != nullptr) ? FText::FromString(KeyFrame.Sprite->GetName()) : LOCTEXT("NoSprite", "(none)");
 
 			return FText::Format(LOCTEXT("KeyFrameTooltip", "Sprite: {0}\nIndex: {1}\nDuration: {2} frame(s)"),
 				SpriteLine,
-				FText::AsNumber(KeyFrameIdx),
+				FText::AsNumber(FrameIndex),
 				FText::AsNumber(KeyFrame.FrameRun));
 		}
 		else
@@ -81,70 +267,36 @@
 		}
 	}
 
-private:
-	void Rebuild()
+	TSharedRef<SWidget> GenerateContextMenu()
 	{
-		MainBoxPtr->ClearChildren();
+		OnSelectionChanged.ExecuteIfBound(FrameIndex);
 
-		// Color each region based on whether a sprite has been set or not for it
-		const auto BorderColorDelegate = [](TAttribute<UPaperFlipbook*> ThisFlipbookPtr, int32 TestIndex) -> FSlateColor
-		{
-			UPaperFlipbook* FlipbookPtr = ThisFlipbookPtr.Get();
-			const bool bFrameValid = (FlipbookPtr != nullptr) && (FlipbookPtr->GetSpriteAtFrame(TestIndex) != nullptr);
-			return bFrameValid ? FLinearColor::White : FLinearColor::Black;
-		};
+		FMenuBuilder MenuBuilder(true, CommandList);
+		MenuBuilder.BeginSection("KeyframeActions", LOCTEXT("KeyframeActionsSectionHeader", "Keyframe Actions"));
 
-		// Create the sections for each keyframe
-		if (UPaperFlipbook* Flipbook = FlipbookBeingEdited.Get())
-		{
-			for (int32 KeyFrameIdx = 0; KeyFrameIdx < Flipbook->GetNumKeyFrames(); ++KeyFrameIdx)
-			{
-				MainBoxPtr->AddSlot()
-				.AutoWidth()
-				.Padding(FMargin(0.0f, 7.0f, 0.0f, 7.0f))
-				[
-					SNew(SBox)
-					.WidthOverride( TAttribute<FOptionalSize>::Create( TAttribute<FOptionalSize>::FGetter::CreateSP(this, &STimelineTrack::GetFrameWidth, KeyFrameIdx) ) )
-					[
-						SNew(SBorder)
-						.BorderImage(FEditorStyle::GetBrush("FlipbookEditor.RegionBody"))
-						.BorderBackgroundColor_Static(BorderColorDelegate, FlipbookBeingEdited, KeyFrameIdx)
-						.OnMouseButtonUp(this, &STimelineTrack::KeyframeOnMouseButtonUp, KeyFrameIdx)
-						.ToolTipText(this, &STimelineTrack::GetKeyframeTooltip, KeyFrameIdx)
-						[
-							SNullWidget::NullWidget
-						]
-					]
-				];
+		// 		MenuBuilder.AddMenuEntry(FGenericCommands::Get().Cut);
+		// 		MenuBuilder.AddMenuEntry(FGenericCommands::Get().Copy);
+		// 		MenuBuilder.AddMenuEntry(FGenericCommands::Get().Paste);
+		MenuBuilder.AddMenuEntry(FGenericCommands::Get().Duplicate);
+		MenuBuilder.AddMenuEntry(FGenericCommands::Get().Delete);
 
-				MainBoxPtr->AddSlot()
-				.AutoWidth()
-				[
-					SNew(SBox).WidthOverride(HandleWidth)
-					[
-						SNew(SFlipbookTrackHandle)
-						.SlateUnitsPerFrame(SlateUnitsPerFrame)
-						.FlipbookBeingEdited(FlipbookBeingEdited)
-						.KeyFrameIdx(KeyFrameIdx)
-					]
-				];
-			}
+		MenuBuilder.AddMenuSeparator();
 
-			NumKeyframesFromLastRebuild = Flipbook->GetNumKeyFrames();
-		}
-		else
-		{
-			NumKeyframesFromLastRebuild = 0;
-		}
+		MenuBuilder.AddMenuEntry(FFlipbookEditorCommands::Get().AddNewFrameBefore);
+		MenuBuilder.AddMenuEntry(FFlipbookEditorCommands::Get().AddNewFrameAfter);
+
+		MenuBuilder.EndSection();
+
+		return MenuBuilder.MakeWidget();
 	}
 
-	FOptionalSize GetFrameWidth(int32 KeyFrameIdx) const
+	FOptionalSize GetFrameWidth() const
 	{
 		UPaperFlipbook* Flipbook = FlipbookBeingEdited.Get();
-		if (Flipbook && Flipbook->IsValidKeyFrameIndex(KeyFrameIdx))
+		if (Flipbook && Flipbook->IsValidKeyFrameIndex(FrameIndex))
 		{
-			const FPaperFlipbookKeyFrame& KeyFrame = Flipbook->GetKeyFrameChecked(KeyFrameIdx);
-			return FMath::Max<float>(0, KeyFrame.FrameRun * SlateUnitsPerFrame.Get() - HandleWidth);
+			const FPaperFlipbookKeyFrame& KeyFrame = Flipbook->GetKeyFrameChecked(FrameIndex);
+			return FMath::Max<float>(0, KeyFrame.FrameRun * SlateUnitsPerFrame.Get() - FFlipbookUIConstants::HandleWidth);
 		}
 		else
 		{
@@ -152,31 +304,90 @@
 		}
 	}
 
-	TSharedRef<SWidget> GenerateContextMenu(int32 FrameIndex)
+protected:
+	int32 FrameIndex;
+	TAttribute<float> SlateUnitsPerFrame;
+	TAttribute<class UPaperFlipbook*> FlipbookBeingEdited;
+	FOnFlipbookKeyframeSelectionChanged OnSelectionChanged;
+	TSharedPtr<const FUICommandList> CommandList;
+};
+
+//////////////////////////////////////////////////////////////////////////
+// SFlipbookTimelineTrack
+
+class SFlipbookTimelineTrack : public SCompoundWidget
+{
+public:
+	SLATE_BEGIN_ARGS(SFlipbookTimelineTrack)
+		: _SlateUnitsPerFrame(1)
+		, _FlipbookBeingEdited(nullptr)
+	{}
+
+		SLATE_ATTRIBUTE( float, SlateUnitsPerFrame )
+		SLATE_ATTRIBUTE( class UPaperFlipbook*, FlipbookBeingEdited )
+		SLATE_EVENT( FOnFlipbookKeyframeSelectionChanged, OnSelectionChanged )
+
+	SLATE_END_ARGS()
+
+	void Construct(const FArguments& InArgs, TSharedPtr<const FUICommandList> InCommandList)
 	{
-		SelectedFrame = FrameIndex;
-		OnSelectionChanged.ExecuteIfBound(SelectedFrame);
+		CommandList = InCommandList;
+		SlateUnitsPerFrame = InArgs._SlateUnitsPerFrame;
+		FlipbookBeingEdited = InArgs._FlipbookBeingEdited;
+		OnSelectionChanged = InArgs._OnSelectionChanged;
 
-		FMenuBuilder MenuBuilder(true, CommandList);
-		MenuBuilder.BeginSection("KeyframeActions", LOCTEXT("KeyframeActionsSectionHeader", "Keyframe Actions"));
+		NumKeyframesFromLastRebuild = 0;
 
-// 		MenuBuilder.AddMenuEntry(FGenericCommands::Get().Cut);
-// 		MenuBuilder.AddMenuEntry(FGenericCommands::Get().Copy);
-// 		MenuBuilder.AddMenuEntry(FGenericCommands::Get().Paste);
-		MenuBuilder.AddMenuEntry(FGenericCommands::Get().Duplicate);
-		MenuBuilder.AddMenuEntry(FGenericCommands::Get().Delete);
+		ChildSlot
+		[
+			SAssignNew(MainBoxPtr, SHorizontalBox)
+		];
 
-		MenuBuilder.AddMenuSeparator();
+		Rebuild();
+	}
 
-		MenuBuilder.AddMenuEntry(FFlipbookEditorCommands::Get().AddNewFrameBefore);
-		MenuBuilder.AddMenuEntry(FFlipbookEditorCommands::Get().AddNewFrameAfter);
+	void Tick(const FGeometry& AllottedGeometry, const double InCurrentTime, const float InDeltaTime) override
+	{
+		SCompoundWidget::Tick(AllottedGeometry, InCurrentTime, InDeltaTime);
 
-		MenuBuilder.EndSection();
+		UPaperFlipbook* Flipbook = FlipbookBeingEdited.Get();
+		int32 NewNumKeyframes = (Flipbook != nullptr) ? Flipbook->GetNumKeyFrames() : 0;
+		if (NewNumKeyframes != NumKeyframesFromLastRebuild)
+		{
+			Rebuild();
+		}
+	}
 
-		return MenuBuilder.MakeWidget();
+private:
+	void Rebuild()
+	{
+		MainBoxPtr->ClearChildren();
+
+		// Create the sections for each keyframe
+		if (UPaperFlipbook* Flipbook = FlipbookBeingEdited.Get())
+		{
+			for (int32 KeyFrameIdx = 0; KeyFrameIdx < Flipbook->GetNumKeyFrames(); ++KeyFrameIdx)
+			{
+				//@TODO: Draggy bits go here
+				MainBoxPtr->AddSlot()
+				.AutoWidth()
+				[
+					SNew(SFlipbookKeyframeWidget, KeyFrameIdx, CommandList)
+					.SlateUnitsPerFrame(this->SlateUnitsPerFrame)
+					.FlipbookBeingEdited(this->FlipbookBeingEdited)
+					.OnSelectionChanged(this->OnSelectionChanged)
+				];
+			}
+
+			NumKeyframesFromLastRebuild = Flipbook->GetNumKeyFrames();
+		}
+		else
+		{
+			NumKeyframesFromLastRebuild = 0;
+		}
 	}
+
 private:
-	int32 SelectedFrame;
 	TAttribute<float> SlateUnitsPerFrame;
 	TAttribute< class UPaperFlipbook* > FlipbookBeingEdited;
 
