Paper2D: Add DLL export to UPaperSpriteComponent and APaperCharacter
#TTP 342269

diff --git a/Engine/Plugins/Experimental/Paper2D/Source/Paper2D/Classes/PaperCharacter.h b/Engine/Plugins/Experimental/Paper2D/Source/Paper2D/Classes/PaperCharacter.h
--- a/Engine/Plugins/Experimental/Paper2D/Source/Paper2D/Classes/PaperCharacter.h
+++ b/Engine/Plugins/Experimental/Paper2D/Source/Paper2D/Classes/PaperCharacter.h
@@ -6,8 +6,8 @@
 
 // APaperCharacter behaves like ACharacter, but uses a UPaperFlipbookComponent instead of a USkeletalMeshComponent as a visual representation
 // Note: The variable named Mesh will not be set up on this actor!
-UCLASS(MinimalAPI, EarlyAccessPreview)
-class APaperCharacter : public ACharacter
+UCLASS(EarlyAccessPreview)
+class PAPER2D_API APaperCharacter : public ACharacter
 {
 	GENERATED_UCLASS_BODY()
 	
diff --git a/Engine/Plugins/Experimental/Paper2D/Source/Paper2D/Classes/PaperSpriteComponent.h b/Engine/Plugins/Experimental/Paper2D/Source/Paper2D/Classes/PaperSpriteComponent.h
--- a/Engine/Plugins/Experimental/Paper2D/Source/Paper2D/Classes/PaperSpriteComponent.h
+++ b/Engine/Plugins/Experimental/Paper2D/Source/Paper2D/Classes/PaperSpriteComponent.h
@@ -6,8 +6,8 @@
 
 #include "PaperSpriteComponent.generated.h"
 
-UCLASS(MinimalAPI, ShowCategories=(Mobility), EarlyAccessPreview, meta=(BlueprintSpawnableComponent))
-class UPaperSpriteComponent : public UMeshComponent
+UCLASS(ShowCategories=(Mobility), EarlyAccessPreview, meta=(BlueprintSpawnableComponent))
+class PAPER2D_API UPaperSpriteComponent : public UMeshComponent
 {
 	GENERATED_UCLASS_BODY()
 
