Paper2D: Support for TexturePacker atlases, including trim, rotation and pivot settings.

diff --git a/Engine/Plugins/Experimental/Paper2D/Source/Paper2D/Classes/PaperSprite.h b/Engine/Plugins/Experimental/Paper2D/Source/Paper2D/Classes/PaperSprite.h
--- a/Engine/Plugins/Experimental/Paper2D/Source/Paper2D/Classes/PaperSprite.h
+++ b/Engine/Plugins/Experimental/Paper2D/Source/Paper2D/Classes/PaperSprite.h
@@ -58,6 +58,23 @@ protected:
 	// Position within BakedSourceTexture
 	UPROPERTY()
 	FVector2D BakedSourceUV;
+
+
+	// Origin within SourceImage, prior to atlasing
+	UPROPERTY()
+	FVector2D OriginInSourceImage;
+
+	// Dimensions of SourceImage
+	UPROPERTY()
+	FVector2D SourceImageDimension;
+
+	// This texture is trimmed, consider the values above
+	UPROPERTY()
+	bool bTrimmed;
+
+	// This texture is rotated in the atlas
+	UPROPERTY()
+	bool bRotated;
 #endif
 
 	// The source texture that the sprite comes from
@@ -152,6 +169,7 @@ public:
 	FVector ConvertTextureSpaceToPivotSpace(FVector Input) const;
 	FVector ConvertPivotSpaceToTextureSpace(FVector Input) const;
 	FVector2D ConvertWorldSpaceToTextureSpace(const FVector& WorldPoint) const;
+	FVector2D ConvertWorldSpaceDeltaToTextureSpace(const FVector& WorldDelta) const;
 
 	// World space WRT the sprite editor *only*
 	FVector ConvertTextureSpaceToWorldSpace(const FVector2D& SourcePoint) const;
@@ -187,6 +205,10 @@ public:
 	void InitializeSprite(UTexture2D* Texture, const FVector2D& Offset, const FVector2D& Dimension);
 	void InitializeSprite(UTexture2D* Texture, const FVector2D& Offset, const FVector2D& Dimension, float InPixelsPerUnit);
 
+	void SetTrim(bool bTrimmed, const FVector2D& OriginInSourceImage, const FVector2D& SourceImageDimension);
+	void SetRotated(bool bRotated);
+	void SetPivot(ESpritePivotMode::Type PivotMode, FVector2D CustomTextureSpacePivot);
+
 	FVector2D GetSourceUV() const { return SourceUV; }
 	FVector2D GetSourceSize() const { return SourceDimension; }
 	UTexture2D* GetSourceTexture() const { return SourceTexture; }
diff --git a/Engine/Plugins/Experimental/Paper2D/Source/Paper2D/Private/PaperSprite.cpp b/Engine/Plugins/Experimental/Paper2D/Source/Paper2D/Private/PaperSprite.cpp
--- a/Engine/Plugins/Experimental/Paper2D/Source/Paper2D/Private/PaperSprite.cpp
+++ b/Engine/Plugins/Experimental/Paper2D/Source/Paper2D/Private/PaperSprite.cpp
@@ -265,6 +265,9 @@ UPaperSprite::UPaperSprite(const FPostConstructInitializeProperties& PCIP)
 	CollisionThickness = 10.0f;
 
 	PixelsPerUnrealUnit = 2.56f;
+
+	bTrimmed = false;
+	bRotated = false;
 #endif
 
 	static ConstructorHelpers::FObjectFinder<UMaterialInterface> MaskedMaterialRef(TEXT("/Paper2D/MaskedUnlitSpriteMaterial.MaskedUnlitSpriteMaterial"));
@@ -1043,20 +1046,48 @@ void UPaperSprite::InitializeSprite(UTexture2D* Texture, const FVector2D& Offset
 	InitializeSprite(Texture, Offset, Dimension, GetDefault<UPaperRuntimeSettings>()->DefaultPixelsPerUnrealUnit);
 }
 
+void UPaperSprite::SetTrim(bool bTrimmed, const FVector2D& OriginInSourceImage, const FVector2D& SourceImageDimension)
+{
+	this->bTrimmed = bTrimmed;
+	this->OriginInSourceImage = OriginInSourceImage;
+	this->SourceImageDimension = SourceImageDimension;
+	RebuildRenderData();
+	RebuildCollisionData();
+}
+
+void UPaperSprite::SetRotated(bool bRotated)
+{
+	this->bRotated = bRotated;
+	RebuildRenderData();
+	RebuildCollisionData();
+}
+
+void UPaperSprite::SetPivot(ESpritePivotMode::Type PivotMode, FVector2D CustomTextureSpacePivot)
+{
+	this->PivotMode = PivotMode;
+	this->CustomPivotPoint = CustomTextureSpacePivot;
+}
+
 FVector2D UPaperSprite::ConvertTextureSpaceToPivotSpace(FVector2D Input) const
 {
 	const FVector2D Pivot = GetPivotPosition();
 
 	const float X = Input.X - Pivot.X;
 	const float Y = -Input.Y + Pivot.Y;
-	
-	return FVector2D(X, Y);
+
+	return bRotated ? FVector2D(-Y, X) : FVector2D(X, Y);
 }
 
 FVector2D UPaperSprite::ConvertPivotSpaceToTextureSpace(FVector2D Input) const
 {
 	const FVector2D Pivot = GetPivotPosition();
 
+	if (bRotated)
+	{
+		Swap(Input.X, Input.Y);
+		Input.Y = -Input.Y;
+	}
+
 	const float X = Input.X + Pivot.X;
 	const float Y = -Input.Y + Pivot.Y;
 
@@ -1087,63 +1118,113 @@ FVector UPaperSprite::ConvertTextureSpaceToWorldSpace(const FVector2D& SourcePoi
 {
 	const float UnitsPerPixel = GetUnrealUnitsPerPixel();
 
-	const FVector2D SourcePointInUU = SourcePoint * UnitsPerPixel;
-
-	const FVector2D SourceDimsInPixels = (SourceTexture != NULL) ? FVector2D(SourceTexture->GetSurfaceWidth(), SourceTexture->GetSurfaceHeight()) : FVector2D::ZeroVector;
-	const FVector2D SourceDimsInUU = SourceDimsInPixels * UnitsPerPixel;
-
-	return (PaperAxisX * SourcePointInUU.X) + (PaperAxisY * (SourceDimsInUU.Y - SourcePointInUU.Y));
+	const FVector2D SourcePointInUU = ConvertTextureSpaceToPivotSpace(SourcePoint) * UnitsPerPixel;
+	return (PaperAxisX * SourcePointInUU.X) + (PaperAxisY * SourcePointInUU.Y);
 }
 
 FVector2D UPaperSprite::ConvertWorldSpaceToTextureSpace(const FVector& WorldPoint) const
 {
-	const FVector2D SourceDims = (SourceTexture != NULL) ? FVector2D(SourceTexture->GetSurfaceWidth(), SourceTexture->GetSurfaceHeight()) : FVector2D::ZeroVector;
 	const FVector ProjectionX = WorldPoint.ProjectOnTo(PaperAxisX);
 	const FVector ProjectionY = WorldPoint.ProjectOnTo(PaperAxisY);
 
 	const float XValue = FMath::Sign(ProjectionX | PaperAxisX) * ProjectionX.Size() * PixelsPerUnrealUnit;
 	const float YValue = FMath::Sign(ProjectionY | PaperAxisY) * ProjectionY.Size() * PixelsPerUnrealUnit;
 
-	return FVector2D(XValue, SourceDims.Y - YValue);
+	return ConvertPivotSpaceToTextureSpace(FVector2D(XValue, YValue));
 }
 
-FTransform UPaperSprite::GetPivotToWorld() const
+FVector2D UPaperSprite::ConvertWorldSpaceDeltaToTextureSpace(const FVector& WorldSpaceDelta) const
 {
-	const FVector2D Pivot = GetPivotPosition();
-	const FVector2D SourceDims = (SourceTexture != NULL) ? FVector2D(SourceTexture->GetSurfaceWidth(), SourceTexture->GetSurfaceHeight()) : FVector2D::ZeroVector;
+	const FVector ProjectionX = WorldSpaceDelta.ProjectOnTo(PaperAxisX);
+	const FVector ProjectionY = WorldSpaceDelta.ProjectOnTo(PaperAxisY);
+
+	float XValue = FMath::Sign(ProjectionX | PaperAxisX) * ProjectionX.Size() * PixelsPerUnrealUnit;
+	float YValue = FMath::Sign(ProjectionY | PaperAxisY) * ProjectionY.Size() * PixelsPerUnrealUnit;
 
-	const FVector Translation(((Pivot.X * PaperAxisX) + ((SourceDims.Y - Pivot.Y) * PaperAxisY)) * GetUnrealUnitsPerPixel());
+	// Undo pivot space rotation, ignoring pivot position
+	if (bRotated)
+	{
+		Swap(XValue, YValue);
+		XValue = -XValue;
+	}
+
+	return FVector2D(XValue, YValue);
+}
+
+FTransform UPaperSprite::GetPivotToWorld() const
+{
+	const FVector Translation(0, 0, 0);
 	return FTransform(Translation);
 }
 
 FVector2D UPaperSprite::GetRawPivotPosition() const
 {
-	switch (PivotMode)
+	FVector2D TopLeftUV = SourceUV;
+	FVector2D Dimension = SourceDimension;
+	
+	if (bTrimmed)
 	{
-	case ESpritePivotMode::Top_Left:
-		return SourceUV;
-	case ESpritePivotMode::Top_Center:
-		return FVector2D(SourceUV.X + SourceDimension.X * 0.5f, SourceUV.Y);
-	case ESpritePivotMode::Top_Right:
-		return FVector2D(SourceUV.X + SourceDimension.X, SourceUV.Y);
-	case ESpritePivotMode::Center_Left:
-		return FVector2D(SourceUV.X, SourceUV.Y + SourceDimension.Y * 0.5f);
-	case ESpritePivotMode::Center_Center:
-		return FVector2D(SourceUV.X + SourceDimension.X * 0.5f, SourceUV.Y + SourceDimension.Y * 0.5f);
-	case ESpritePivotMode::Center_Right:
-		return FVector2D(SourceUV.X + SourceDimension.X, SourceUV.Y + SourceDimension.Y * 0.5f);
-	case ESpritePivotMode::Bottom_Left:
-		return FVector2D(SourceUV.X, SourceUV.Y + SourceDimension.Y);
-	case ESpritePivotMode::Bottom_Center:
-		return FVector2D(SourceUV.X + SourceDimension.X * 0.5f, SourceUV.Y + SourceDimension.Y);
-	case ESpritePivotMode::Bottom_Right:
-		return SourceUV + SourceDimension;
+		TopLeftUV = SourceUV - OriginInSourceImage;
+		Dimension = SourceImageDimension;
+	}
 
-	default:
-	case ESpritePivotMode::Custom:
-		return CustomPivotPoint;
-		break;
-	};
+	if (bRotated)
+	{
+		switch (PivotMode)
+		{
+		case ESpritePivotMode::Top_Left:
+			return FVector2D(TopLeftUV.X + Dimension.X, TopLeftUV.Y);
+		case ESpritePivotMode::Top_Center:
+			return FVector2D(TopLeftUV.X + Dimension.X, TopLeftUV.Y + Dimension.Y * 0.5f);
+		case ESpritePivotMode::Top_Right:
+			return FVector2D(TopLeftUV.X + Dimension.X, TopLeftUV.Y + Dimension.Y);
+		case ESpritePivotMode::Center_Left:
+			return FVector2D(TopLeftUV.X + Dimension.X * 0.5f, TopLeftUV.Y);
+		case ESpritePivotMode::Center_Center:
+			return FVector2D(TopLeftUV.X + Dimension.X * 0.5f, TopLeftUV.Y + Dimension.Y * 0.5f);
+		case ESpritePivotMode::Center_Right:
+			return FVector2D(TopLeftUV.X + Dimension.X * 0.5f, TopLeftUV.Y + Dimension.Y);
+		case ESpritePivotMode::Bottom_Left:
+			return TopLeftUV;
+		case ESpritePivotMode::Bottom_Center:
+			return FVector2D(TopLeftUV.X, TopLeftUV.Y + Dimension.Y * 0.5f);
+		case ESpritePivotMode::Bottom_Right:
+			return FVector2D(TopLeftUV.X, TopLeftUV.Y + Dimension.Y);
+		default:
+		case ESpritePivotMode::Custom:
+			return CustomPivotPoint;
+			break;
+		};
+	}
+	else
+	{
+		switch (PivotMode)
+		{
+		case ESpritePivotMode::Top_Left:
+			return TopLeftUV;
+		case ESpritePivotMode::Top_Center:
+			return FVector2D(TopLeftUV.X + Dimension.X * 0.5f, TopLeftUV.Y);
+		case ESpritePivotMode::Top_Right:
+			return FVector2D(TopLeftUV.X + Dimension.X, TopLeftUV.Y);
+		case ESpritePivotMode::Center_Left:
+			return FVector2D(TopLeftUV.X, TopLeftUV.Y + Dimension.Y * 0.5f);
+		case ESpritePivotMode::Center_Center:
+			return FVector2D(TopLeftUV.X + Dimension.X * 0.5f, TopLeftUV.Y + Dimension.Y * 0.5f);
+		case ESpritePivotMode::Center_Right:
+			return FVector2D(TopLeftUV.X + Dimension.X, TopLeftUV.Y + Dimension.Y * 0.5f);
+		case ESpritePivotMode::Bottom_Left:
+			return FVector2D(TopLeftUV.X, TopLeftUV.Y + Dimension.Y);
+		case ESpritePivotMode::Bottom_Center:
+			return FVector2D(TopLeftUV.X + Dimension.X * 0.5f, TopLeftUV.Y + Dimension.Y);
+		case ESpritePivotMode::Bottom_Right:
+			return TopLeftUV + Dimension;
+
+		default:
+		case ESpritePivotMode::Custom:
+			return CustomPivotPoint;
+			break;
+		};
+	}
 }
 
 FVector2D UPaperSprite::GetPivotPosition() const
diff --git a/Engine/Plugins/Experimental/Paper2D/Source/Paper2DEditor/Private/PaperJSONHelpers.cpp b/Engine/Plugins/Experimental/Paper2D/Source/Paper2DEditor/Private/PaperJSONHelpers.cpp
--- a/Engine/Plugins/Experimental/Paper2D/Source/Paper2DEditor/Private/PaperJSONHelpers.cpp
+++ b/Engine/Plugins/Experimental/Paper2D/Source/Paper2DEditor/Private/PaperJSONHelpers.cpp
@@ -126,6 +126,28 @@ bool FPaperJSONHelpers::ReadSize(const TSharedPtr<class FJsonObject> Item, const
 	}
 }
 
+bool FPaperJSONHelpers::ReadXY(const TSharedPtr<class FJsonObject> Item, const FString& Key, FVector2D& Out_XY)
+{
+	const TSharedPtr<FJsonObject> Struct = ReadObject(Item, Key);
+	if (Struct.IsValid())
+	{
+		if (ReadFloatNoDefault(Struct, TEXT("x"), /*out*/ Out_XY.X) &&
+			ReadFloatNoDefault(Struct, TEXT("y"), /*out*/ Out_XY.Y))
+		{
+			return true;
+		}
+		else
+		{
+			Out_XY = FVector2D::ZeroVector;
+			return false;
+		}
+	}
+	else
+	{
+		return false;
+	}
+}
+
 bool FPaperJSONHelpers::ReadIntPoint(const TSharedPtr<class FJsonObject> Item, const FString& Key, FIntPoint& Out_XY)
 {
 	const TSharedPtr<FJsonObject> Struct = ReadObject(Item, Key);
diff --git a/Engine/Plugins/Experimental/Paper2D/Source/Paper2DEditor/Private/SpriteEditor/SpriteEditorSelections.h b/Engine/Plugins/Experimental/Paper2D/Source/Paper2DEditor/Private/SpriteEditor/SpriteEditorSelections.h
--- a/Engine/Plugins/Experimental/Paper2D/Source/Paper2DEditor/Private/SpriteEditor/SpriteEditorSelections.h
+++ b/Engine/Plugins/Experimental/Paper2D/Source/Paper2DEditor/Private/SpriteEditor/SpriteEditorSelections.h
@@ -304,7 +304,9 @@ public:
 				TargetVertexIndex = (Vertices.Num() > 0) ? (TargetVertexIndex % Vertices.Num()) : 0;
 				if (Vertices.IsValidIndex(TargetVertexIndex))
 				{
-					Vertices[TargetVertexIndex] += Delta;//@TODO: Should apply the inverse transform (in case scale/rotation is supported later)
+					const FVector WorldSpaceDelta = PaperAxisX * Delta.X + PaperAxisY * Delta.Y;
+					const FVector2D TextureSpaceDelta = Sprite->ConvertWorldSpaceDeltaToTextureSpace(WorldSpaceDelta);
+					Vertices[TargetVertexIndex] += TextureSpaceDelta;
 
 					Geometry.GeometryType = ESpritePolygonMode::FullyCustom;
 				}
@@ -326,12 +328,7 @@ public:
 				TargetVertexIndex = (Vertices.Num() > 0) ? (TargetVertexIndex % Vertices.Num()) : 0;
 				if (Vertices.IsValidIndex(TargetVertexIndex))
 				{
-					UTexture2D* SourceTexture = Sprite->GetSourceTexture();
-					const FVector2D SourceDims = (SourceTexture != NULL) ? FVector2D(SourceTexture->GetSurfaceWidth(), SourceTexture->GetSurfaceHeight()) : FVector2D::ZeroVector;
-
-					const FVector2D Result2D = Vertices[TargetVertexIndex];
-					const FVector PixelSpaceResult = (Result2D.X * PaperAxisX) + ((SourceDims.Y - Result2D.Y) * PaperAxisY);
-					Result = PixelSpaceResult * Sprite->GetUnrealUnitsPerPixel();
+					Result = Sprite->ConvertTextureSpaceToWorldSpace(Vertices[TargetVertexIndex]);
 				}
 			}
 		}
diff --git a/Engine/Plugins/Experimental/Paper2D/Source/Paper2DEditor/Private/SpriteEditor/SpriteEditorViewportClient.cpp b/Engine/Plugins/Experimental/Paper2D/Source/Paper2DEditor/Private/SpriteEditor/SpriteEditorViewportClient.cpp
--- a/Engine/Plugins/Experimental/Paper2D/Source/Paper2DEditor/Private/SpriteEditor/SpriteEditorViewportClient.cpp
+++ b/Engine/Plugins/Experimental/Paper2D/Source/Paper2DEditor/Private/SpriteEditor/SpriteEditorViewportClient.cpp
@@ -126,6 +126,21 @@ FVector FSpriteEditorViewportClient::TextureSpaceToWorldSpace(const FVector2D& S
 	return Sprite->ConvertTextureSpaceToWorldSpace(SourcePoint);
 }
 
+FVector2D FSpriteEditorViewportClient::SourceTextureSpaceToScreenSpace(const FSceneView& View, const FVector2D& SourcePoint) const
+{
+	const FVector WorldSpacePoint = SourceTextureSpaceToWorldSpace(SourcePoint);
+
+	FVector2D PixelLocation;
+	View.WorldToPixel(WorldSpacePoint, /*out*/ PixelLocation);
+	return PixelLocation;
+}
+
+FVector FSpriteEditorViewportClient::SourceTextureSpaceToWorldSpace(const FVector2D& SourcePoint) const
+{
+	UPaperSprite* Sprite = SourceTextureViewComponent->GetSprite();
+	return Sprite->ConvertTextureSpaceToWorldSpace(SourcePoint);
+}
+
 void FSpriteEditorViewportClient::DrawTriangleList(FViewport& InViewport, FSceneView& View, FCanvas& Canvas, const TArray<FVector2D>& Triangles)
 {
 	// Draw the collision data (non-interactive)
@@ -174,10 +189,10 @@ void FSpriteEditorViewportClient::DrawSourceRegion(FViewport& InViewport, FScene
     const bool bDrawCornerHitProxies = true;
 
 	FVector2D BoundsVertices[4];
-	BoundsVertices[0] = TextureSpaceToScreenSpace(View, Sprite->SourceUV);
-	BoundsVertices[1] = TextureSpaceToScreenSpace(View, Sprite->SourceUV + FVector2D(Sprite->SourceDimension.X, 0));
-	BoundsVertices[2] = TextureSpaceToScreenSpace(View, Sprite->SourceUV + FVector2D(Sprite->SourceDimension.X, Sprite->SourceDimension.Y));
-	BoundsVertices[3] = TextureSpaceToScreenSpace(View, Sprite->SourceUV + FVector2D(0, Sprite->SourceDimension.Y));
+	BoundsVertices[0] = SourceTextureSpaceToScreenSpace(View, Sprite->SourceUV);
+	BoundsVertices[1] = SourceTextureSpaceToScreenSpace(View, Sprite->SourceUV + FVector2D(Sprite->SourceDimension.X, 0));
+	BoundsVertices[2] = SourceTextureSpaceToScreenSpace(View, Sprite->SourceUV + FVector2D(Sprite->SourceDimension.X, Sprite->SourceDimension.Y));
+	BoundsVertices[3] = SourceTextureSpaceToScreenSpace(View, Sprite->SourceUV + FVector2D(0, Sprite->SourceDimension.Y));
 	for (int32 VertexIndex = 0; VertexIndex < 4; ++VertexIndex)
 	{
 		const int32 NextVertexIndex = (VertexIndex + 1) % 4;
@@ -644,7 +659,9 @@ void FSpriteEditorViewportClient::Draw(const FSceneView* View, FPrimitiveDrawInt
 {
 	FEditorViewportClient::Draw(View, PDI);
 
-	if (bShowPivot)
+	// We don't draw the pivot when showing the source region
+	// The pivot may be outside the actual texture bounds there
+	if (bShowPivot && !bShowSourceTexture && CurrentMode != ESpriteEditorMode::EditSourceRegionMode)
 	{
 		FUnrealEdUtils::DrawWidget(View, PDI, RenderSpriteComponent->ComponentToWorld.ToMatrixWithScale(), 0, 0, EAxisList::XZ, EWidgetMovementMode::WMM_Translate);
 	}
@@ -681,6 +698,13 @@ void FSpriteEditorViewportClient::Tick(float DeltaSeconds)
 			SourceTextureViewComponent->SetVisibility(bSourceTextureViewComponentVisibility);
 		}
 
+		bool bRenderTextureViewComponentVisibility = !bSourceTextureViewComponentVisibility;
+		if (bRenderTextureViewComponentVisibility != RenderSpriteComponent->IsVisible())
+		{
+			bDeferZoomToSprite = true;
+			RenderSpriteComponent->SetVisibility(bRenderTextureViewComponentVisibility);
+		}
+
 		// Zoom in on the sprite
 		//@TODO: Fix this properly so it doesn't need to be deferred, or wait for the viewport to initialize
 		FIntPoint Size = Viewport->GetSizeXY();
@@ -757,7 +781,7 @@ void FSpriteEditorViewportClient::ProcessClick(FSceneView& View, HHitProxy* HitP
 		{
 			FVector4 WorldPoint = View.PixelToWorld(HitX, HitY, 0);
 			UPaperSprite* Sprite = GetSpriteBeingEdited();
-			FVector2D TexturePoint = Sprite->ConvertWorldSpaceToTextureSpace(WorldPoint);
+			FVector2D TexturePoint = SourceTextureViewComponent->GetSprite()->ConvertWorldSpaceToTextureSpace(WorldPoint);
 			Sprite->ExtractSourceRegionFromTexturePoint(TexturePoint);
 		}
 
diff --git a/Engine/Plugins/Experimental/Paper2D/Source/Paper2DEditor/Private/SpriteEditor/SpriteEditorViewportClient.h b/Engine/Plugins/Experimental/Paper2D/Source/Paper2DEditor/Private/SpriteEditor/SpriteEditorViewportClient.h
--- a/Engine/Plugins/Experimental/Paper2D/Source/Paper2DEditor/Private/SpriteEditor/SpriteEditorViewportClient.h
+++ b/Engine/Plugins/Experimental/Paper2D/Source/Paper2DEditor/Private/SpriteEditor/SpriteEditorViewportClient.h
@@ -164,6 +164,10 @@ private:
 	FVector2D TextureSpaceToScreenSpace(const FSceneView& View, const FVector2D& SourcePoint) const;
 	FVector TextureSpaceToWorldSpace(const FVector2D& SourcePoint) const;
 
+	// Position relative to source texture (ignoring rotation and other transformations applied to extract the sprite)
+	FVector2D SourceTextureSpaceToScreenSpace(const FSceneView& View, const FVector2D& SourcePoint) const;
+	FVector SourceTextureSpaceToWorldSpace(const FVector2D& SourcePoint) const;
+
 	void DrawTriangleList(FViewport& InViewport, FSceneView& View, FCanvas& Canvas, const TArray<FVector2D>& Triangles);
 	void DrawBoundsAsText(FViewport& InViewport, FSceneView& View, FCanvas& Canvas, int32& YPos);
 	void DrawGeometry(FViewport& InViewport, FSceneView& View, FCanvas& Canvas, const FSpritePolygonCollection& Geometry, const FLinearColor& GeometryVertexColor, bool bIsRenderGeometry);
diff --git a/Engine/Plugins/Experimental/Paper2D/Source/Paper2DEditor/Public/PaperJSONHelpers.h b/Engine/Plugins/Experimental/Paper2D/Source/Paper2DEditor/Public/PaperJSONHelpers.h
--- a/Engine/Plugins/Experimental/Paper2D/Source/Paper2DEditor/Public/PaperJSONHelpers.h
+++ b/Engine/Plugins/Experimental/Paper2D/Source/Paper2DEditor/Public/PaperJSONHelpers.h
@@ -32,6 +32,9 @@ public:
 	// Returns true if the object named Key is a struct containing two floats (w,h), populating WH with the values)
 	static bool ReadSize(const TSharedPtr<class FJsonObject> Item, const FString& Key, FVector2D& Out_WH);
 
+	// Returns true if the object named Key is a struct containing two floats (x,y), populating WH with the values)
+	static bool ReadXY(const TSharedPtr<class FJsonObject> Item, const FString& Key, FVector2D& Out_WH);
+
 	// Returns true if the object named Key is a struct containing two floats (x,y), populating XY with the values)
 	static bool ReadIntPoint(const TSharedPtr<class FJsonObject> Item, const FString& Key, FIntPoint& Out_XY);
 
diff --git a/Engine/Plugins/Experimental/Paper2D/Source/PaperJsonImporter/Private/PaperJsonImporterFactory.cpp b/Engine/Plugins/Experimental/Paper2D/Source/PaperJsonImporter/Private/PaperJsonImporterFactory.cpp
--- a/Engine/Plugins/Experimental/Paper2D/Source/PaperJsonImporter/Private/PaperJsonImporterFactory.cpp
+++ b/Engine/Plugins/Experimental/Paper2D/Source/PaperJsonImporter/Private/PaperJsonImporterFactory.cpp
@@ -24,6 +24,8 @@ struct FSpriteFrame
 
 	FVector2D ImageSourceSize;
 
+	FVector2D Pivot;
+
 	bool bTrimmed;
 	bool bRotated;
 };
@@ -42,6 +44,7 @@ UPaperJsonImporterFactory::UPaperJsonImporterFactory(const FPostConstructInitial
 	bText = true;
 
 	Formats.Add(TEXT("json;Spritesheet JSON file"));
+	Formats.Add(TEXT("paper2dsprites;Spritesheet JSON file"));
 }
 
 FText UPaperJsonImporterFactory::GetToolTip() const
@@ -57,18 +60,12 @@ static bool ParseFrame(TSharedPtr<FJsonObject> &FrameData, FSpriteFrame &OutFram
 	//   "rotated": false,
 	//   "trimmed": true,
 	//   "spriteSourceSize": {"x":0,"y":11,"w":216,"h":240},
-	//   "sourceSize": {"w":216,"h":240}
+	//   "sourceSize": {"w":216,"h":240},	
+	//   "pivot": {"x":0.5,"y":0.5}			[optional]
 
 	OutFrame.bRotated = FPaperJSONHelpers::ReadBoolean(FrameData, TEXT("rotated"), false);
 	OutFrame.bTrimmed = FPaperJSONHelpers::ReadBoolean(FrameData, TEXT("trimmed"), false);
 
-	// Put this warning back in until rotation / trim support is implemented
-	if (OutFrame.bRotated || !OutFrame.bTrimmed)
-	{
-		//@TODO: Should learn how to support this combo
-		UE_LOG(LogPaperJsonImporter, Warning, TEXT("Frame %s is either rotated or not trimmed, which may not be handled correctly"), *OutFrame.FrameName.ToString());
-	}
-
 	bReadFrameSuccessfully = bReadFrameSuccessfully && FPaperJSONHelpers::ReadRectangle(FrameData, "frame", /*out*/ OutFrame.SpritePosInSheet, /*out*/ OutFrame.SpriteSizeInSheet);
 
 	if (OutFrame.bTrimmed)
@@ -83,12 +80,26 @@ static bool ParseFrame(TSharedPtr<FJsonObject> &FrameData, FSpriteFrame &OutFram
 		OutFrame.ImageSourceSize = OutFrame.SpriteSizeInSheet;
 	}
 
+	if (!FPaperJSONHelpers::ReadXY(FrameData, "pivot", /*out*/ OutFrame.Pivot))
+	{
+		OutFrame.Pivot.X = 0.5f;
+		OutFrame.Pivot.Y = 0.5f;
+	}
+
 	// A few more prerequisites to sort out before rotation can be enabled
-	//if (OutFrame.bRotated)
-	//{
-	//	// We rotate the sprite 90 deg CCW
-	//	Swap(OutFrame.SpriteSizeInSheet.X, OutFrame.SpriteSizeInSheet.Y);
-	//}
+	if (OutFrame.bRotated)
+	{
+		// Sprite Source Pos is from top left, our pivot when rotated is bottom left
+		OutFrame.SpriteSourcePos.Y = OutFrame.ImageSourceSize.Y - OutFrame.SpriteSourcePos.Y - OutFrame.SpriteSizeInSheet.Y;
+
+		// We rotate final sprite geometry 90 deg CCW to fixup
+		// These need to be swapped to be valid in texture space.
+		Swap(OutFrame.SpriteSizeInSheet.X, OutFrame.SpriteSizeInSheet.Y);
+		Swap(OutFrame.ImageSourceSize.X, OutFrame.ImageSourceSize.Y);
+		
+		Swap(OutFrame.SpriteSourcePos.X, OutFrame.SpriteSourcePos.Y);
+		Swap(OutFrame.SpriteSourceSize.X, OutFrame.SpriteSourceSize.Y);
+	}
 
 	return bReadFrameSuccessfully;
 }
@@ -206,6 +217,51 @@ static bool ParseFramesFromSpriteArray(const TArray<TSharedPtr<FJsonValue>>& Arr
 	return bLoadedSuccessfully;
 }
 
+static ESpritePivotMode::Type GetBestPivotType(FVector2D JsonPivot)
+{
+	// Not assuming layout of ESpritePivotMode
+	if (JsonPivot.X == 0 && JsonPivot.Y == 0)
+	{
+		return ESpritePivotMode::Top_Left;
+	}
+	else if (JsonPivot.X == 0.5f && JsonPivot.Y == 0)
+	{
+		return ESpritePivotMode::Top_Center;
+	}
+	else if (JsonPivot.X == 1.0f && JsonPivot.Y == 0)
+	{
+		return ESpritePivotMode::Top_Right;
+	}
+	else if (JsonPivot.X == 0 && JsonPivot.Y == 0.5f)
+	{
+		return ESpritePivotMode::Center_Left;
+	}
+	else if (JsonPivot.X == 0.5f && JsonPivot.Y == 0.5f)
+	{
+		return ESpritePivotMode::Center_Center;
+	}
+	else if (JsonPivot.X == 1.0f && JsonPivot.Y == 0.5f)
+	{
+		return ESpritePivotMode::Center_Right;
+	}
+	else if (JsonPivot.X == 0 && JsonPivot.Y == 1.0f)
+	{
+		return ESpritePivotMode::Bottom_Left;
+	}
+	else if (JsonPivot.X == 0.5f && JsonPivot.Y == 1.0f)
+	{
+		return ESpritePivotMode::Bottom_Center;
+	}
+	else if (JsonPivot.X == 1.0f && JsonPivot.Y == 1.0f)
+	{
+		return ESpritePivotMode::Bottom_Right;
+	}
+	else
+	{
+		return ESpritePivotMode::Custom;
+	}
+}
+
 UObject* UPaperJsonImporterFactory::FactoryCreateText(UClass* InClass, UObject* InParent, FName InName, EObjectFlags Flags, UObject* Context, const TCHAR* Type, const TCHAR*& Buffer, const TCHAR* BufferEnd, FFeedbackContext* Warn)
 {
 	Flags |= RF_Transactional;
@@ -388,12 +444,27 @@ UObject* UPaperJsonImporterFactory::FactoryCreateText(UClass* InClass, UObject*
 
 					TargetSprite->Modify();
 
-					// TargetSprite->InitializeSprite(ImportedTexture, Frame.SpritePosInSheet, Frame.SpriteSizeInSheet, GetDefault<UPaperRuntimeSettings>()->DefaultPixelsPerUnrealUnit, Frame.bRotated);
 					TargetSprite->InitializeSprite(ImportedTexture, Frame.SpritePosInSheet, Frame.SpriteSizeInSheet, GetDefault<UPaperRuntimeSettings>()->DefaultPixelsPerUnrealUnit);
 
-					//@TODO: Need to support pivot behavior - Total guess at pivot behavior
-					//const FVector2D SizeDifference = Frame.SpriteSourceSize - Frame.SpriteSizeInSheet;
-					//TargetSprite->SpriteData.Destination = FVector(SizeDifference.X * -0.5f, 0.0f, SizeDifference.Y * -0.5f);
+					if (Frame.bRotated)
+					{
+						TargetSprite->SetRotated(true);
+					}
+
+					if (Frame.bTrimmed)
+					{
+						TargetSprite->SetTrim(Frame.bTrimmed, Frame.SpriteSourcePos, Frame.ImageSourceSize);
+					}
+
+					// Set up pivot on object based on Texture Packer json
+					ESpritePivotMode::Type PivotType = GetBestPivotType(Frame.Pivot);
+					FVector2D TextureSpacePivotPoint = FVector2D::ZeroVector;
+					if (PivotType == ESpritePivotMode::Custom)
+					{
+						TextureSpacePivotPoint.X = Frame.SpritePosInSheet.X - Frame.SpriteSourcePos.X + Frame.ImageSourceSize.X * Frame.Pivot.X;
+						TextureSpacePivotPoint.Y = Frame.SpritePosInSheet.Y - Frame.SpriteSourcePos.Y + Frame.ImageSourceSize.Y * Frame.Pivot.Y;
+					}
+					TargetSprite->SetPivot(PivotType, TextureSpacePivotPoint);
 
 					// Create the entry in the animation
 					FPaperFlipbookKeyFrame* DestFrame = new (EditLock.KeyFrames) FPaperFlipbookKeyFrame();
