Paper2D: Disabled "Src Tex" toolbar button until we decide what to do with it

diff --git a/Engine/Plugins/Experimental/Paper2D/Source/Paper2DEditor/Private/SpriteEditor/SpriteEditor.cpp b/Engine/Plugins/Experimental/Paper2D/Source/Paper2DEditor/Private/SpriteEditor/SpriteEditor.cpp
--- a/Engine/Plugins/Experimental/Paper2D/Source/Paper2DEditor/Private/SpriteEditor/SpriteEditor.cpp
+++ b/Engine/Plugins/Experimental/Paper2D/Source/Paper2DEditor/Private/SpriteEditor/SpriteEditor.cpp
@@ -504,11 +504,12 @@ void FSpriteEditor::ExtendToolbar()
 // 			}
 // 			ToolbarBuilder.EndSection();
 
-			ToolbarBuilder.BeginSection("Command");
-			{
-				ToolbarBuilder.AddToolBarButton(FSpriteEditorCommands::Get().SetShowSourceTexture);
-			}
-			ToolbarBuilder.EndSection();
+			// Temporarily disabled until we decide what to do with it. Its not very useful as it is right now.
+//			ToolbarBuilder.BeginSection("Command");
+//			{
+//				ToolbarBuilder.AddToolBarButton(FSpriteEditorCommands::Get().SetShowSourceTexture);
+//			}
+//			ToolbarBuilder.EndSection();
 
 			ToolbarBuilder.BeginSection("Tools");
 			{
