Paper2D: Create flipbook from selected sprites.

diff --git a/Engine/Plugins/Experimental/Paper2D/Source/Paper2DEditor/Classes/PaperFlipbookFactory.h b/Engine/Plugins/Experimental/Paper2D/Source/Paper2DEditor/Classes/PaperFlipbookFactory.h
--- a/Engine/Plugins/Experimental/Paper2D/Source/Paper2DEditor/Classes/PaperFlipbookFactory.h
+++ b/Engine/Plugins/Experimental/Paper2D/Source/Paper2DEditor/Classes/PaperFlipbookFactory.h
@@ -11,6 +11,8 @@ class UPaperFlipbookFactory : public UFactory
 {
 	GENERATED_UCLASS_BODY()
 
+	TArray<FPaperFlipbookKeyFrame> KeyFrames;
+
 	// UFactory interface
 	virtual UObject* FactoryCreateNew(UClass* Class, UObject* InParent, FName Name, EObjectFlags Flags, UObject* Context, FFeedbackContext* Warn) override;
 	// End of UFactory interface
diff --git a/Engine/Plugins/Experimental/Paper2D/Source/Paper2DEditor/Private/PaperFlipbookFactory.cpp b/Engine/Plugins/Experimental/Paper2D/Source/Paper2DEditor/Private/PaperFlipbookFactory.cpp
--- a/Engine/Plugins/Experimental/Paper2D/Source/Paper2DEditor/Private/PaperFlipbookFactory.cpp
+++ b/Engine/Plugins/Experimental/Paper2D/Source/Paper2DEditor/Private/PaperFlipbookFactory.cpp
@@ -18,7 +18,10 @@ UPaperFlipbookFactory::UPaperFlipbookFactory(const class FPostConstructInitializ
 UObject* UPaperFlipbookFactory::FactoryCreateNew(UClass* Class, UObject* InParent, FName Name, EObjectFlags Flags, UObject* Context, FFeedbackContext* Warn)
 {
 	UPaperFlipbook* NewFlipbook = ConstructObject<UPaperFlipbook>(Class, InParent, Name, Flags | RF_Transactional);
-
+	{
+		FScopedFlipbookMutator EditLock(NewFlipbook);
+		EditLock.KeyFrames = KeyFrames;
+	}
 	return NewFlipbook;
 }
 
diff --git a/Engine/Plugins/Experimental/Paper2D/Source/Paper2DEditor/Private/SpriteAssetTypeActions.cpp b/Engine/Plugins/Experimental/Paper2D/Source/Paper2DEditor/Private/SpriteAssetTypeActions.cpp
--- a/Engine/Plugins/Experimental/Paper2D/Source/Paper2DEditor/Private/SpriteAssetTypeActions.cpp
+++ b/Engine/Plugins/Experimental/Paper2D/Source/Paper2DEditor/Private/SpriteAssetTypeActions.cpp
@@ -4,6 +4,8 @@
 
 #include "SpriteAssetTypeActions.h"
 #include "SpriteEditor/SpriteEditor.h"
+#include "AssetToolsModule.h"
+#include "ContentBrowserModule.h"
 
 #define LOCTEXT_NAMESPACE "AssetTypeActions"
 
@@ -44,6 +46,114 @@ uint32 FSpriteAssetTypeActions::GetCategories()
 	return EAssetTypeCategories::Misc;
 }
 
+void FSpriteAssetTypeActions::GetActions(const TArray<UObject*>& InObjects, FMenuBuilder& MenuBuilder)
+{
+	auto Sprites = GetTypedWeakObjectPtrs<UPaperSprite>(InObjects);
+
+	MenuBuilder.AddMenuEntry(
+		LOCTEXT("Sprite_CreateFlipbook", "Create Flipbook"),
+		LOCTEXT("Sprite_CreateFlipbookTooltip", "Creates a flipbook from the selected sprites."),
+		FSlateIcon(),
+		FUIAction(
+		FExecuteAction::CreateSP(this, &FSpriteAssetTypeActions::ExecuteCreateFlipbook, Sprites),
+		FCanExecuteAction()
+		)
+		);
+}
+
 //////////////////////////////////////////////////////////////////////////
 
-#undef LOCTEXT_NAMESPACE
\ No newline at end of file
+
+void FSpriteAssetTypeActions::ExecuteCreateFlipbook(TArray<TWeakObjectPtr<UPaperSprite>> Objects)
+{
+	FAssetToolsModule& AssetToolsModule = FModuleManager::GetModuleChecked<FAssetToolsModule>("AssetTools");
+	FContentBrowserModule& ContentBrowserModule = FModuleManager::LoadModuleChecked<FContentBrowserModule>("ContentBrowser");
+	TArray<UPaperSprite*> Sprites;
+
+	for (auto ObjIt = Objects.CreateConstIterator(); ObjIt; ++ObjIt)
+	{
+		UPaperSprite *Object = (*ObjIt).Get();
+		if (Object && Object->IsValidLowLevel())
+		{
+			Sprites.Add(Object);
+		}
+	}
+
+	// Attempt to natural sort if possible
+	// Cache results if it becomes a bottleneck
+	struct FSpriteSortPredicate
+	{
+		FSpriteSortPredicate() {}
+
+		void ExtractNumber(const FString& String, FString& BareString, int& Number) const
+		{
+			int SplitPoint = 0;
+			for (int i = 0; i < String.Len(); ++i)
+			{
+				const TCHAR CurrentChar = String[i];
+				if (CurrentChar >= TEXT('0') && CurrentChar <= TEXT('9'))
+				{
+					SplitPoint = i;
+					break;
+				}
+			}
+			if (SplitPoint > 0)
+			{
+				FString NumberString = String.Mid(SplitPoint);
+				if (NumberString.IsNumeric())
+				{
+					// Valid string with number
+					BareString = String.Left(SplitPoint);
+					Number = FCString::Atoi(*NumberString);
+					return;
+				}
+			}
+		
+			BareString = String;
+			Number = -1;
+		}
+
+		// Sort predicate operator
+		bool operator()(UPaperSprite& LHS, UPaperSprite& RHS) const
+		{
+			FString LeftString;
+			int LeftNumber;
+			ExtractNumber(LHS.GetName(), /*out*/LeftString, /*out*/LeftNumber);
+
+			FString RightString;
+			int RightNumber;
+			ExtractNumber(RHS.GetName(), /*out*/RightString, /*out*/RightNumber);
+
+			return (LeftString == RightString) ? LeftNumber < RightNumber : LeftString < RightString;
+		}
+	};
+	Sprites.Sort(FSpriteSortPredicate());
+	
+	// Create the flipbook
+	if (Sprites.Num() > 0)
+	{
+		const FString SpritePathName = Sprites[0]->GetOutermost()->GetPathName();
+		const FString LongPackagePath = FPackageName::GetLongPackagePath(Sprites[0]->GetOutermost()->GetPathName());
+
+		const FString NewFlipBookDefaultPath = LongPackagePath + TEXT("/NewFlipbook");
+		FString DefaultSuffix;
+		FString AssetName;
+		FString PackageName;
+		AssetToolsModule.Get().CreateUniqueAssetName(NewFlipBookDefaultPath, /*out*/ DefaultSuffix, /*out*/ PackageName, /*out*/ AssetName);
+		const FString PackagePath = FPackageName::GetLongPackagePath(PackageName);
+
+		UPaperFlipbookFactory* FlipbookFactory = ConstructObject<UPaperFlipbookFactory>(UPaperFlipbookFactory::StaticClass());
+		for (int32 SpriteIndex = 0; SpriteIndex < Sprites.Num(); ++SpriteIndex)
+		{
+			UPaperSprite* Sprite = Sprites[SpriteIndex];
+			FPaperFlipbookKeyFrame* KeyFrame = new (FlipbookFactory->KeyFrames) FPaperFlipbookKeyFrame();
+			KeyFrame->Sprite = Sprite;
+			KeyFrame->FrameRun = 1;
+		}
+
+		ContentBrowserModule.Get().CreateNewAsset(AssetName, PackagePath, UPaperFlipbook::StaticClass(), FlipbookFactory);
+	}
+}
+
+
+#undef LOCTEXT_NAMESPACE
diff --git a/Engine/Plugins/Experimental/Paper2D/Source/Paper2DEditor/Private/SpriteAssetTypeActions.h b/Engine/Plugins/Experimental/Paper2D/Source/Paper2DEditor/Private/SpriteAssetTypeActions.h
--- a/Engine/Plugins/Experimental/Paper2D/Source/Paper2DEditor/Private/SpriteAssetTypeActions.h
+++ b/Engine/Plugins/Experimental/Paper2D/Source/Paper2DEditor/Private/SpriteAssetTypeActions.h
@@ -11,7 +11,13 @@ public:
 	virtual FText GetName() const override;
 	virtual FColor GetTypeColor() const override;
 	virtual UClass* GetSupportedClass() const override;
+	virtual bool HasActions(const TArray<UObject*>& InObjects) const override { return true; }
+	virtual void GetActions(const TArray<UObject*>& InObjects, FMenuBuilder& MenuBuilder) override;
 	virtual void OpenAssetEditor(const TArray<UObject*>& InObjects, TSharedPtr<class IToolkitHost> EditWithinLevelEditor = TSharedPtr<IToolkitHost>()) override;
 	virtual uint32 GetCategories() override;
 	// End of IAssetTypeActions interface
+
+private:
+	/* Handler for "Create Flipbook" */
+	void ExecuteCreateFlipbook(TArray<TWeakObjectPtr<UPaperSprite>> Objects);
 };
\ No newline at end of file
