Paper2D: Improve "Edit Source Region" performance by deferring rebuilding slow render and collision mesh generation until after user has stopped dragging

diff --git a/Engine/Plugins/Experimental/Paper2D/Source/Paper2DEditor/Private/SpriteEditor/SpriteEditorViewportClient.cpp b/Engine/Plugins/Experimental/Paper2D/Source/Paper2DEditor/Private/SpriteEditor/SpriteEditorViewportClient.cpp
--- a/Engine/Plugins/Experimental/Paper2D/Source/Paper2DEditor/Private/SpriteEditor/SpriteEditorViewportClient.cpp
+++ b/Engine/Plugins/Experimental/Paper2D/Source/Paper2DEditor/Private/SpriteEditor/SpriteEditorViewportClient.cpp
@@ -822,13 +822,16 @@ bool FSpriteEditorViewportClient::InputWidgetDelta(FViewport* Viewport, EAxisLis
 
 		if (SelectionSet.Num() > 0)
 		{
-			UPaperSprite* Sprite = GetSpriteBeingEdited();
-			// Sprite->RebuildCollisionData();
-			// Sprite->RebuildRenderData();
-			Sprite->PostEditChange();
-			bManipulationDirtiedSomething = true;
+			if (!IsInSourceRegionEditMode())
+			{
+				UPaperSprite* Sprite = GetSpriteBeingEdited();
+				// Sprite->RebuildCollisionData();
+				// Sprite->RebuildRenderData();
+				Sprite->PostEditChange();
+				Invalidate();
+			}
 
-			Invalidate();
+			bManipulationDirtiedSomething = true;
 		}
 	}
 
