Paper2D: Make per-sprite sprite sheet import settings editable (source trim and rotation in advanced category) to allow further modification of sprites imported from a sprite sheet
#TTP 342780, integrated from 4.4 branch CL# 2239346

diff --git a/Engine/Plugins/Experimental/Paper2D/Source/Paper2D/Classes/PaperSprite.h b/Engine/Plugins/Experimental/Paper2D/Source/Paper2D/Classes/PaperSprite.h
--- a/Engine/Plugins/Experimental/Paper2D/Source/Paper2D/Classes/PaperSprite.h
+++ b/Engine/Plugins/Experimental/Paper2D/Source/Paper2D/Classes/PaperSprite.h
@@ -59,22 +59,21 @@ protected:
 	UPROPERTY()
 	FVector2D BakedSourceUV;
 
-
 	// Origin within SourceImage, prior to atlasing
-	UPROPERTY()
-	FVector2D OriginInSourceImage;
+	UPROPERTY(Category=Sprite, EditAnywhere, AdvancedDisplay, meta=(EditCondition="bTrimmedInSourceImage"))
+	FVector2D OriginInSourceImageBeforeTrimming;
 
 	// Dimensions of SourceImage
-	UPROPERTY()
-	FVector2D SourceImageDimension;
+	UPROPERTY(Category=Sprite, EditAnywhere, AdvancedDisplay, meta=(EditCondition="bTrimmedInSourceImage"))
+	FVector2D SourceImageDimensionBeforeTrimming;
 
 	// This texture is trimmed, consider the values above
-	UPROPERTY()
-	bool bTrimmed;
+	UPROPERTY(Category=Sprite, EditAnywhere, AdvancedDisplay)
+	bool bTrimmedInSourceImage;
 
 	// This texture is rotated in the atlas
-	UPROPERTY()
-	bool bRotated;
+	UPROPERTY(Category=Sprite, EditAnywhere, AdvancedDisplay)
+	bool bRotatedInSourceImage;
 #endif
 
 	// The source texture that the sprite comes from
diff --git a/Engine/Plugins/Experimental/Paper2D/Source/Paper2D/Private/PaperSprite.cpp b/Engine/Plugins/Experimental/Paper2D/Source/Paper2D/Private/PaperSprite.cpp
--- a/Engine/Plugins/Experimental/Paper2D/Source/Paper2D/Private/PaperSprite.cpp
+++ b/Engine/Plugins/Experimental/Paper2D/Source/Paper2D/Private/PaperSprite.cpp
@@ -266,8 +266,8 @@ UPaperSprite::UPaperSprite(const FPostConstructInitializeProperties& PCIP)
 
 	PixelsPerUnrealUnit = 2.56f;
 
-	bTrimmed = false;
-	bRotated = false;
+	bTrimmedInSourceImage = false;
+	bRotatedInSourceImage = false;
 #endif
 
 	static ConstructorHelpers::FObjectFinder<UMaterialInterface> MaskedMaterialRef(TEXT("/Paper2D/MaskedUnlitSpriteMaterial.MaskedUnlitSpriteMaterial"));
@@ -1048,16 +1048,16 @@ void UPaperSprite::InitializeSprite(UTexture2D* Texture, const FVector2D& Offset
 
 void UPaperSprite::SetTrim(bool bTrimmed, const FVector2D& OriginInSourceImage, const FVector2D& SourceImageDimension)
 {
-	this->bTrimmed = bTrimmed;
-	this->OriginInSourceImage = OriginInSourceImage;
-	this->SourceImageDimension = SourceImageDimension;
+	this->bTrimmedInSourceImage = bTrimmed;
+	this->OriginInSourceImageBeforeTrimming = OriginInSourceImage;
+	this->SourceImageDimensionBeforeTrimming = SourceImageDimension;
 	RebuildRenderData();
 	RebuildCollisionData();
 }
 
 void UPaperSprite::SetRotated(bool bRotated)
 {
-	this->bRotated = bRotated;
+	this->bRotatedInSourceImage = bRotated;
 	RebuildRenderData();
 	RebuildCollisionData();
 }
@@ -1075,14 +1075,14 @@ FVector2D UPaperSprite::ConvertTextureSpaceToPivotSpace(FVector2D Input) const
 	const float X = Input.X - Pivot.X;
 	const float Y = -Input.Y + Pivot.Y;
 
-	return bRotated ? FVector2D(-Y, X) : FVector2D(X, Y);
+	return bRotatedInSourceImage ? FVector2D(-Y, X) : FVector2D(X, Y);
 }
 
 FVector2D UPaperSprite::ConvertPivotSpaceToTextureSpace(FVector2D Input) const
 {
 	const FVector2D Pivot = GetPivotPosition();
 
-	if (bRotated)
+	if (bRotatedInSourceImage)
 	{
 		Swap(Input.X, Input.Y);
 		Input.Y = -Input.Y;
@@ -1142,7 +1142,7 @@ FVector2D UPaperSprite::ConvertWorldSpaceDeltaToTextureSpace(const FVector& Worl
 	float YValue = FMath::Sign(ProjectionY | PaperAxisY) * ProjectionY.Size() * PixelsPerUnrealUnit;
 
 	// Undo pivot space rotation, ignoring pivot position
-	if (bRotated)
+	if (bRotatedInSourceImage)
 	{
 		Swap(XValue, YValue);
 		XValue = -XValue;
@@ -1162,13 +1162,13 @@ FVector2D UPaperSprite::GetRawPivotPosition() const
 	FVector2D TopLeftUV = SourceUV;
 	FVector2D Dimension = SourceDimension;
 	
-	if (bTrimmed)
+	if (bTrimmedInSourceImage)
 	{
-		TopLeftUV = SourceUV - OriginInSourceImage;
-		Dimension = SourceImageDimension;
+		TopLeftUV = SourceUV - OriginInSourceImageBeforeTrimming;
+		Dimension = SourceImageDimensionBeforeTrimming;
 	}
 
-	if (bRotated)
+	if (bRotatedInSourceImage)
 	{
 		switch (PivotMode)
 		{
