Paper2D: Use sprite BakedRenderData to render the thumbnail.
#TTP 342943

diff --git a/Engine/Plugins/Experimental/Paper2D/Source/Paper2DEditor/Private/PaperSpriteThumbnailRenderer.cpp b/Engine/Plugins/Experimental/Paper2D/Source/Paper2DEditor/Private/PaperSpriteThumbnailRenderer.cpp
--- a/Engine/Plugins/Experimental/Paper2D/Source/Paper2DEditor/Private/PaperSpriteThumbnailRenderer.cpp
+++ b/Engine/Plugins/Experimental/Paper2D/Source/Paper2DEditor/Private/PaperSpriteThumbnailRenderer.cpp
@@ -53,47 +53,65 @@ void UPaperSpriteThumbnailRenderer::DrawFrame(class UPaperSprite* Sprite, int32
 		}
 
 		// Draw the sprite itself
-		const float TextureWidth = SourceTexture->GetSurfaceWidth();
-		const float TextureHeight = SourceTexture->GetSurfaceHeight();
+		// Use the baked render data, so we don't have to care about rotations and possibly
+		// other sprites overlapping in source, UV region, etc.
+		const TArray<FVector4> &BakedRenderData = Sprite->BakedRenderData;
+		TArray<FVector2D> CanvasPositions;
+		TArray<FVector2D> CanvasUVs;
 
-		float FinalX = (float)X;
-		float FinalY = (float)Y;
-		float FinalWidth = (float)Width;
-		float FinalHeight = (float)Height;
-		const float DesiredWidth = Sprite->GetSourceSize().X;
-		const float DesiredHeight = Sprite->GetSourceSize().Y;
+		for (int Vertex = 0; Vertex < BakedRenderData.Num(); ++Vertex)
+		{
+			new(CanvasPositions)FVector2D(BakedRenderData[Vertex].X, -BakedRenderData[Vertex].Y);
+			new(CanvasUVs)FVector2D(BakedRenderData[Vertex].Z, BakedRenderData[Vertex].W);
+		}
 
-		if (DesiredWidth > DesiredHeight)
+		FVector2D MinPoint(BIG_NUMBER, BIG_NUMBER);
+		FVector2D MaxPoint(-BIG_NUMBER, -BIG_NUMBER);
+		for (int Vertex = 0; Vertex < CanvasPositions.Num(); ++Vertex)
 		{
-			const float ScaleFactor = Width / DesiredWidth;
-			FinalHeight = ScaleFactor * DesiredHeight;
-			FinalY += (Height - FinalHeight) * 0.5f;
+			MinPoint.X = FMath::Min(MinPoint.X, CanvasPositions[Vertex].X);
+			MinPoint.Y = FMath::Min(MinPoint.Y, CanvasPositions[Vertex].Y);
+			MaxPoint.X = FMath::Max(MaxPoint.X, CanvasPositions[Vertex].X);
+			MaxPoint.Y = FMath::Max(MaxPoint.Y, CanvasPositions[Vertex].Y);
+		}
+
+		float ScaleFactor = 1;
+		float UnscaledWidth = MaxPoint.X - MinPoint.X;
+		float UnscaledHeight = MaxPoint.Y - MinPoint.Y;
+		FVector2D Origin(X + Width / 2, Y + Height / 2);
+		if (UnscaledWidth > 0 && UnscaledHeight > 0 && UnscaledWidth > UnscaledHeight)
+		{ 
+			ScaleFactor = Width / UnscaledWidth;
 		}
 		else
 		{
-			const float ScaleFactor = Height / DesiredHeight;
-			FinalWidth = ScaleFactor * DesiredWidth;
-			FinalX += (Width - FinalWidth) * 0.5f;
+			ScaleFactor = Height / UnscaledHeight;
 		}
 
-		const FLinearColor SpriteColor(FLinearColor::White);
+		// Scale and recenter
+		FVector2D CanvasPositionCenter = (MaxPoint + MinPoint) * 0.5f;
+		for (int Vertex = 0; Vertex < CanvasPositions.Num(); ++Vertex)
+		{
+			CanvasPositions[Vertex] = (CanvasPositions[Vertex] - CanvasPositionCenter) * ScaleFactor + Origin;
+		}
 
-		Canvas->DrawTile(
-			FinalX,
-			FinalY,
-			FinalWidth,
-			FinalHeight,
-			Sprite->GetSourceUV().X / TextureWidth,
-			Sprite->GetSourceUV().Y / TextureHeight,
-			DesiredWidth / TextureWidth,
-			DesiredHeight / TextureHeight,
-			SpriteColor,
-			SourceTexture->Resource,
-			bUseTranslucentBlend);
+		// Draw triangles
+		TArray<FCanvasUVTri> Triangles;
+		const FLinearColor SpriteColor(FLinearColor::White);
+		for (int Vertex = 0; Vertex < CanvasPositions.Num(); Vertex += 3)
+		{
+			FCanvasUVTri *Triangle = new(Triangles)FCanvasUVTri();
+			Triangle->V0_Pos = CanvasPositions[Vertex + 0]; Triangle->V0_UV = CanvasUVs[Vertex + 0]; Triangle->V0_Color = SpriteColor;
+			Triangle->V1_Pos = CanvasPositions[Vertex + 1]; Triangle->V1_UV = CanvasUVs[Vertex + 1]; Triangle->V1_Color = SpriteColor;
+			Triangle->V2_Pos = CanvasPositions[Vertex + 2]; Triangle->V2_UV = CanvasUVs[Vertex + 2]; Triangle->V2_Color = SpriteColor;
+		}
+		FCanvasTriangleItem CanvasTriangle(Triangles, SourceTexture->Resource);
+		CanvasTriangle.BlendMode = bUseTranslucentBlend ? ESimpleElementBlendMode::SE_BLEND_Translucent : ESimpleElementBlendMode::SE_BLEND_Opaque;
+		Canvas->DrawItem(CanvasTriangle);
 	}
 	else
 	{
 		// Fallback for a bogus sprite
 		DrawGrid(X, Y, Width, Height, Canvas);
 	}
-}
\ No newline at end of file
+}
