Paper2D: Create flipbook from selection - sort by last number in name, not first

diff --git a/Engine/Plugins/Experimental/Paper2D/Source/Paper2DEditor/Private/SpriteAssetTypeActions.cpp b/Engine/Plugins/Experimental/Paper2D/Source/Paper2DEditor/Private/SpriteAssetTypeActions.cpp
--- a/Engine/Plugins/Experimental/Paper2D/Source/Paper2DEditor/Private/SpriteAssetTypeActions.cpp
+++ b/Engine/Plugins/Experimental/Paper2D/Source/Paper2DEditor/Private/SpriteAssetTypeActions.cpp
@@ -87,30 +87,30 @@ void FSpriteAssetTypeActions::ExecuteCreateFlipbook(TArray<TWeakObjectPtr<UPaper
 
 		void ExtractNumber(const FString& String, FString& BareString, int& Number) const
 		{
-			int SplitPoint = 0;
-			for (int i = 0; i < String.Len(); ++i)
+			bool bExtracted = false;
+			int LastCharacter = String.Len() - 1;
+			if (LastCharacter >= 0 && FChar::IsDigit(String[LastCharacter]))
 			{
-				const TCHAR CurrentChar = String[i];
-				if (CurrentChar >= TEXT('0') && CurrentChar <= TEXT('9'))
+				while (LastCharacter > 0 && FChar::IsDigit(String[LastCharacter - 1]))
 				{
-					SplitPoint = i;
-					break;
+					LastCharacter--;
 				}
-			}
-			if (SplitPoint > 0)
-			{
-				FString NumberString = String.Mid(SplitPoint);
-				if (NumberString.IsNumeric())
+
+				if (LastCharacter >= 0)
 				{
-					// Valid string with number
-					BareString = String.Left(SplitPoint);
+					int FirstDigit = LastCharacter;
+					FString NumberString = String.Mid(FirstDigit);
+					BareString = String.Left(FirstDigit);
 					Number = FCString::Atoi(*NumberString);
-					return;
+					bExtracted = true;
 				}
 			}
 		
-			BareString = String;
-			Number = -1;
+			if (!bExtracted)
+			{
+				BareString = String;
+				Number = -1;
+			}
 		}
 
 		// Sort predicate operator
@@ -127,6 +127,7 @@ void FSpriteAssetTypeActions::ExecuteCreateFlipbook(TArray<TWeakObjectPtr<UPaper
 			return (LeftString == RightString) ? LeftNumber < RightNumber : LeftString < RightString;
 		}
 	};
+
 	Sprites.Sort(FSpriteSortPredicate());
 	
 	// Create the flipbook
