Paper2D: Allow me to inherit from PaperRenderSceneProxy in a gameplay module (added DLLExports)

diff --git a/Engine/Plugins/Experimental/Paper2D/Source/Paper2D/Private/PaperRenderSceneProxy.h b/Engine/Plugins/Experimental/Paper2D/Source/Paper2D/Private/PaperRenderSceneProxy.h
--- a/Engine/Plugins/Experimental/Paper2D/Source/Paper2D/Private/PaperRenderSceneProxy.h
+++ b/Engine/Plugins/Experimental/Paper2D/Source/Paper2D/Private/PaperRenderSceneProxy.h
@@ -8,7 +8,7 @@
 // FPaperSpriteVertex
 
 /** A Paper2D sprite vertex. */
-struct FPaperSpriteVertex
+struct PAPER2D_API FPaperSpriteVertex
 {
 	FVector Position;
 	FPackedNormal TangentX;
@@ -35,7 +35,7 @@ struct FPaperSpriteVertex
 //////////////////////////////////////////////////////////////////////////
 // FSpriteRenderSection
 
-struct FSpriteRenderSection
+struct PAPER2D_API FSpriteRenderSection
 {
 	int32 VertexOffset;
 	int32 NumVertices;
@@ -82,7 +82,7 @@ struct FSpriteRenderSection
 //////////////////////////////////////////////////////////////////////////
 // FPaperRenderSceneProxy
 
-class FPaperRenderSceneProxy : public FPrimitiveSceneProxy
+class PAPER2D_API FPaperRenderSceneProxy : public FPrimitiveSceneProxy
 {
 public:
 	FPaperRenderSceneProxy(const UPrimitiveComponent* InComponent);
