Paper2D: Source Region edit mode improvments - allow slicing up sprites quickly by ctrl-dragging
- display other sprites from the same sprite sheet, click to edit

diff --git a/Engine/Plugins/Experimental/Paper2D/Source/Paper2DEditor/Private/SpriteEditor/SpriteEditorViewportClient.cpp b/Engine/Plugins/Experimental/Paper2D/Source/Paper2DEditor/Private/SpriteEditor/SpriteEditorViewportClient.cpp
--- a/Engine/Plugins/Experimental/Paper2D/Source/Paper2DEditor/Private/SpriteEditor/SpriteEditorViewportClient.cpp
+++ b/Engine/Plugins/Experimental/Paper2D/Source/Paper2DEditor/Private/SpriteEditor/SpriteEditorViewportClient.cpp
@@ -8,6 +8,10 @@
 #include "ScopedTransaction.h"
 #include "Runtime/Engine/Public/ComponentReregisterContext.h"
 
+#include "AssetToolsModule.h"
+#include "ContentBrowserModule.h"
+#include "AssetRegistryModule.h"
+
 #define LOCTEXT_NAMESPACE "SpriteEditor"
 
 //////////////////////////////////////////////////////////////////////////
@@ -36,6 +40,17 @@ struct HSpriteSelectableObjectHitProxy : public HHitProxy
 };
 IMPLEMENT_HIT_PROXY(HSpriteSelectableObjectHitProxy, HHitProxy);
 
+
+//////////////////////////////////////////////////////////////////////////
+// FRelatedSprite
+
+struct FRelatedSprite
+{
+	FAssetData AssetData;
+	FVector2D SourceUV;
+	FVector2D SourceDimension;
+};
+
 //////////////////////////////////////////////////////////////////////////
 // FSpriteEditorViewportClient
 
@@ -59,9 +74,12 @@ FSpriteEditorViewportClient::FSpriteEditorViewportClient(TWeakPtr<FSpriteEditor>
 	bShowSockets = true;
 	bShowNormals = true;
 	bShowPivot = true;
+	bShowRelatedSprites = true;
 
 	bDeferZoomToSprite = true;
 
+	bIsMarqueeTracking = false;
+
 	EngineShowFlags.DisableAdvancedFeatures();
 	EngineShowFlags.CompositeEditorPrimitives = true;
 
@@ -173,7 +191,30 @@ void FSpriteEditorViewportClient::DrawTriangleList(FViewport& InViewport, FScene
 	}
 }
 
-void FSpriteEditorViewportClient::DrawSourceRegion(FViewport& InViewport, FSceneView& View, FCanvas& Canvas, const FLinearColor& GeometryVertexColor, bool bIsRenderGeometry)
+void FSpriteEditorViewportClient::DrawRelatedSprites(FViewport& InViewport, FSceneView& View, FCanvas& Canvas, const FLinearColor& LineColor)
+{
+	for (int SpriteIndex = 0; SpriteIndex < RelatedSprites.Num(); ++SpriteIndex)
+	{
+		FRelatedSprite& RelatedSprite = RelatedSprites[SpriteIndex];
+		const FVector2D& SourceUV = RelatedSprite.SourceUV;
+		const FVector2D& SourceDimension = RelatedSprite.SourceDimension;
+		FVector2D BoundsVertices[4];
+		BoundsVertices[0] = SourceTextureSpaceToScreenSpace(View, SourceUV);
+		BoundsVertices[1] = SourceTextureSpaceToScreenSpace(View, SourceUV + FVector2D(SourceDimension.X, 0));
+		BoundsVertices[2] = SourceTextureSpaceToScreenSpace(View, SourceUV + FVector2D(SourceDimension.X, SourceDimension.Y));
+		BoundsVertices[3] = SourceTextureSpaceToScreenSpace(View, SourceUV + FVector2D(0, SourceDimension.Y));
+		for (int32 VertexIndex = 0; VertexIndex < 4; ++VertexIndex)
+		{
+			const int32 NextVertexIndex = (VertexIndex + 1) % 4;
+
+			FCanvasLineItem LineItem(BoundsVertices[VertexIndex], BoundsVertices[NextVertexIndex]);
+			LineItem.SetColor(LineColor);
+			Canvas.DrawItem(LineItem);
+		}
+	}
+}
+
+void FSpriteEditorViewportClient::DrawSourceRegion(FViewport& InViewport, FSceneView& View, FCanvas& Canvas, const FLinearColor& GeometryVertexColor)
 {
 	const bool bIsHitTesting = Canvas.IsHitTesting();
 	UPaperSprite* Sprite = GetSpriteBeingEdited();
@@ -244,6 +285,23 @@ void FSpriteEditorViewportClient::DrawSourceRegion(FViewport& InViewport, FScene
 	}
 }
 
+void FSpriteEditorViewportClient::DrawMarquee(FViewport& InViewport, FSceneView& View, FCanvas& Canvas, const FLinearColor& MarqueeColor)
+{
+	FVector2D MarqueeDelta = MarqueeEndPos - MarqueeStartPos;
+	FVector2D MarqueeVertices[4];
+	MarqueeVertices[0] = FVector2D(MarqueeStartPos.X, MarqueeStartPos.Y);
+	MarqueeVertices[1] = MarqueeVertices[0] + FVector2D(MarqueeDelta.X, 0);
+	MarqueeVertices[2] = MarqueeVertices[0] + FVector2D(MarqueeDelta.X, MarqueeDelta.Y);
+	MarqueeVertices[3] = MarqueeVertices[0] + FVector2D(0, MarqueeDelta.Y);
+	for (int32 MarqueeVertexIndex = 0; MarqueeVertexIndex < 4; ++MarqueeVertexIndex)
+	{
+		const int32 NextVertexIndex = (MarqueeVertexIndex + 1) % 4;
+		FCanvasLineItem MarqueeLine(MarqueeVertices[MarqueeVertexIndex], MarqueeVertices[NextVertexIndex]);
+		MarqueeLine.SetColor(MarqueeColor);
+		Canvas.DrawItem(MarqueeLine);
+	}
+}
+
 void FSpriteEditorViewportClient::DrawGeometry(FViewport& InViewport, FSceneView& View, FCanvas& Canvas, const FSpritePolygonCollection& Geometry, const FLinearColor& GeometryVertexColor, bool bIsRenderGeometry)
 {
 	const bool bIsHitTesting = Canvas.IsHitTesting();
@@ -559,7 +617,7 @@ void FSpriteEditorViewportClient::DrawCanvas(FViewport& Viewport, FSceneView& Vi
 	int32 YPos = 42;
 
 	static const FText GeomHelpStr = LOCTEXT("GeomEditHelp", "Select an edge and press Insert to add a vertex.\nSelect an edge or vertex and press Delete to remove it.\n");
-	static const FText SourceRegionHelpStr = LOCTEXT("SourceRegionHelp", "Drag handles to adjust source region\nDouble-click on an image region to select all connected pixels");
+	static const FText SourceRegionHelpStr = LOCTEXT("SourceRegionHelp", "Drag handles to adjust source region\nDouble-click on an image region to select all connected pixels\nHold down Ctrl and drag a rectangle to create a new sprite at that position");
 
 	switch (CurrentMode)
 	{
@@ -639,9 +697,15 @@ void FSpriteEditorViewportClient::DrawCanvas(FViewport& Viewport, FSceneView& Vi
 				YPos += 18;
 			}
 
-			// Draw the custom render geometry
+			if (bShowRelatedSprites)
+			{
+				// Alpha is being ignored in FBatchedElements::AddLine :(
+				const FLinearColor RelatedBoundsColor(0.3f, 0.3f, 0.3f, 0.2f);
+				DrawRelatedSprites(Viewport, View, Canvas, RelatedBoundsColor);
+			}
+
 			const FLinearColor BoundsColor(1.0f, 1.0f, 1.0f, 0.8f);
-			DrawSourceRegion(Viewport, View, Canvas, BoundsColor, true);
+			DrawSourceRegion(Viewport, View, Canvas, BoundsColor);
 		}
 		break;
 	case ESpriteEditorMode::AddSpriteMode:
@@ -649,6 +713,12 @@ void FSpriteEditorViewportClient::DrawCanvas(FViewport& Viewport, FSceneView& Vi
 		break;
 	}
 
+	if (bIsMarqueeTracking)
+	{
+		const FLinearColor MarqueeColor(1.0f, 1.0f, 1.0f, 0.5f);
+		DrawMarquee(Viewport, View, Canvas, MarqueeColor);
+	}
+
 	if (bShowSockets)
 	{
 		DrawSocketNames(Viewport, View, Canvas);
@@ -716,6 +786,13 @@ void FSpriteEditorViewportClient::Tick(float DeltaSeconds)
 		}		
 	}
 
+	if (bIsMarqueeTracking)
+	{
+		int32 HitX = Viewport->GetMouseX();
+		int32 HitY = Viewport->GetMouseY();
+		MarqueeEndPos = FVector2D(HitX, HitY);
+	}
+
 	FPaperEditorViewportClient::Tick(DeltaSeconds);
 
 	if (!GIntraFrameDebuggingGameThread)
@@ -743,6 +820,52 @@ bool FSpriteEditorViewportClient::IsShowMeshEdgesChecked() const
 	return EngineShowFlags.MeshEdges;
 }
 
+// Find all related sprites (not including self)
+void FSpriteEditorViewportClient::UpdateRelatedSpritesList()
+{
+	UPaperSprite* Sprite = GetSpriteBeingEdited();
+	UTexture2D* Texture = Sprite->GetSourceTexture();
+	if (Texture != nullptr)
+	{
+		FARFilter Filter;
+		Filter.ClassNames.Add(UPaperSprite::StaticClass()->GetFName());
+		const FString TextureString = FAssetData(Texture).GetExportTextName();
+		const FName SourceTexturePropName(TEXT("SourceTexture"));
+		Filter.TagsAndValues.Add(SourceTexturePropName, TextureString);
+		FAssetRegistryModule& AssetRegistryModule = FModuleManager::LoadModuleChecked<FAssetRegistryModule>("AssetRegistry");
+		TArray<FAssetData> SpriteAssetData;
+		AssetRegistryModule.Get().GetAssets(Filter, SpriteAssetData);
+
+		FAssetData CurrentAssetData(Sprite);
+
+		RelatedSprites.Empty();
+		for (int i = 0; i < SpriteAssetData.Num(); ++i)
+		{
+			FAssetData& SpriteAsset = SpriteAssetData[i];
+			if (SpriteAsset == Sprite)
+			{
+				continue;
+			}
+
+			const FString* SourceUVString = SpriteAsset.TagsAndValues.Find("SourceUV");
+			const FString* SourceDimensionString = SpriteAsset.TagsAndValues.Find("SourceDimension");
+			FVector2D SourceUV, SourceDimension;
+			if (SourceUVString != nullptr && SourceDimensionString != nullptr)
+			{
+				if (SourceUV.InitFromString(*SourceUVString) && SourceDimension.InitFromString(*SourceDimensionString))
+				{
+					FRelatedSprite RelatedSprite;
+					RelatedSprite.AssetData = SpriteAsset;
+					RelatedSprite.SourceUV = SourceUV;
+					RelatedSprite.SourceDimension = SourceDimension;
+					
+					RelatedSprites.Add(RelatedSprite);
+				}
+			}
+		}
+	}
+}
+
 void FSpriteEditorViewportClient::UpdateMouseDelta()
 {
 	FPaperEditorViewportClient::UpdateMouseDelta();
@@ -751,15 +874,18 @@ void FSpriteEditorViewportClient::UpdateMouseDelta()
 void FSpriteEditorViewportClient::ProcessClick(FSceneView& View, HHitProxy* HitProxy, FKey Key, EInputEvent Event, uint32 HitX, uint32 HitY)
 {
 	const FViewportClick Click(&View, this, Key, Event, HitX, HitY);
+	const bool bIsCtrlKeyDown = Viewport->KeyState(EKeys::LeftControl) || Viewport->KeyState(EKeys::RightControl);
+	bool bHandled = false;
+
 	if (HSpriteSelectableObjectHitProxy* SelectedItemProxy = HitProxyCast<HSpriteSelectableObjectHitProxy>(HitProxy))
 	{
-		const bool bIsCtrlKeyDown = Viewport->KeyState(EKeys::LeftControl) || Viewport->KeyState(EKeys::RightControl);
 		if (!bIsCtrlKeyDown)
 		{
 			ClearSelectionSet();
 		}
 
 		SelectionSet.Add(SelectedItemProxy->Data);
+		bHandled = true;
 	}
 // 	else if (HWidgetUtilProxy* PivotProxy = HitProxyCast<HWidgetUtilProxy>(HitProxy))
 // 	{
@@ -777,21 +903,146 @@ void FSpriteEditorViewportClient::ProcessClick(FSceneView& View, HHitProxy* HitP
 	{
 		ClearSelectionSet();
 
-		if (Event == EInputEvent::IE_DoubleClick && CurrentMode == ESpriteEditorMode::EditSourceRegionMode)
+		if (CurrentMode == ESpriteEditorMode::EditSourceRegionMode)
 		{
-			FVector4 WorldPoint = View.PixelToWorld(HitX, HitY, 0);
-			UPaperSprite* Sprite = GetSpriteBeingEdited();
-			FVector2D TexturePoint = SourceTextureViewComponent->GetSprite()->ConvertWorldSpaceToTextureSpace(WorldPoint);
-			Sprite->ExtractSourceRegionFromTexturePoint(TexturePoint);
+			if (Event == EInputEvent::IE_DoubleClick)
+			{
+				FVector4 WorldPoint = View.PixelToWorld(HitX, HitY, 0);
+				UPaperSprite* Sprite = GetSpriteBeingEdited();
+				FVector2D TexturePoint = SourceTextureViewComponent->GetSprite()->ConvertWorldSpaceToTextureSpace(WorldPoint);
+				if (bIsCtrlKeyDown)
+				{
+					UPaperSprite* NewSprite = CreateNewSprite(Sprite->GetSourceUV(), Sprite->GetSourceSize());
+					if (NewSprite != nullptr)
+					{
+						NewSprite->ExtractSourceRegionFromTexturePoint(TexturePoint);
+						bHandled = true;
+					}
+				}
+				else
+				{
+					Sprite->ExtractSourceRegionFromTexturePoint(TexturePoint);
+					bHandled = true;
+				}
+			}
+			else
+			{
+				FVector4 WorldPoint = View.PixelToWorld(HitX, HitY, 0);
+				FVector2D TexturePoint = SourceTextureViewComponent->GetSprite()->ConvertWorldSpaceToTextureSpace(WorldPoint);
+				for (int RelatedSpriteIndex = 0; RelatedSpriteIndex < RelatedSprites.Num(); ++RelatedSpriteIndex)
+				{
+					FRelatedSprite& RelatedSprite = RelatedSprites[RelatedSpriteIndex];
+					if (TexturePoint.X >= RelatedSprite.SourceUV.X && TexturePoint.Y >= RelatedSprite.SourceUV.Y &&
+						TexturePoint.X < (RelatedSprite.SourceUV.X + RelatedSprite.SourceDimension.X) &&
+						TexturePoint.Y < (RelatedSprite.SourceUV.Y + RelatedSprite.SourceDimension.Y))
+					{
+						// Select this sprite
+						if (UPaperSprite* LoadedSprite = Cast<UPaperSprite>(RelatedSprite.AssetData.GetAsset()))
+						{
+							if (SpriteEditorPtr.IsValid())
+							{
+								SpriteEditorPtr.Pin()->SetSpriteBeingEdited(LoadedSprite);
+								break;
+							}
+						}
+						bHandled = true;
+					}
+				}
+			}
 		}
 
-		FPaperEditorViewportClient::ProcessClick(View, HitProxy, Key, Event, HitX, HitY);
+		if (!bHandled)
+		{
+			FPaperEditorViewportClient::ProcessClick(View, HitProxy, Key, Event, HitX, HitY);
+		}
 	}
 }
 
+// Create a new sprite and return this sprite. The sprite editor will now be editing this new sprite
+// Returns nullptr if failed
+UPaperSprite* FSpriteEditorViewportClient::CreateNewSprite(FVector2D TopLeft, FVector2D Dimensions)
+{
+	FAssetToolsModule& AssetToolsModule = FModuleManager::Get().LoadModuleChecked<FAssetToolsModule>("AssetTools");
+	FContentBrowserModule& ContentBrowserModule = FModuleManager::LoadModuleChecked<FContentBrowserModule>("ContentBrowser");
+
+	UPaperSprite* CurrentSprite = GetSpriteBeingEdited();
+	UPaperSprite* CreatedSprite = nullptr;
+
+	// Create the factory used to generate the sprite
+	UPaperSpriteFactory* SpriteFactory = ConstructObject<UPaperSpriteFactory>(UPaperSpriteFactory::StaticClass());
+	SpriteFactory->InitialTexture = CurrentSprite->SourceTexture;
+	SpriteFactory->bUseSourceRegion = true;
+	SpriteFactory->InitialSourceUV = TopLeft;
+	SpriteFactory->InitialSourceDimension = Dimensions;
+
+
+	// Get a unique name for the sprite
+	FString Name, PackageName;
+	AssetToolsModule.Get().CreateUniqueAssetName(CurrentSprite->GetOutermost()->GetName(), TEXT(""), /*out*/ PackageName, /*out*/ Name);
+	const FString PackagePath = FPackageName::GetLongPackagePath(PackageName);
+	if (UObject* NewAsset = AssetToolsModule.Get().CreateAsset(Name, PackagePath, UPaperSprite::StaticClass(), SpriteFactory))
+	{
+		TArray<UObject*> Objects;
+		Objects.Add(NewAsset);
+		ContentBrowserModule.Get().SyncBrowserToAssets(Objects);
+
+		UPaperSprite* NewSprite = Cast<UPaperSprite>(NewAsset);
+		if (SpriteEditorPtr.IsValid() && NewSprite != nullptr)
+		{
+			SpriteEditorPtr.Pin()->SetSpriteBeingEdited(NewSprite);
+		}
+
+		CreatedSprite = NewSprite;
+	}
+
+	return CreatedSprite;
+}
+
 bool FSpriteEditorViewportClient::InputKey(FViewport* Viewport, int32 ControllerId, FKey Key, EInputEvent Event, float AmountDepressed, bool bGamepad)
 {
 	bool bHandled = false;
+	FInputEventState InputState(Viewport, Key, Event);
+
+	// Handle marquee tracking in source region edit mode
+	if (IsInSourceRegionEditMode())
+	{
+		if (Key == EKeys::LeftMouseButton)
+		{
+			int32 HitX = Viewport->GetMouseX();
+			int32 HitY = Viewport->GetMouseY();
+
+			if (Event == IE_Pressed && InputState.IsCtrlButtonPressed())
+			{
+				HHitProxy*	HitResult = Viewport->GetHitProxy(HitX, HitY);
+				if (HitResult == nullptr)
+				{
+					bIsMarqueeTracking = true;
+					MarqueeStartPos = FVector2D(HitX, HitY);
+					MarqueeEndPos = MarqueeStartPos;
+					bHandled = true;
+				}
+			}
+			else if (bIsMarqueeTracking && Event == IE_Released)
+			{
+				MarqueeEndPos = FVector2D(HitX, HitY);
+				bIsMarqueeTracking = false;
+				bHandled = true;
+
+				FVector2D TextureSpaceStartPos, TextureSpaceDimensions;
+				if (ConvertMarqueeToSourceTextureSpace(/*out*/TextureSpaceStartPos, /*out*/TextureSpaceDimensions))
+				{
+					//@TODO: Warn if overlapping with another sprite
+					CreateNewSprite(TextureSpaceStartPos, TextureSpaceDimensions);
+				}
+			}
+		}
+		else if (bIsMarqueeTracking && Key == EKeys::Escape)
+		{
+			// Allow user to cancel operation
+			bIsMarqueeTracking = false;
+			bHandled = true;
+		}
+	}
 
 	// Start the drag
 	//@TODO: EKeys::LeftMouseButton
@@ -941,6 +1192,11 @@ void FSpriteEditorViewportClient::NotifySpriteBeingEditedHasChanged()
 	RenderSpriteComponent->SetSprite(Sprite);
 	UpdateSourceTextureSpriteFromSprite(Sprite);
 
+	if (IsInSourceRegionEditMode()) 
+	{
+		UpdateRelatedSpritesList();
+	}
+
 	//
 	bDeferZoomToSprite = true;
 }
@@ -1061,6 +1317,56 @@ void FSpriteEditorViewportClient::ClearSelectionSet()
 	SelectionSet.Empty();
 }
 
+void FSpriteEditorViewportClient::ResetMarqueeTracking()
+{
+	bIsMarqueeTracking = false;
+}
+
+bool FSpriteEditorViewportClient::ConvertMarqueeToSourceTextureSpace(/*out*/FVector2D& OutStartPos, /*out*/FVector2D& OutDimension)
+{
+	bool bSuccessful = false;
+	UPaperSprite* Sprite = SourceTextureViewComponent->GetSprite();
+	UTexture2D* SpriteSourceTexture = Sprite->GetSourceTexture();
+	if (SpriteSourceTexture != nullptr)
+	{
+		// Calculate world space positions
+		FSceneViewFamilyContext ViewFamily(FSceneViewFamily::ConstructionValues(Viewport, GetScene(), EngineShowFlags));
+		FSceneView* View = CalcSceneView(&ViewFamily);
+		FVector StartPos = View->PixelToWorld(MarqueeStartPos.X, MarqueeStartPos.Y, 0);
+		FVector EndPos = View->PixelToWorld(MarqueeEndPos.X, MarqueeEndPos.Y, 0);
+
+		// Convert to source texture space to work out the pixels dragged
+		FVector2D TextureSpaceStartPos = Sprite->ConvertWorldSpaceToTextureSpace(StartPos);
+		FVector2D TextureSpaceEndPos = Sprite->ConvertWorldSpaceToTextureSpace(EndPos);
+
+		if (TextureSpaceStartPos.X > TextureSpaceEndPos.X)
+		{
+			Swap(TextureSpaceStartPos.X, TextureSpaceEndPos.X);
+		}
+		if (TextureSpaceStartPos.Y > TextureSpaceEndPos.Y)
+		{
+			Swap(TextureSpaceStartPos.Y, TextureSpaceEndPos.Y);
+		}
+
+		int SourceTextureWidth = SpriteSourceTexture->GetSizeX();
+		int SourceTextureHeight = SpriteSourceTexture->GetSizeY();
+		TextureSpaceStartPos.X = FMath::Clamp((int)TextureSpaceStartPos.X, 0, SourceTextureWidth - 1);
+		TextureSpaceStartPos.Y = FMath::Clamp((int)TextureSpaceStartPos.Y, 0, SourceTextureHeight - 1);
+		TextureSpaceEndPos.X = FMath::Clamp((int)TextureSpaceEndPos.X, 0, SourceTextureWidth - 1);
+		TextureSpaceEndPos.Y = FMath::Clamp((int)TextureSpaceEndPos.Y, 0, SourceTextureHeight - 1);
+
+		FVector2D TextureSpaceDimensions = TextureSpaceEndPos - TextureSpaceStartPos;
+		if (TextureSpaceDimensions.X > 0 || TextureSpaceDimensions.Y > 0)
+		{
+			OutStartPos = TextureSpaceStartPos;
+			OutDimension = TextureSpaceDimensions;
+			bSuccessful = true;
+		}
+	}
+
+	return bSuccessful;
+}
+
 //////////////////////////////////////////////////////////////////////////
 
 #undef LOCTEXT_NAMESPACE
\ No newline at end of file
diff --git a/Engine/Plugins/Experimental/Paper2D/Source/Paper2DEditor/Private/SpriteEditor/SpriteEditorViewportClient.h b/Engine/Plugins/Experimental/Paper2D/Source/Paper2DEditor/Private/SpriteEditor/SpriteEditorViewportClient.h
--- a/Engine/Plugins/Experimental/Paper2D/Source/Paper2DEditor/Private/SpriteEditor/SpriteEditorViewportClient.h
+++ b/Engine/Plugins/Experimental/Paper2D/Source/Paper2DEditor/Private/SpriteEditor/SpriteEditorViewportClient.h
@@ -26,6 +26,8 @@ namespace ESpriteEditorMode
 	};
 }
 
+struct FRelatedSprite;
+
 //////////////////////////////////////////////////////////////////////////
 // FSpriteEditorViewportClient
 
@@ -70,11 +72,11 @@ public:
 	void ToggleShowMeshEdges();
 	bool IsShowMeshEdgesChecked() const;
 
-	void EnterViewMode() { CurrentMode = ESpriteEditorMode::ViewMode; }
-	void EnterSourceRegionEditMode() { CurrentMode = ESpriteEditorMode::EditSourceRegionMode; }
-	void EnterCollisionEditMode() { CurrentMode = ESpriteEditorMode::EditCollisionMode; }
-	void EnterRenderingEditMode() { CurrentMode = ESpriteEditorMode::EditRenderingGeomMode; }
-	void EnterAddSpriteMode() { CurrentMode = ESpriteEditorMode::AddSpriteMode; }
+	void EnterViewMode() { CurrentMode = ESpriteEditorMode::ViewMode; ResetMarqueeTracking(); }
+	void EnterSourceRegionEditMode() { CurrentMode = ESpriteEditorMode::EditSourceRegionMode; ResetMarqueeTracking(); UpdateRelatedSpritesList(); }
+	void EnterCollisionEditMode() { CurrentMode = ESpriteEditorMode::EditCollisionMode; ResetMarqueeTracking(); }
+	void EnterRenderingEditMode() { CurrentMode = ESpriteEditorMode::EditRenderingGeomMode; ResetMarqueeTracking(); }
+	void EnterAddSpriteMode() { CurrentMode = ESpriteEditorMode::AddSpriteMode; ResetMarqueeTracking(); }
 
 	bool IsInViewMode() const { return CurrentMode == ESpriteEditorMode::ViewMode; }
 	bool IsInSourceRegionEditMode() const { return CurrentMode == ESpriteEditorMode::EditSourceRegionMode; }
@@ -103,6 +105,10 @@ public:
 	// Invalidate any references to the sprite being edited; it has changed
 	void NotifySpriteBeingEditedHasChanged();
 
+	void UpdateRelatedSpritesList();
+
+	UPaperSprite* CreateNewSprite(FVector2D TopLeft, FVector2D Dimensions);
+
 	ESpriteEditorMode::Type GetCurrentMode() const
 	{
 		return CurrentMode;
@@ -155,6 +161,17 @@ private:
 
 	// Should we zoom to the sprite next tick?
 	bool bDeferZoomToSprite;
+
+	// Should we show related sprites in the source texture?
+	bool bShowRelatedSprites;
+
+	// Marquee tracking
+	bool bIsMarqueeTracking;
+	FVector2D MarqueeStartPos, MarqueeEndPos;
+
+	// Other sprites
+	TArray<FRelatedSprite> RelatedSprites;
+
 private:
 	UPaperSprite* GetSpriteBeingEdited() const
 	{
@@ -175,15 +192,19 @@ private:
 	void DrawRenderStats(FViewport& InViewport, FSceneView& View, FCanvas& Canvas, class UPaperSprite* Sprite, int32& YPos);
 	void DrawSockets(const FSceneView* View, FPrimitiveDrawInterface* PDI);
 	void DrawSocketNames(FViewport& InViewport, FSceneView& View, FCanvas& Canvas);
-	void DrawSourceRegion(FViewport& InViewport, FSceneView& View, FCanvas& Canvas, const FLinearColor& GeometryVertexColor, bool bIsRenderGeometry);
+	void DrawSourceRegion(FViewport& InViewport, FSceneView& View, FCanvas& Canvas, const FLinearColor& GeometryVertexColor);
+	void DrawRelatedSprites(FViewport& InViewport, FSceneView& View, FCanvas& Canvas, const FLinearColor& GeometryVertexColor);
+	void DrawMarquee(FViewport& InViewport, FSceneView& View, FCanvas& Canvas, const FLinearColor& MarqueeColor);
 
 	void BeginTransaction(const FText& SessionName);
 	void EndTransaction();
 
 	void UpdateSourceTextureSpriteFromSprite(UPaperSprite* SourceSprite);
-
 	void ClearSelectionSet();
 
+
+	void ResetMarqueeTracking();
+	bool ConvertMarqueeToSourceTextureSpace(/*out*/FVector2D& OutStartPos, /*out*/FVector2D& OutDimension);
 	
 	// Can return null
 	FSpritePolygonCollection* GetGeometryBeingEdited() const;
