Paper2D: Blind-fix non-unity compilation error

diff --git a/Engine/Plugins/Experimental/Paper2D/Source/Paper2DEditor/Private/SpriteEditor/SpriteEditorViewportClient.cpp b/Engine/Plugins/Experimental/Paper2D/Source/Paper2DEditor/Private/SpriteEditor/SpriteEditorViewportClient.cpp
--- a/Engine/Plugins/Experimental/Paper2D/Source/Paper2DEditor/Private/SpriteEditor/SpriteEditorViewportClient.cpp
+++ b/Engine/Plugins/Experimental/Paper2D/Source/Paper2DEditor/Private/SpriteEditor/SpriteEditorViewportClient.cpp
@@ -40,17 +40,6 @@ struct HSpriteSelectableObjectHitProxy : public HHitProxy
 };
 IMPLEMENT_HIT_PROXY(HSpriteSelectableObjectHitProxy, HHitProxy);
 
-
-//////////////////////////////////////////////////////////////////////////
-// FRelatedSprite
-
-struct FRelatedSprite
-{
-	FAssetData AssetData;
-	FVector2D SourceUV;
-	FVector2D SourceDimension;
-};
-
 //////////////////////////////////////////////////////////////////////////
 // FSpriteEditorViewportClient
 
diff --git a/Engine/Plugins/Experimental/Paper2D/Source/Paper2DEditor/Private/SpriteEditor/SpriteEditorViewportClient.h b/Engine/Plugins/Experimental/Paper2D/Source/Paper2DEditor/Private/SpriteEditor/SpriteEditorViewportClient.h
--- a/Engine/Plugins/Experimental/Paper2D/Source/Paper2DEditor/Private/SpriteEditor/SpriteEditorViewportClient.h
+++ b/Engine/Plugins/Experimental/Paper2D/Source/Paper2DEditor/Private/SpriteEditor/SpriteEditorViewportClient.h
@@ -26,7 +26,15 @@ namespace ESpriteEditorMode
 	};
 }
 
-struct FRelatedSprite;
+//////////////////////////////////////////////////////////////////////////
+// FRelatedSprite
+
+struct FRelatedSprite
+{
+	FAssetData AssetData;
+	FVector2D SourceUV;
+	FVector2D SourceDimension;
+};
 
 //////////////////////////////////////////////////////////////////////////
 // FSpriteEditorViewportClient
