Paper2D: Fixing another non-unity build error

diff --git a/Engine/Plugins/Experimental/Paper2D/Source/Paper2DEditor/Private/SpriteEditor/SpriteEditorViewportClient.h b/Engine/Plugins/Experimental/Paper2D/Source/Paper2DEditor/Private/SpriteEditor/SpriteEditorViewportClient.h
--- a/Engine/Plugins/Experimental/Paper2D/Source/Paper2DEditor/Private/SpriteEditor/SpriteEditorViewportClient.h
+++ b/Engine/Plugins/Experimental/Paper2D/Source/Paper2DEditor/Private/SpriteEditor/SpriteEditorViewportClient.h
@@ -10,6 +10,7 @@
 #include "PreviewScene.h"
 #include "ScopedTransaction.h"
 #include "SpriteEditorSelections.h"
+#include "AssetData.h"
 
 //////////////////////////////////////////////////////////////////////////
 // 
