FPrimitiveSceneProxy::PreRenderView / DrawDynamicElements refactored into GetDynamicMeshElements
* Both paths still exist in code during this transition.  The GetDynamicMeshElements path is currently disabled, enable with 'r.UseGetDynamicMeshElements 1'.
* DrawDynamicElements was called once per pass, all sorts of pass-specific cruft had leaked in along with short-term modifications of shared structures (vertex factories and material proxies).  DrawDynamicElements was not const, so a bunch of 'update' type work had leaked in.  Various PreRenderView implementations relied on the view being in the view family.
* GetDynamicMeshElements is called once in InitViews, and the resulting meshes are used in various passes.  GetDynamicMeshElements is const and must not modify the proxy, update work should be pushed to the proxy from the game thread.  GetDynamicMeshElements can be called on views that are not in the view family, which allows a proper view for shadow depth passes.
* View mode override functionality moved out of the individual proxies into a centralized place, ApplyViewModeOverrides
* GPU particles not working yet, FRecastRenderingSceneProxy not implemented

diff --git a/Engine/Plugins/Experimental/Paper2D/Source/Paper2D/Private/PaperRenderSceneProxy.cpp b/Engine/Plugins/Experimental/Paper2D/Source/Paper2D/Private/PaperRenderSceneProxy.cpp
--- a/Engine/Plugins/Experimental/Paper2D/Source/Paper2D/Private/PaperRenderSceneProxy.cpp
+++ b/Engine/Plugins/Experimental/Paper2D/Source/Paper2D/Private/PaperRenderSceneProxy.cpp
@@ -160,9 +160,9 @@ FPaperRenderSceneProxy::FPaperRenderSceneProxy(const UPrimitiveComponent* InComp
 	, Material(NULL)
 	, Owner(InComponent->GetOwner())
 	, bCastShadow(InComponent->CastShadow)
-	, WireframeColor(FLinearColor::White)
 	, CollisionResponse(InComponent->GetCollisionResponseToChannels())
 {
+	WireframeColor = FLinearColor::White;
 }
 
 FPaperRenderSceneProxy::~FPaperRenderSceneProxy()
@@ -185,6 +185,49 @@ void FPaperRenderSceneProxy::CreateRenderThreadResources()
 #endif
 }
 
+void FPaperRenderSceneProxy::GetDynamicMeshElements(const TArray<const FSceneView*>& Views, const FSceneViewFamily& ViewFamily, uint32 VisibilityMap, FMeshElementCollector& Collector) const
+{
+	QUICK_SCOPE_CYCLE_COUNTER(STAT_PaperRenderSceneProxy_GetDynamicMeshElements);
+	checkSlow(IsInRenderingThread());
+
+	for (int32 ViewIndex = 0; ViewIndex < Views.Num(); ViewIndex++)
+	{
+		if (VisibilityMap & (1 << ViewIndex))
+		{
+			const FSceneView* View = Views[ViewIndex];
+	
+#if !(UE_BUILD_SHIPPING || UE_BUILD_TEST)
+			if (View->Family->EngineShowFlags.Paper2DSprites)
+			{
+				RenderBounds(Collector.GetPDI(ViewIndex), ViewFamily.EngineShowFlags, GetBounds(), (Owner == NULL) || IsSelected());
+			}
+#endif
+
+			// Set the selection/hover color from the current engine setting.
+			FLinearColor OverrideColor;
+			bool bUseOverrideColor = false;
+
+			const bool bShowAsSelected = !(GIsEditor && ViewFamily.EngineShowFlags.Selection) || IsSelected();
+			if (bShowAsSelected || IsHovered())
+			{
+				bUseOverrideColor = true;
+				OverrideColor = GetSelectionColor(FLinearColor::White, bShowAsSelected, IsHovered());
+			}
+
+			bUseOverrideColor = false;
+
+			// Sprites of locked actors draw in red.
+			//FLinearColor LevelColorToUse = IsSelected() ? ColorToUse : (FLinearColor)LevelColor;
+			//FLinearColor PropertyColorToUse = PropertyColor;
+
+#if TEST_BATCHING
+#else
+			GetDynamicMeshElementsForView(View, ViewIndex, bUseOverrideColor, OverrideColor, Collector);
+#endif
+		}
+	}
+}
+
 void FPaperRenderSceneProxy::DrawDynamicElements(FPrimitiveDrawInterface* PDI, const FSceneView* View)
 {
 	QUICK_SCOPE_CYCLE_COUNTER(STAT_PaperRenderSceneProxy_DrawDynamicElements);
@@ -234,6 +277,15 @@ FVertexFactory* FPaperRenderSceneProxy::GetPaperSpriteVertexFactory() const
 	return &GPaperSpriteVertexFactory;
 }
 
+void FPaperRenderSceneProxy::GetDynamicMeshElementsForView(const FSceneView* View, int32 ViewIndex, bool bUseOverrideColor, const FLinearColor& OverrideColor, FMeshElementCollector& Collector) const
+{
+	if (Material != nullptr)
+	{
+		GetBatchMesh(View, bUseOverrideColor, OverrideColor, Material, BatchedSprites, ViewIndex, Collector);
+	}
+	GetNewBatchMeshes(View, bUseOverrideColor, OverrideColor, ViewIndex, Collector);
+}
+
 void FPaperRenderSceneProxy::DrawDynamicElements_RichMesh(FPrimitiveDrawInterface* PDI, const FSceneView* View, bool bUseOverrideColor, const FLinearColor& OverrideColor)
 {
 	if (Material != nullptr)
@@ -311,10 +363,178 @@ void FPaperRenderSceneProxy::DrawNewBatches(FPrimitiveDrawInterface* PDI, const
 	}
 }
 
+void FPaperRenderSceneProxy::GetNewBatchMeshes(const FSceneView* View, bool bUseOverrideColor, const FLinearColor& OverrideColor, int32 ViewIndex, FMeshElementCollector& Collector) const
+{
+	//@TODO: Doesn't support OverrideColor yet
+	if (BatchedSections.Num() == 0)
+	{
+		return;
+	}
+
+	const uint8 DPG = GetDepthPriorityGroup(View);
+	FVertexFactory* VertexFactory = GetPaperSpriteVertexFactory();
+	const bool bIsWireframeView = View->Family->EngineShowFlags.Wireframe;
+
+	for (const FSpriteRenderSection& Batch : BatchedSections)
+	{
+		const FLinearColor EffectiveWireframeColor = (Batch.Material->GetBlendMode() != BLEND_Opaque) ? WireframeColor : FLinearColor::Green;
+		FMaterialRenderProxy* ParentMaterialProxy = Batch.Material->GetRenderProxy((View->Family->EngineShowFlags.Selection) && IsSelected(), IsHovered());
+		FTexture* TextureResource = (Batch.Texture != nullptr) ? Batch.Texture->Resource : nullptr;
+
+		if ((TextureResource != nullptr) && (Batch.NumVertices > 0))
+		{
+			FMeshBatch& Mesh = Collector.AllocateMesh();
+
+			Mesh.bCanApplyViewModeOverrides = true;
+			Mesh.bUseWireframeSelectionColoring = IsSelected();
+
+			// Implementing our own wireframe coloring as the automatic one (controlled by Mesh.bCanApplyViewModeOverrides) only supports per-FPrimitiveSceneProxy WireframeColor
+			if (bIsWireframeView)
+			{
+				auto WireframeMaterialInstance = new FColoredMaterialRenderProxy(
+					GEngine->WireframeMaterial->GetRenderProxy(IsSelected(), IsHovered()),
+					GetSelectionColor(EffectiveWireframeColor, IsSelected(), IsHovered(), false)
+					);
+
+				Collector.RegisterOneFrameMaterialProxy(WireframeMaterialInstance);
+
+				ParentMaterialProxy = WireframeMaterialInstance;
+
+				Mesh.bWireframe = true;
+				// We are applying our own wireframe override
+				Mesh.bCanApplyViewModeOverrides = false;
+			}
+
+			// Create a texture override material proxy and register it as a dynamic resource so that it won't be deleted until the rendering thread has finished with it
+			FTextureOverrideRenderProxy* TextureOverrideMaterialProxy = new FTextureOverrideRenderProxy(ParentMaterialProxy, Batch.Texture, TEXT("SpriteTexture"));
+			Collector.RegisterOneFrameMaterialProxy(TextureOverrideMaterialProxy);
+
+			Mesh.UseDynamicData = true;
+			// Note: memory referenced by Collector must be valid as long as the Collector is around, see Collector.AllocateOneFrameResource
+			Mesh.DynamicVertexData = BatchedVertices.GetTypedData();
+			Mesh.DynamicVertexStride = sizeof(FPaperSpriteVertex);
+
+			Mesh.VertexFactory = VertexFactory;
+			Mesh.LCI = NULL;
+			Mesh.ReverseCulling = IsLocalToWorldDeterminantNegative() ? true : false;
+			Mesh.CastShadow = bCastShadow;
+			Mesh.DepthPriorityGroup = DPG;
+			Mesh.Type = PT_TriangleList;
+			Mesh.bDisableBackfaceCulling = true;
+			Mesh.MaterialRenderProxy = TextureOverrideMaterialProxy;
+
+			// Set up the FMeshBatchElement.
+			FMeshBatchElement& BatchElement = Mesh.Elements[0];
+			BatchElement.IndexBuffer = NULL;
+			BatchElement.DynamicIndexData = NULL;
+			BatchElement.DynamicIndexStride = 0;
+			BatchElement.PrimitiveUniformBuffer = GetUniformBuffer();
+			BatchElement.FirstIndex = Batch.VertexOffset;
+			BatchElement.MinVertexIndex = Batch.VertexOffset;
+			BatchElement.MaxVertexIndex = Batch.VertexOffset + Batch.NumVertices;
+			BatchElement.NumPrimitives = Batch.NumVertices / 3;
+
+			Collector.AddMesh(ViewIndex, Mesh);
+		}
+	}
+}
+
+class FPaperVertexArray : public FOneFrameResource
+{
+public:
+	TArray<FPaperSpriteVertex, TInlineAllocator<6> > Vertices;
+};
+
+void FPaperRenderSceneProxy::GetBatchMesh(const FSceneView* View, bool bUseOverrideColor, const FLinearColor& OverrideColor, class UMaterialInterface* BatchMaterial, const TArray<FSpriteDrawCallRecord>& Batch, int32 ViewIndex, FMeshElementCollector& Collector) const
+{
+	const uint8 DPG = GetDepthPriorityGroup(View);
+	FVertexFactory* VertexFactory = GetPaperSpriteVertexFactory();
+
+	const bool bIsWireframeView = View->Family->EngineShowFlags.Wireframe;
+	const FLinearColor EffectiveWireframeColor = (BatchMaterial->GetBlendMode() != BLEND_Opaque) ? WireframeColor : FLinearColor::Green;
+
+	for (const FSpriteDrawCallRecord& Record : Batch)
+	{
+		FTexture* TextureResource = (Record.Texture != nullptr) ? Record.Texture->Resource : nullptr;
+		if ((TextureResource != nullptr) && (Record.RenderVerts.Num() > 0))
+		{
+			const FLinearColor SpriteColor = bUseOverrideColor ? OverrideColor : Record.Color;
+			const FVector EffectiveOrigin = Record.Destination;
+
+			FPaperVertexArray& VertexArray = Collector.AllocateOneFrameResource<FPaperVertexArray>();
+			VertexArray.Vertices.Empty(Record.RenderVerts.Num());
+
+			for (int32 SVI = 0; SVI < Record.RenderVerts.Num(); ++SVI)
+			{
+				const FVector4& SourceVert = Record.RenderVerts[SVI];
+				const FVector Pos((PaperAxisX * SourceVert.X) + (PaperAxisY * SourceVert.Y) + EffectiveOrigin);
+				const FVector2D UV(SourceVert.Z, SourceVert.W);
+
+				new (VertexArray.Vertices) FPaperSpriteVertex(Pos, UV, SpriteColor);
+			}
+
+			// Set up the FMeshElement.
+			FMeshBatch& Mesh = Collector.AllocateMesh();
+
+			Mesh.UseDynamicData = true;
+			Mesh.DynamicVertexData = VertexArray.Vertices.GetTypedData();
+			Mesh.DynamicVertexStride = sizeof(FPaperSpriteVertex);
+			Mesh.VertexFactory = VertexFactory;
+			Mesh.bCanApplyViewModeOverrides = true;
+			Mesh.bUseWireframeSelectionColoring = IsSelected();
+
+			FMaterialRenderProxy* ParentMaterialProxy = BatchMaterial->GetRenderProxy((View->Family->EngineShowFlags.Selection) && IsSelected(), IsHovered());
+
+			// Implementing our own wireframe coloring as the automatic one (controlled by Mesh.bCanApplyViewModeOverrides) only supports per-FPrimitiveSceneProxy WireframeColor
+			if (bIsWireframeView)
+			{
+				auto WireframeMaterialInstance = new FColoredMaterialRenderProxy(
+					GEngine->WireframeMaterial->GetRenderProxy(IsSelected(), IsHovered()),
+					GetSelectionColor(EffectiveWireframeColor, IsSelected(), IsHovered(), false)
+					);
+
+				Collector.RegisterOneFrameMaterialProxy(WireframeMaterialInstance);
+
+				ParentMaterialProxy = WireframeMaterialInstance;
+
+				Mesh.bWireframe = true;
+				// We are applying our own wireframe override
+				Mesh.bCanApplyViewModeOverrides = false;
+			}
+
+			// Create a texture override material proxy and register it as a dynamic resource so that it won't be deleted until the rendering thread has finished with it
+			FTextureOverrideRenderProxy* TextureOverrideMaterialProxy = new FTextureOverrideRenderProxy(ParentMaterialProxy, Record.Texture, TEXT("SpriteTexture"));
+			Collector.RegisterOneFrameMaterialProxy(TextureOverrideMaterialProxy);
+
+			Mesh.MaterialRenderProxy = TextureOverrideMaterialProxy;
+			Mesh.LCI = NULL;
+			Mesh.ReverseCulling = IsLocalToWorldDeterminantNegative() ? true : false;
+			Mesh.CastShadow = false;
+			Mesh.DepthPriorityGroup = DPG;
+			Mesh.Type = PT_TriangleList;
+			Mesh.bDisableBackfaceCulling = true;
+
+			// Set up the FMeshBatchElement.
+			FMeshBatchElement& BatchElement = Mesh.Elements[0];
+			BatchElement.IndexBuffer = NULL;
+			BatchElement.DynamicIndexData = NULL;
+			BatchElement.DynamicIndexStride = 0;
+			BatchElement.FirstIndex = 0;
+			BatchElement.MinVertexIndex = 0;
+			BatchElement.MaxVertexIndex = VertexArray.Vertices.Num();
+			BatchElement.PrimitiveUniformBuffer = GetUniformBuffer();
+			BatchElement.NumPrimitives = VertexArray.Vertices.Num() / 3;
+
+			const bool bIsWireframeView = View->Family->EngineShowFlags.Wireframe;
+
+			Collector.AddMesh(ViewIndex, Mesh);
+		}
+	}
+}
+
 void FPaperRenderSceneProxy::DrawBatch(FPrimitiveDrawInterface* PDI, const FSceneView* View, bool bUseOverrideColor, const FLinearColor& OverrideColor, class UMaterialInterface* BatchMaterial, const TArray<FSpriteDrawCallRecord>& Batch)
 {
 	const uint8 DPG = GetDepthPriorityGroup(View);
-	
 	FVertexFactory* VertexFactory = GetPaperSpriteVertexFactory();
 
 	const bool bIsWireframeView = View->Family->EngineShowFlags.Wireframe;
@@ -493,4 +713,4 @@ bool FPaperRenderSceneProxy::IsCollisionView(const FSceneView* View, bool& bDraw
 	}
 
 	return bInCollisionView;
-}
\ No newline at end of file
+}
diff --git a/Engine/Plugins/Experimental/Paper2D/Source/Paper2D/Private/PaperRenderSceneProxy.h b/Engine/Plugins/Experimental/Paper2D/Source/Paper2D/Private/PaperRenderSceneProxy.h
--- a/Engine/Plugins/Experimental/Paper2D/Source/Paper2D/Private/PaperRenderSceneProxy.h
+++ b/Engine/Plugins/Experimental/Paper2D/Source/Paper2D/Private/PaperRenderSceneProxy.h
@@ -89,6 +89,7 @@ public:
 	~FPaperRenderSceneProxy();
 
 	// FPrimitiveSceneProxy interface.
+	virtual void GetDynamicMeshElements(const TArray<const FSceneView*>& Views, const FSceneViewFamily& ViewFamily, uint32 VisibilityMap, FMeshElementCollector& Collector) const override;
 	virtual void DrawDynamicElements(FPrimitiveDrawInterface* PDI, const FSceneView* View) override;
 	virtual FPrimitiveViewRelevance GetViewRelevance(const FSceneView* View) override;
 	virtual void OnTransformChanged() override;
@@ -106,10 +107,14 @@ public:
 
 protected:
 	virtual void DrawDynamicElements_RichMesh(FPrimitiveDrawInterface* PDI, const FSceneView* View, bool bUseOverrideColor, const FLinearColor& OverrideColor);
+	virtual void GetDynamicMeshElementsForView(const FSceneView* View, int32 ViewIndex, bool bUseOverrideColor, const FLinearColor& OverrideColor, FMeshElementCollector& Collector) const;
 
 	void DrawBatch(FPrimitiveDrawInterface* PDI, const FSceneView* View, bool bUseOverrideColor, const FLinearColor& OverrideColor, class UMaterialInterface* BatchMaterial, const TArray<FSpriteDrawCallRecord>& Batch);
 	void DrawNewBatches(FPrimitiveDrawInterface* PDI, const FSceneView* View, bool bUseOverrideColor, const FLinearColor& OverrideColor);
 
+	void GetBatchMesh(const FSceneView* View, bool bUseOverrideColor, const FLinearColor& OverrideColor, class UMaterialInterface* BatchMaterial, const TArray<FSpriteDrawCallRecord>& Batch, int32 ViewIndex, FMeshElementCollector& Collector) const;
+	void GetNewBatchMeshes(const FSceneView* View, bool bUseOverrideColor, const FLinearColor& OverrideColor, int32 ViewIndex, FMeshElementCollector& Collector) const;
+
 	bool IsCollisionView(const FSceneView* View, bool& bDrawSimpleCollision, bool& bDrawComplexCollision) const;
 
 	friend class FPaperBatchSceneProxy;
@@ -127,9 +132,6 @@ protected:
 
 	bool bCastShadow;
 
-	// The color to draw as in wireframe mode
-	FLinearColor WireframeColor;
-
 	// The view relevance for the associated material
 	FMaterialRelevance MaterialRelevance;
 
diff --git a/Engine/Plugins/Experimental/Paper2D/Source/Paper2D/Private/PaperSpriteSceneProxy.cpp b/Engine/Plugins/Experimental/Paper2D/Source/Paper2D/Private/PaperSpriteSceneProxy.cpp
--- a/Engine/Plugins/Experimental/Paper2D/Source/Paper2D/Private/PaperSpriteSceneProxy.cpp
+++ b/Engine/Plugins/Experimental/Paper2D/Source/Paper2D/Private/PaperSpriteSceneProxy.cpp
@@ -67,7 +67,7 @@ void FPaperSpriteSceneProxy::DrawDynamicElements(FPrimitiveDrawInterface* PDI, c
 
 					const bool bUseSeparateColorPerHull = (Owner == NULL);
 					const bool bDrawSolid = false;
-					BodySetup->AggGeom.DrawAggGeom(PDI, GeomTransform, GetSelectionColor(CollisionColor, bDrawWireSelected, IsHovered()), &CollisionMaterialInstance, bUseSeparateColorPerHull, bDrawSolid);
+					BodySetup->AggGeom.DrawAggGeom(PDI, GeomTransform, GetSelectionColor(CollisionColor, bDrawWireSelected, IsHovered()), &CollisionMaterialInstance, bUseSeparateColorPerHull, bDrawSolid, false);
 				}
 			}
 		}
@@ -76,6 +76,70 @@ void FPaperSpriteSceneProxy::DrawDynamicElements(FPrimitiveDrawInterface* PDI, c
 	FPaperRenderSceneProxy::DrawDynamicElements(PDI, View);
 }
 
+void FPaperSpriteSceneProxy::GetDynamicMeshElements(const TArray<const FSceneView*>& Views, const FSceneViewFamily& ViewFamily, uint32 VisibilityMap, FMeshElementCollector& Collector) const
+{
+	if (SourceSprite != nullptr)
+	{
+		// Show 3D physics
+		if ((ViewFamily.EngineShowFlags.Collision /*@TODO: && bIsCollisionEnabled*/) && AllowDebugViewmodes())
+		{
+			if (UBodySetup2D* BodySetup2D = Cast<UBodySetup2D>(SourceSprite->BodySetup))
+			{
+				//@TODO: Draw 2D debugging geometry
+			}
+			else if (UBodySetup* BodySetup = SourceSprite->BodySetup)
+			{
+				if (FMath::Abs(GetLocalToWorld().Determinant()) < SMALL_NUMBER)
+				{
+					// Catch this here or otherwise GeomTransform below will assert
+					// This spams so commented out
+					//UE_LOG(LogStaticMesh, Log, TEXT("Zero scaling not supported (%s)"), *StaticMesh->GetPathName());
+				}
+				else
+				{
+					for (int32 ViewIndex = 0; ViewIndex < Views.Num(); ViewIndex++)
+					{
+						if (VisibilityMap & (1 << ViewIndex))
+						{
+							// Make a material for drawing solid collision stuff
+							const UMaterial* LevelColorationMaterial = ViewFamily.EngineShowFlags.Lighting 
+								? GEngine->ShadedLevelColorationLitMaterial : GEngine->ShadedLevelColorationUnlitMaterial;
+
+							auto CollisionMaterialInstance = new FColoredMaterialRenderProxy(
+								LevelColorationMaterial->GetRenderProxy(IsSelected(), IsHovered()),
+								WireframeColor
+								);
+
+							Collector.RegisterOneFrameMaterialProxy(CollisionMaterialInstance);
+
+							// Draw the static mesh's body setup.
+
+							// Get transform without scaling.
+							FTransform GeomTransform(GetLocalToWorld());
+
+							// In old wireframe collision mode, always draw the wireframe highlighted (selected or not).
+							bool bDrawWireSelected = IsSelected();
+							if (ViewFamily.EngineShowFlags.Collision)
+							{
+								bDrawWireSelected = true;
+							}
+
+							// Differentiate the color based on bBlockNonZeroExtent.  Helps greatly with skimming a level for optimization opportunities.
+							const FColor CollisionColor = FColor(220,149,223,255);
+
+							const bool bUseSeparateColorPerHull = (Owner == NULL);
+							const bool bDrawSolid = false;
+							BodySetup->AggGeom.GetAggGeom(GeomTransform, GetSelectionColor(CollisionColor, bDrawWireSelected, IsHovered()), CollisionMaterialInstance, bUseSeparateColorPerHull, bDrawSolid, false, ViewIndex, Collector);
+						}
+					}
+				}
+			}
+		}
+	}
+
+	FPaperRenderSceneProxy::GetDynamicMeshElements(Views, ViewFamily, VisibilityMap, Collector);
+}
+
 void FPaperSpriteSceneProxy::SetSprite_RenderThread(const FSpriteDrawCallRecord& NewDynamicData, int32 SplitIndex)
 {
 	BatchedSprites.Empty();
@@ -115,3 +179,16 @@ void FPaperSpriteSceneProxy::DrawDynamicElements_RichMesh(FPrimitiveDrawInterfac
 		DrawBatch(PDI, View, bUseOverrideColor, OverrideColor, AlternateMaterial, AlternateBatchedSprites);
 	}
 }
+
+void FPaperSpriteSceneProxy::GetDynamicMeshElementsForView(const FSceneView* View, int32 ViewIndex, bool bUseOverrideColor, const FLinearColor& OverrideColor, FMeshElementCollector& Collector) const
+{
+	if (Material != nullptr)
+	{
+		GetBatchMesh(View, bUseOverrideColor, OverrideColor, Material, BatchedSprites, ViewIndex, Collector);
+	}
+
+	if ((AlternateMaterial != nullptr) && (AlternateBatchedSprites.Num() > 0))
+	{
+		GetBatchMesh(View, bUseOverrideColor, OverrideColor, AlternateMaterial, AlternateBatchedSprites, ViewIndex, Collector);
+	}
+}
\ No newline at end of file
diff --git a/Engine/Plugins/Experimental/Paper2D/Source/Paper2D/Private/PaperSpriteSceneProxy.h b/Engine/Plugins/Experimental/Paper2D/Source/Paper2D/Private/PaperSpriteSceneProxy.h
--- a/Engine/Plugins/Experimental/Paper2D/Source/Paper2D/Private/PaperSpriteSceneProxy.h
+++ b/Engine/Plugins/Experimental/Paper2D/Source/Paper2D/Private/PaperSpriteSceneProxy.h
@@ -13,6 +13,7 @@ public:
 	FPaperSpriteSceneProxy(const UPaperSpriteComponent* InComponent);
 
 	// FPrimitiveSceneProxy interface
+	virtual void GetDynamicMeshElements(const TArray<const FSceneView*>& Views, const FSceneViewFamily& ViewFamily, uint32 VisibilityMap, FMeshElementCollector& Collector) const override;
 	virtual void DrawDynamicElements(FPrimitiveDrawInterface* PDI, const FSceneView* View) override;
 	// End of FPrimitiveSceneProxy interface
 
@@ -22,6 +23,7 @@ protected:
 
 	// FPaperRenderSceneProxy interface
 	virtual void DrawDynamicElements_RichMesh(FPrimitiveDrawInterface* PDI, const FSceneView* View, bool bUseOverrideColor, const FLinearColor& OverrideColor) override;
+	virtual void GetDynamicMeshElementsForView(const FSceneView* View, int32 ViewIndex, bool bUseOverrideColor, const FLinearColor& OverrideColor, FMeshElementCollector& Collector) const override;
 	// End of FPaperRenderSceneProxy interface
 
 	UMaterialInterface* AlternateMaterial;
diff --git a/Engine/Plugins/Experimental/Paper2D/Source/Paper2D/Private/PaperTileMapRenderSceneProxy.cpp b/Engine/Plugins/Experimental/Paper2D/Source/Paper2D/Private/PaperTileMapRenderSceneProxy.cpp
--- a/Engine/Plugins/Experimental/Paper2D/Source/Paper2D/Private/PaperTileMapRenderSceneProxy.cpp
+++ b/Engine/Plugins/Experimental/Paper2D/Source/Paper2D/Private/PaperTileMapRenderSceneProxy.cpp
@@ -24,6 +24,143 @@ FPaperTileMapRenderSceneProxy::FPaperTileMapRenderSceneProxy(const UPaperTileMap
 	}
 }
 
+void FPaperTileMapRenderSceneProxy::GetDynamicMeshElements(const TArray<const FSceneView*>& Views, const FSceneViewFamily& ViewFamily, uint32 VisibilityMap, FMeshElementCollector& Collector) const
+{
+	QUICK_SCOPE_CYCLE_COUNTER(STAT_FPaperTileMapRenderSceneProxy_GetDynamicMeshElements);
+	checkSlow(IsInRenderingThread());
+
+	for (int32 ViewIndex = 0; ViewIndex < Views.Num(); ViewIndex++)
+	{
+		if (VisibilityMap & (1 << ViewIndex))
+		{
+			const FSceneView* View = Views[ViewIndex];
+			FPrimitiveDrawInterface* PDI = Collector.GetPDI(ViewIndex);
+
+			// Draw the tile maps
+			//@TODO: RenderThread race condition
+			if (TileMap != NULL)
+			{
+				FColor WireframeColor = FColor(0, 255, 255, 255);
+
+				if ((View->Family->EngineShowFlags.Collision /*@TODO: && bIsCollisionEnabled*/) && AllowDebugViewmodes())
+				{
+					if (UBodySetup2D* BodySetup2D = Cast<UBodySetup2D>(TileMap->BodySetup))
+					{
+						//@TODO: Draw 2D debugging geometry
+					}
+					else if (UBodySetup* BodySetup = TileMap->BodySetup)
+					{
+						if (FMath::Abs(GetLocalToWorld().Determinant()) < SMALL_NUMBER)
+						{
+							// Catch this here or otherwise GeomTransform below will assert
+							// This spams so commented out
+							//UE_LOG(LogStaticMesh, Log, TEXT("Zero scaling not supported (%s)"), *StaticMesh->GetPathName());
+						}
+						else
+						{
+							// Make a material for drawing solid collision stuff
+							const UMaterial* LevelColorationMaterial = View->Family->EngineShowFlags.Lighting
+								? GEngine->ShadedLevelColorationLitMaterial : GEngine->ShadedLevelColorationUnlitMaterial;
+
+							auto CollisionMaterialInstance = new FColoredMaterialRenderProxy(
+								LevelColorationMaterial->GetRenderProxy(IsSelected(), IsHovered()),
+								WireframeColor
+								);
+
+							// Draw the static mesh's body setup.
+
+							// Get transform without scaling.
+							FTransform GeomTransform(GetLocalToWorld());
+
+							// In old wireframe collision mode, always draw the wireframe highlighted (selected or not).
+							bool bDrawWireSelected = IsSelected();
+							if (View->Family->EngineShowFlags.Collision)
+							{
+								bDrawWireSelected = true;
+							}
+
+							// Differentiate the color based on bBlockNonZeroExtent.  Helps greatly with skimming a level for optimization opportunities.
+							FColor CollisionColor = FColor(157, 149, 223, 255);
+
+							const bool bPerHullColor = false;
+							const bool bDrawSimpleSolid = false;
+							BodySetup->AggGeom.GetAggGeom(GeomTransform, GetSelectionColor(CollisionColor, bDrawWireSelected, IsHovered()), CollisionMaterialInstance, bPerHullColor, bDrawSimpleSolid, UseEditorDepthTest(), ViewIndex, Collector);
+						}
+					}
+				}
+
+				// Draw the bounds
+				RenderBounds(PDI, View->Family->EngineShowFlags, GetBounds(), IsSelected());
+
+#if WITH_EDITOR
+				// Draw the debug outline
+				if (View->Family->EngineShowFlags.Grid)
+				{
+					const uint8 DPG = SDPG_Foreground;//GetDepthPriorityGroup(View);
+
+					// Draw separation wires if selected
+					FLinearColor OverrideColor;
+					bool bUseOverrideColor = false;
+
+					const bool bShowAsSelected = !(GIsEditor && View->Family->EngineShowFlags.Selection) || IsSelected();
+					if (bShowAsSelected || IsHovered())
+					{
+						bUseOverrideColor = true;
+						OverrideColor = GetSelectionColor(FLinearColor::White, bShowAsSelected, IsHovered());
+					}
+
+					FTransform LocalToWorld(GetLocalToWorld());
+
+					const float TW = TileMap->TileWidth;
+					const float TH = TileMap->TileHeight;
+					//@TODO: PAPER: Tiles probably shouldn't be drawn with their pivot at the center - RE: all the 0.5f in here
+
+					if (bUseOverrideColor)
+					{
+						// Draw horizontal lines
+						for (int32 Y = 0; Y <= TileMap->MapHeight; ++Y)
+						{
+							int32 X = 0;
+							const FVector Start((X - 0.5f) * TW, 0.0f, -(Y - 0.5f) * TH);
+
+							X = TileMap->MapWidth;
+							const FVector End((X - 0.5f) * TW, 0.0f, -(Y - 0.5f) * TH);
+
+							PDI->DrawLine(LocalToWorld.TransformPosition(Start), LocalToWorld.TransformPosition(End), OverrideColor, DPG);
+						}
+
+						// Draw vertical lines
+						for (int32 X = 0; X <= TileMap->MapWidth; ++X)
+						{
+							int32 Y = 0;
+							const FVector Start((X - 0.5f) * TW, 0.0f, -(Y - 0.5f) * TH);
+
+							Y = TileMap->MapHeight;
+							const FVector End((X - 0.5f) * TW, 0.0f, -(Y - 0.5f) * TH);
+
+							PDI->DrawLine(LocalToWorld.TransformPosition(Start), LocalToWorld.TransformPosition(End), OverrideColor, DPG);
+						}
+					}
+
+
+					// Create a local space bounding box
+					const FVector TopLeft(-TW*0.5f, -2.0f, TH*0.5f);
+					const FVector Dimensions(TileMap->MapWidth * TW, 4.0f, -TileMap->MapHeight * TH);
+					const FBox OutlineBox(TopLeft, TopLeft + Dimensions);
+
+					// Draw it		
+					DrawWireBox(PDI, OutlineBox.TransformBy(LocalToWorld), FLinearColor::White, DPG); //@TODO: Paper: Doesn't handle rotation well - need to use DrawOrientedWireBox
+				}
+#endif
+			}
+		}
+	}
+
+	// Draw all of the queued up sprites
+	FPaperRenderSceneProxy::GetDynamicMeshElements(Views, ViewFamily, VisibilityMap, Collector);
+}
+
+
 void FPaperTileMapRenderSceneProxy::DrawDynamicElements(FPrimitiveDrawInterface* PDI, const FSceneView* View)
 {
 	// Draw the tile maps
@@ -74,7 +211,7 @@ void FPaperTileMapRenderSceneProxy::DrawDynamicElements(FPrimitiveDrawInterface*
 
 					const bool bPerHullColor = false;
 					const bool bDrawSimpleSolid = false;
-					BodySetup->AggGeom.DrawAggGeom(PDI, GeomTransform, GetSelectionColor(CollisionColor, bDrawWireSelected, IsHovered()), &CollisionMaterialInstance, bPerHullColor, bDrawSimpleSolid);
+					BodySetup->AggGeom.DrawAggGeom(PDI, GeomTransform, GetSelectionColor(CollisionColor, bDrawWireSelected, IsHovered()), &CollisionMaterialInstance, bPerHullColor, bDrawSimpleSolid, UseEditorDepthTest());
 				}
 			}
 		}
diff --git a/Engine/Plugins/Experimental/Paper2D/Source/Paper2D/Private/PaperTileMapRenderSceneProxy.h b/Engine/Plugins/Experimental/Paper2D/Source/Paper2D/Private/PaperTileMapRenderSceneProxy.h
--- a/Engine/Plugins/Experimental/Paper2D/Source/Paper2D/Private/PaperTileMapRenderSceneProxy.h
+++ b/Engine/Plugins/Experimental/Paper2D/Source/Paper2D/Private/PaperTileMapRenderSceneProxy.h
@@ -13,6 +13,7 @@ public:
 	FPaperTileMapRenderSceneProxy(const UPaperTileMapRenderComponent* InComponent);
 
 	// FPrimitiveSceneProxy interface.
+	virtual void GetDynamicMeshElements(const TArray<const FSceneView*>& Views, const FSceneViewFamily& ViewFamily, uint32 VisibilityMap, FMeshElementCollector& Collector) const override;
 	virtual void DrawDynamicElements(FPrimitiveDrawInterface* PDI, const FSceneView* View) override;
 	// End of FPrimitiveSceneProxy interface.
 
diff --git a/Engine/Plugins/Experimental/Paper2D/Source/Paper2D/Private/Terrain/PaperTerrainComponent.cpp b/Engine/Plugins/Experimental/Paper2D/Source/Paper2D/Private/Terrain/PaperTerrainComponent.cpp
--- a/Engine/Plugins/Experimental/Paper2D/Source/Paper2D/Private/Terrain/PaperTerrainComponent.cpp
+++ b/Engine/Plugins/Experimental/Paper2D/Source/Paper2D/Private/Terrain/PaperTerrainComponent.cpp
@@ -37,6 +37,7 @@ protected:
 protected:
 	// FPaperRenderSceneProxy interface
 	virtual void DrawDynamicElements_RichMesh(FPrimitiveDrawInterface* PDI, const FSceneView* View, bool bUseOverrideColor, const FLinearColor& OverrideColor) override;
+	virtual void GetDynamicMeshElementsForView(const FSceneView* View, int32 ViewIndex, bool bUseOverrideColor, const FLinearColor& OverrideColor, FMeshElementCollector& Collector) const override;
 	// End of FPaperRenderSceneProxy interface
 };
 
@@ -64,6 +65,17 @@ void FPaperTerrainSceneProxy::DrawDynamicElements_RichMesh(FPrimitiveDrawInterfa
 	}
 }
 
+void FPaperTerrainSceneProxy::GetDynamicMeshElementsForView(const FSceneView* View, int32 ViewIndex, bool bUseOverrideColor, const FLinearColor& OverrideColor, FMeshElementCollector& Collector) const
+{
+	for (const FPaperTerrainMaterialPair& Batch : DrawingData)
+	{
+		if (Batch.Material != nullptr)
+		{
+			GetBatchMesh(View, bUseOverrideColor, OverrideColor, Batch.Material, Batch.Records, ViewIndex, Collector);
+		}
+	}
+}
+
 //////////////////////////////////////////////////////////////////////////
 // UPaperTerrainComponent
 
diff --git a/Engine/Plugins/Runtime/CableComponent/Source/CableComponent/Private/CableComponent.cpp b/Engine/Plugins/Runtime/CableComponent/Source/CableComponent/Private/CableComponent.cpp
--- a/Engine/Plugins/Runtime/CableComponent/Source/CableComponent/Private/CableComponent.cpp
+++ b/Engine/Plugins/Runtime/CableComponent/Source/CableComponent/Private/CableComponent.cpp
@@ -130,17 +130,17 @@ public:
 		}
 	}
 
-	int32 GetRequiredVertexCount()
+	int32 GetRequiredVertexCount() const
 	{
 		return (NumSegments + 1) * (NumSides + 1);
 	}
 
-	int32 GetRequiredIndexCount()
+	int32 GetRequiredIndexCount() const
 	{
 		return (NumSegments * NumSides * 2) * 3;
 	}
 
-	int32 GetVertIndex(int32 AlongIdx, int32 AroundIdx)
+	int32 GetVertIndex(int32 AlongIdx, int32 AroundIdx) const
 	{
 		return (AlongIdx * (NumSides+1)) + AroundIdx;
 	}
@@ -242,6 +242,60 @@ public:
 		RHIUnlockIndexBuffer(IndexBuffer.IndexBufferRHI);
 	}
 
+	virtual void GetDynamicMeshElements(const TArray<const FSceneView*>& Views, const FSceneViewFamily& ViewFamily, uint32 VisibilityMap, FMeshElementCollector& Collector) const override
+	{
+		QUICK_SCOPE_CYCLE_COUNTER( STAT_CableSceneProxy_GetDynamicMeshElements );
+
+		const bool bWireframe = AllowDebugViewmodes() && ViewFamily.EngineShowFlags.Wireframe;
+
+		auto WireframeMaterialInstance = new FColoredMaterialRenderProxy(
+			GEngine->WireframeMaterial ? GEngine->WireframeMaterial->GetRenderProxy(IsSelected()) : NULL,
+			FLinearColor(0, 0.5f, 1.f)
+			);
+
+		Collector.RegisterOneFrameMaterialProxy(WireframeMaterialInstance);
+
+		FMaterialRenderProxy* MaterialProxy = NULL;
+		if(bWireframe)
+		{
+			MaterialProxy = WireframeMaterialInstance;
+		}
+		else
+		{
+			MaterialProxy = Material->GetRenderProxy(IsSelected());
+		}
+
+		for (int32 ViewIndex = 0; ViewIndex < Views.Num(); ViewIndex++)
+		{
+			if (VisibilityMap & (1 << ViewIndex))
+			{
+				const FSceneView* View = Views[ViewIndex];
+				// Draw the mesh.
+				FMeshBatch& Mesh = Collector.AllocateMesh();
+				FMeshBatchElement& BatchElement = Mesh.Elements[0];
+				BatchElement.IndexBuffer = &IndexBuffer;
+				Mesh.bWireframe = bWireframe;
+				Mesh.VertexFactory = &VertexFactory;
+				Mesh.MaterialRenderProxy = MaterialProxy;
+				BatchElement.PrimitiveUniformBuffer = CreatePrimitiveUniformBufferImmediate(GetLocalToWorld(), GetBounds(), GetLocalBounds(), true, UseEditorDepthTest());
+				BatchElement.FirstIndex = 0;
+				BatchElement.NumPrimitives = GetRequiredIndexCount()/3;
+				BatchElement.MinVertexIndex = 0;
+				BatchElement.MaxVertexIndex = GetRequiredVertexCount();
+				Mesh.ReverseCulling = IsLocalToWorldDeterminantNegative();
+				Mesh.Type = PT_TriangleList;
+				Mesh.DepthPriorityGroup = SDPG_World;
+				Mesh.bCanApplyViewModeOverrides = false;
+				Collector.AddMesh(ViewIndex, Mesh);
+
+#if !(UE_BUILD_SHIPPING || UE_BUILD_TEST)
+				// Render bounds
+				RenderBounds(Collector.GetPDI(ViewIndex), ViewFamily.EngineShowFlags, GetBounds(), IsSelected());
+#endif
+			}
+		}
+	}
+
 	virtual void DrawDynamicElements(FPrimitiveDrawInterface* PDI,const FSceneView* View)
 	{
 		QUICK_SCOPE_CYCLE_COUNTER( STAT_CableSceneProxy_DrawDynamicElements );
@@ -270,7 +324,7 @@ public:
 		Mesh.bWireframe = bWireframe;
 		Mesh.VertexFactory = &VertexFactory;
 		Mesh.MaterialRenderProxy = MaterialProxy;
-		BatchElement.PrimitiveUniformBuffer = CreatePrimitiveUniformBufferImmediate(GetLocalToWorld(), GetBounds(), GetLocalBounds(), true);
+		BatchElement.PrimitiveUniformBuffer = CreatePrimitiveUniformBufferImmediate(GetLocalToWorld(), GetBounds(), GetLocalBounds(), true, UseEditorDepthTest());
 		BatchElement.FirstIndex = 0;
 		BatchElement.NumPrimitives = GetRequiredIndexCount()/3;
 		BatchElement.MinVertexIndex = 0;
diff --git a/Engine/Plugins/Runtime/CustomMeshComponent/Source/CustomMeshComponent/Private/CustomMeshComponent.cpp b/Engine/Plugins/Runtime/CustomMeshComponent/Source/CustomMeshComponent/Private/CustomMeshComponent.cpp
--- a/Engine/Plugins/Runtime/CustomMeshComponent/Source/CustomMeshComponent/Private/CustomMeshComponent.cpp
+++ b/Engine/Plugins/Runtime/CustomMeshComponent/Source/CustomMeshComponent/Private/CustomMeshComponent.cpp
@@ -142,6 +142,55 @@ public:
 		VertexFactory.ReleaseResource();
 	}
 
+	virtual void GetDynamicMeshElements(const TArray<const FSceneView*>& Views, const FSceneViewFamily& ViewFamily, uint32 VisibilityMap, FMeshElementCollector& Collector) const override
+	{
+		QUICK_SCOPE_CYCLE_COUNTER( STAT_CustomMeshSceneProxy_GetDynamicMeshElements );
+
+		const bool bWireframe = AllowDebugViewmodes() && ViewFamily.EngineShowFlags.Wireframe;
+
+		auto WireframeMaterialInstance = new FColoredMaterialRenderProxy(
+			GEngine->WireframeMaterial ? GEngine->WireframeMaterial->GetRenderProxy(IsSelected()) : NULL,
+			FLinearColor(0, 0.5f, 1.f)
+			);
+
+		Collector.RegisterOneFrameMaterialProxy(WireframeMaterialInstance);
+
+		FMaterialRenderProxy* MaterialProxy = NULL;
+		if(bWireframe)
+		{
+			MaterialProxy = WireframeMaterialInstance;
+		}
+		else
+		{
+			MaterialProxy = Material->GetRenderProxy(IsSelected());
+		}
+
+		for (int32 ViewIndex = 0; ViewIndex < Views.Num(); ViewIndex++)
+		{
+			if (VisibilityMap & (1 << ViewIndex))
+			{
+				const FSceneView* View = Views[ViewIndex];
+				// Draw the mesh.
+				FMeshBatch& Mesh = Collector.AllocateMesh();
+				FMeshBatchElement& BatchElement = Mesh.Elements[0];
+				BatchElement.IndexBuffer = &IndexBuffer;
+				Mesh.bWireframe = bWireframe;
+				Mesh.VertexFactory = &VertexFactory;
+				Mesh.MaterialRenderProxy = MaterialProxy;
+				BatchElement.PrimitiveUniformBuffer = CreatePrimitiveUniformBufferImmediate(GetLocalToWorld(), GetBounds(), GetLocalBounds(), true, UseEditorDepthTest());
+				BatchElement.FirstIndex = 0;
+				BatchElement.NumPrimitives = IndexBuffer.Indices.Num() / 3;
+				BatchElement.MinVertexIndex = 0;
+				BatchElement.MaxVertexIndex = VertexBuffer.Vertices.Num() - 1;
+				Mesh.ReverseCulling = IsLocalToWorldDeterminantNegative();
+				Mesh.Type = PT_TriangleList;
+				Mesh.DepthPriorityGroup = SDPG_World;
+				Mesh.bCanApplyViewModeOverrides = false;
+				Collector.AddMesh(ViewIndex, Mesh);
+			}
+		}
+	}
+
 	virtual void DrawDynamicElements(FPrimitiveDrawInterface* PDI,const FSceneView* View)
 	{
 		QUICK_SCOPE_CYCLE_COUNTER( STAT_CustomMeshSceneProxy_DrawDynamicElements );
@@ -170,7 +219,7 @@ public:
 		Mesh.bWireframe = bWireframe;
 		Mesh.VertexFactory = &VertexFactory;
 		Mesh.MaterialRenderProxy = MaterialProxy;
-		BatchElement.PrimitiveUniformBuffer = CreatePrimitiveUniformBufferImmediate(GetLocalToWorld(), GetBounds(), GetLocalBounds(), true);
+		BatchElement.PrimitiveUniformBuffer = CreatePrimitiveUniformBufferImmediate(GetLocalToWorld(), GetBounds(), GetLocalBounds(), true, UseEditorDepthTest());
 		BatchElement.FirstIndex = 0;
 		BatchElement.NumPrimitives = IndexBuffer.Indices.Num() / 3;
 		BatchElement.MinVertexIndex = 0;
diff --git a/Engine/Shaders/BasePassPixelShader.usf b/Engine/Shaders/BasePassPixelShader.usf
--- a/Engine/Shaders/BasePassPixelShader.usf
+++ b/Engine/Shaders/BasePassPixelShader.usf
@@ -454,7 +454,7 @@ float ClipForEditorPrimitives(FMaterialPixelParameters MaterialParameters)
 #if EDITOR_PRIMITIVE_MATERIAL && (FEATURE_LEVEL >= FEATURE_LEVEL_SM4 || ES2_EMULATION)
 	// Depth test manually if compositing editor primitives since the depth buffer is different (MSAA only)
 	BRANCH
-	if (bEnableEditorPrimitiveDepthTest)
+	if (Primitive.UseEditorDepthTest)
 	{
 		bool bIsPerspective = (View.ViewToClip._m33 < 1.0f);
  		
diff --git a/Engine/Source/Developer/GameplayDebugger/Private/GameplayDebuggingComponent.cpp b/Engine/Source/Developer/GameplayDebugger/Private/GameplayDebuggingComponent.cpp
--- a/Engine/Source/Developer/GameplayDebugger/Private/GameplayDebuggingComponent.cpp
+++ b/Engine/Source/Developer/GameplayDebugger/Private/GameplayDebuggingComponent.cpp
@@ -48,6 +48,14 @@ public:
 		}
 	}
 
+	virtual void GetDynamicMeshElements(const TArray<const FSceneView*>& Views, const FSceneViewFamily& ViewFamily, uint32 VisibilityMap, FMeshElementCollector& Collector) const override
+	{
+		for (int32 Index = 0; Index < ChildProxies.Num(); ++Index)
+		{
+			ChildProxies[Index]->GetDynamicMeshElements(Views, ViewFamily, VisibilityMap, Collector);
+		}
+	}
+
 	virtual void DrawDynamicElements(FPrimitiveDrawInterface* PDI, const FSceneView* View) override
 	{
 		for (int32 Index = 0; Index < ChildProxies.Num(); ++Index)
diff --git a/Engine/Source/Editor/UnrealEd/Private/Animation/DebugSkelMeshComponent.cpp b/Engine/Source/Editor/UnrealEd/Private/Animation/DebugSkelMeshComponent.cpp
--- a/Engine/Source/Editor/UnrealEd/Private/Animation/DebugSkelMeshComponent.cpp
+++ b/Engine/Source/Editor/UnrealEd/Private/Animation/DebugSkelMeshComponent.cpp
@@ -28,13 +28,36 @@ public:
 		FSkeletalMeshSceneProxy( InComponent, InSkelMeshResource )
 	{
 		SkeletalMeshComponent = InComponent;
-		WireframeOverlayColor = InWireframeOverlayColor;
+		WireframeColor = FLinearColor(InWireframeOverlayColor);
 	}
 
 	virtual ~FDebugSkelMeshSceneProxy()
 	{
 	}
 
+	virtual void GetDynamicMeshElements(const TArray<const FSceneView*>& Views, const FSceneViewFamily& ViewFamily, uint32 VisibilityMap, FMeshElementCollector& Collector) const override
+	{
+		if (SkeletalMeshComponent->bDrawMesh)
+		{
+			// We don't want to draw the mesh geometry to the hit testing render target
+			// so that we can get to triangle strips that are partially obscured by other
+			// triangle strips easier.
+			const bool bSelectable = false;
+			GetMeshElementsConditionallySelectable(Views, ViewFamily, bSelectable, VisibilityMap, Collector);
+		}
+
+		//@todo - the rendering thread should never read from UObjects directly!  These are race conditions, the properties should be mirrored on the proxy
+		if( SkeletalMeshComponent->MeshObject && (SkeletalMeshComponent->bDrawNormals || SkeletalMeshComponent->bDrawTangents || SkeletalMeshComponent->bDrawBinormals) )
+		{
+			for (int32 ViewIndex = 0; ViewIndex < Views.Num(); ViewIndex++)
+			{
+				if (VisibilityMap & (1 << ViewIndex))
+				{
+					SkeletalMeshComponent->MeshObject->DrawVertexElements(Collector.GetPDI(ViewIndex), SkeletalMeshComponent->ComponentToWorld, SkeletalMeshComponent->bDrawNormals, SkeletalMeshComponent->bDrawTangents, SkeletalMeshComponent->bDrawBinormals);
+				}
+			}
+		}
+	}
 
 	/** 
 	* Draw the scene proxy as a dynamic element
diff --git a/Engine/Source/Runtime/Engine/Classes/Components/ModelComponent.h b/Engine/Source/Runtime/Engine/Classes/Components/ModelComponent.h
--- a/Engine/Source/Runtime/Engine/Classes/Components/ModelComponent.h
+++ b/Engine/Source/Runtime/Engine/Classes/Components/ModelComponent.h
@@ -75,7 +75,6 @@ public:
 	virtual bool ShouldRecreateProxyOnUpdateTransform() const override;
 #if WITH_EDITOR
 	virtual void GetStaticLightingInfo(FStaticLightingPrimitiveInfo& OutPrimitiveInfo,const TArray<ULightComponent*>& InRelevantLights,const FLightingBuildOptions& Options) override;
-	virtual bool SetupLightmapResolutionViewInfo(FPrimitiveSceneProxy& Proxy) const override;
 #endif
 	virtual ELightMapInteractionType GetStaticLightingType() const override	{ return LMIT_Texture;	}
 	virtual void GetStreamingTextureInfo(TArray<FStreamingTexturePrimitiveInfo>& OutStreamingTextures) const override;
diff --git a/Engine/Source/Runtime/Engine/Classes/Components/PrimitiveComponent.h b/Engine/Source/Runtime/Engine/Classes/Components/PrimitiveComponent.h
--- a/Engine/Source/Runtime/Engine/Classes/Components/PrimitiveComponent.h
+++ b/Engine/Source/Runtime/Engine/Classes/Components/PrimitiveComponent.h
@@ -1143,17 +1143,6 @@ public:
 	
 	virtual bool GetShadowIndirectOnly() const { return false; }
 
-#if WITH_EDITOR
-	/**
-	 *	Setup the information required for rendering LightMap Density mode
-	 *	for this component.
-	 *
-	 *	@param	Proxy		The scene proxy for the component (information is set on it)
-	 *	@return	bool		true if successful, false if not.
-	 */
-	virtual bool SetupLightmapResolutionViewInfo(FPrimitiveSceneProxy& Proxy) const;
-#endif // WITH_EDITOR
-
 	/**
 	 *	Set the angular velocity of all bodies in this component.
 	 *
diff --git a/Engine/Source/Runtime/Engine/Classes/PhysicsEngine/AggregateGeometry2D.h b/Engine/Source/Runtime/Engine/Classes/PhysicsEngine/AggregateGeometry2D.h
--- a/Engine/Source/Runtime/Engine/Classes/PhysicsEngine/AggregateGeometry2D.h
+++ b/Engine/Source/Runtime/Engine/Classes/PhysicsEngine/AggregateGeometry2D.h
@@ -129,7 +129,7 @@ public:
 public:
 	void EmptyElements();
 
-	void DrawAggGeom(class FPrimitiveDrawInterface* PDI, const FTransform& Transform, const FColor Color, const FMaterialRenderProxy* MatInst, bool bPerHullColor, bool bDrawSolid) const;
+	void DrawAggGeom(class FPrimitiveDrawInterface* PDI, const FTransform& Transform, const FColor Color, const FMaterialRenderProxy* MatInst, bool bPerHullColor, bool bDrawSolid, bool bUseEditorDepthTest) const;
 
 	// Returns the area of all contained elements
  	float GetArea(const FVector& Scale) const;
diff --git a/Engine/Source/Runtime/Engine/Classes/PhysicsEngine/BodySetup.h b/Engine/Source/Runtime/Engine/Classes/PhysicsEngine/BodySetup.h
--- a/Engine/Source/Runtime/Engine/Classes/PhysicsEngine/BodySetup.h
+++ b/Engine/Source/Runtime/Engine/Classes/PhysicsEngine/BodySetup.h
@@ -80,8 +80,9 @@ struct FKConvexElem
 		, ConvexMeshNegX(NULL)
 		{}
 
-	ENGINE_API void	DrawElemWire(class FPrimitiveDrawInterface* PDI, const FTransform& ElemTM, const FColor Color) const;
-		void	AddCachedSolidConvexGeom(TArray<FDynamicMeshVertex>& VertexBuffer, TArray<int32>& IndexBuffer, const FColor VertexColor);
+		ENGINE_API void	DrawElemWire(class FPrimitiveDrawInterface* PDI, const FTransform& ElemTM, const FColor Color) const;
+
+		void AddCachedSolidConvexGeom(TArray<FDynamicMeshVertex>& VertexBuffer, TArray<int32>& IndexBuffer, const FColor VertexColor) const;
 
 		/** Reset the hull to empty all arrays */
 		ENGINE_API void	Reset();
@@ -167,6 +168,7 @@ struct FKSphereElem
 	
 	ENGINE_API void	DrawElemWire(class FPrimitiveDrawInterface* PDI, const FTransform& ElemTM, float Scale, const FColor Color) const;
 	ENGINE_API void	DrawElemSolid(class FPrimitiveDrawInterface* PDI, const FTransform& ElemTM, float Scale, const FMaterialRenderProxy* MaterialRenderProxy) const;
+	ENGINE_API void GetElemSolid(const FTransform& ElemTM, float Scale, const FMaterialRenderProxy* MaterialRenderProxy, int32 ViewIndex, class FMeshElementCollector& Collector) const;
 	ENGINE_API FBox CalcAABB(const FTransform& BoneTM, float Scale) const;
 
 	ENGINE_API void ScaleElem(FVector DeltaSize, float MinSize);
@@ -251,6 +253,7 @@ struct FKBoxElem
 
 	ENGINE_API void	DrawElemWire(class FPrimitiveDrawInterface* PDI, const FTransform& ElemTM, float Scale, const FColor Color) const;
 	ENGINE_API void	DrawElemSolid(class FPrimitiveDrawInterface* PDI, const FTransform& ElemTM, float Scale, const FMaterialRenderProxy* MaterialRenderProxy) const;
+	ENGINE_API void GetElemSolid(const FTransform& ElemTM, float Scale, const FMaterialRenderProxy* MaterialRenderProxy, int32 ViewIndex, FMeshElementCollector& Collector) const;
 	ENGINE_API FBox CalcAABB(const FTransform& BoneTM, float Scale) const;
 
 	ENGINE_API void ScaleElem(FVector DeltaSize, float MinSize);
@@ -322,6 +325,7 @@ struct FKSphylElem
 	FORCEINLINE float GetVolume(const FVector& Scale) const { float ScaledRadius = Radius * Scale.GetMin(); return PI * FMath::Square(ScaledRadius) * ( 1.3333f * ScaledRadius + (Length * Scale.GetMin())); }
 
 	ENGINE_API void	DrawElemWire(class FPrimitiveDrawInterface* PDI, const FTransform& ElemTM, float Scale, const FColor Color) const;
+	ENGINE_API void GetElemSolid(const FTransform& ElemTM, float Scale, const FMaterialRenderProxy* MaterialRenderProxy, int32 ViewIndex, FMeshElementCollector& Collector) const;
 	ENGINE_API void	DrawElemSolid(class FPrimitiveDrawInterface* PDI, const FTransform& ElemTM, float Scale, const FMaterialRenderProxy* MaterialRenderProxy) const;
 	ENGINE_API FBox CalcAABB(const FTransform& BoneTM, float Scale) const;
 
@@ -372,7 +376,9 @@ struct ENGINE_API FKAggregateGeom
 
 	void Serialize( const FArchive& Ar );
 
-	void DrawAggGeom(class FPrimitiveDrawInterface* PDI, const FTransform& Transform, const FColor Color, const FMaterialRenderProxy* MatInst, bool bPerHullColor, bool bDrawSolid);
+	void DrawAggGeom(class FPrimitiveDrawInterface* PDI, const FTransform& Transform, const FColor Color, const FMaterialRenderProxy* MatInst, bool bPerHullColor, bool bDrawSolid, bool bUseEditorDepthTest);
+
+	void GetAggGeom(const FTransform& Transform, const FColor Color, const FMaterialRenderProxy* MatInst, bool bPerHullColor, bool bDrawSolid, bool bUseEditorDepthTest, int32 ViewIndex, class FMeshElementCollector& Collector) const;
 
 	/** Release the RenderInfo (if its there) and safely clean up any resources. Call on the game thread. */
 	void FreeRenderInfo();
diff --git a/Engine/Source/Runtime/Engine/Private/AI/Navigation/NavLinkRenderingComponent.cpp b/Engine/Source/Runtime/Engine/Private/AI/Navigation/NavLinkRenderingComponent.cpp
--- a/Engine/Source/Runtime/Engine/Private/AI/Navigation/NavLinkRenderingComponent.cpp
+++ b/Engine/Source/Runtime/Engine/Private/AI/Navigation/NavLinkRenderingComponent.cpp
@@ -137,6 +137,38 @@ void FNavLinkRenderingProxy::StoreSegmentLinks(const FTransform& LocalToWorld, c
 	}
 }
 
+void FNavLinkRenderingProxy::GetDynamicMeshElements(const TArray<const FSceneView*>& Views, const FSceneViewFamily& ViewFamily, uint32 VisibilityMap, FMeshElementCollector& Collector) const 
+{
+	if (LinkOwnerActor && LinkOwnerActor->GetWorld())
+	{
+		const UNavigationSystem* NavSys = LinkOwnerActor->GetWorld()->GetNavigationSystem();
+		TArray<float> StepHeights;
+		if (NavSys != NULL)
+		{
+			StepHeights.Reserve(NavSys->NavDataSet.Num());
+			for(int32 DataIndex = 0; DataIndex < NavSys->NavDataSet.Num(); ++DataIndex)
+			{
+				const ARecastNavMesh* NavMesh = Cast<const ARecastNavMesh>(NavSys->NavDataSet[DataIndex]);
+				if (NavMesh != NULL && NavMesh->AgentMaxStepHeight > 0 && NavMesh->bEnableDrawing)
+				{
+					StepHeights.Add(NavMesh->AgentMaxStepHeight);
+				}
+			}
+		}
+
+		static const FColor RadiusColor(150, 160, 150, 48);
+		FMaterialRenderProxy* const MeshColorInstance = new(FMemStack::Get()) FColoredMaterialRenderProxy(GEngine->DebugMeshMaterial->GetRenderProxy(false), RadiusColor);
+
+		for (int32 ViewIndex = 0; ViewIndex < Views.Num(); ViewIndex++)
+		{
+			if (VisibilityMap & (1 << ViewIndex))
+			{
+				FNavLinkRenderingProxy::GetLinkMeshes(OffMeshPointLinks, OffMeshSegmentLinks, StepHeights, MeshColorInstance, ViewIndex, Collector);
+			}
+		}
+	}
+}
+
 void FNavLinkRenderingProxy::DrawDynamicElements(FPrimitiveDrawInterface* PDI,const FSceneView* View)
 {
 	QUICK_SCOPE_CYCLE_COUNTER( STAT_NavLinkRenderingProxy_DrawDynamicElements );
@@ -164,6 +196,89 @@ void FNavLinkRenderingProxy::DrawDynamicElements(FPrimitiveDrawInterface* PDI,co
 	}
 }
 
+void FNavLinkRenderingProxy::GetLinkMeshes(const TArray<FNavLinkDrawing>& OffMeshPointLinks, const TArray<FNavLinkSegmentDrawing>& OffMeshSegmentLinks, TArray<float>& StepHeights, FMaterialRenderProxy* const MeshColorInstance, int32 ViewIndex, FMeshElementCollector& Collector)
+{
+	static const FColor LinkColor(0,0,166);
+	static const float LinkArcThickness = 3.f;
+	static const float LinkArcHeight = 0.4f;
+
+	if (StepHeights.Num() == 0)
+	{
+		StepHeights.Add(FNavigationSystem::FallbackAgentHeight / 2);
+	}
+
+	FPrimitiveDrawInterface* PDI = Collector.GetPDI(ViewIndex);
+
+	for (int32 LinkIndex = 0; LinkIndex < OffMeshPointLinks.Num(); ++LinkIndex)
+	{
+		const FNavLinkDrawing& Link = OffMeshPointLinks[LinkIndex];
+		const uint32 Segments = FPlatformMath::Max<uint32>(LinkArcHeight*(Link.Right-Link.Left).Size()/10, 8);
+		DrawArc(PDI, Link.Left, Link.Right, LinkArcHeight, Segments, Link.Color, SDPG_World, 3.5f);
+		const FVector VOffset(0,0,FVector::Dist(Link.Left, Link.Right)*1.333f);
+
+		switch (Link.Direction)
+		{
+		case ENavLinkDirection::LeftToRight:
+			DrawArrowHead(PDI, Link.Right, Link.Left+VOffset, 30.f, Link.Color, SDPG_World, 3.5f);
+			break;
+		case ENavLinkDirection::RightToLeft:
+			DrawArrowHead(PDI, Link.Left, Link.Right+VOffset, 30.f, Link.Color, SDPG_World, 3.5f);
+			break;
+		case ENavLinkDirection::BothWays:
+		default:
+			DrawArrowHead(PDI, Link.Right, Link.Left+VOffset, 30.f, Link.Color, SDPG_World, 3.5f);
+			DrawArrowHead(PDI, Link.Left, Link.Right+VOffset, 30.f, Link.Color, SDPG_World, 3.5f);
+			break;
+		}
+
+		// draw snap-spheres on both ends
+		for (int32 StepHeightIndex = 0; StepHeightIndex < StepHeights.Num(); ++StepHeightIndex)
+		{
+			GetCylinderMesh(Link.Right, FVector(1,0,0), FVector(0,1,0), FVector(0,0,1), Link.SnapRadius, StepHeights[StepHeightIndex], 10, MeshColorInstance, SDPG_World, ViewIndex, Collector);
+			GetCylinderMesh(Link.Left, FVector(1,0,0), FVector(0,1,0), FVector(0,0,1), Link.SnapRadius, StepHeights[StepHeightIndex], 10, MeshColorInstance, SDPG_World, ViewIndex, Collector);
+		}
+	}
+
+	static const float SegmentArcHeight = 0.25f;
+	for (int32 LinkIndex = 0; LinkIndex < OffMeshSegmentLinks.Num(); ++LinkIndex)
+	{
+		const FNavLinkSegmentDrawing& Link = OffMeshSegmentLinks[LinkIndex];
+		const uint32 SegmentsStart = FPlatformMath::Max<uint32>(SegmentArcHeight*(Link.RightStart-Link.LeftStart).Size()/10, 8);
+		const uint32 SegmentsEnd = FPlatformMath::Max<uint32>(SegmentArcHeight*(Link.RightEnd-Link.LeftEnd).Size()/10, 8);
+		DrawArc(PDI, Link.LeftStart, Link.RightStart, SegmentArcHeight, SegmentsStart, Link.Color, SDPG_World, 3.5f);
+		DrawArc(PDI, Link.LeftEnd, Link.RightEnd, SegmentArcHeight, SegmentsEnd, Link.Color, SDPG_World, 3.5f);
+		const FVector VOffset(0,0,FVector::Dist(Link.LeftStart, Link.RightStart)*1.333f);
+
+		switch (Link.Direction)
+		{
+		case ENavLinkDirection::LeftToRight:
+			DrawArrowHead(PDI, Link.RightStart, Link.LeftStart+VOffset, 30.f, Link.Color, SDPG_World, 3.5f);
+			DrawArrowHead(PDI, Link.RightEnd, Link.LeftEnd+VOffset, 30.f, Link.Color, SDPG_World, 3.5f);
+			break;
+		case ENavLinkDirection::RightToLeft:
+			DrawArrowHead(PDI, Link.LeftStart, Link.RightStart+VOffset, 30.f, Link.Color, SDPG_World, 3.5f);
+			DrawArrowHead(PDI, Link.LeftEnd, Link.RightEnd+VOffset, 30.f, Link.Color, SDPG_World, 3.5f);
+			break;
+		case ENavLinkDirection::BothWays:
+		default:
+			DrawArrowHead(PDI, Link.RightStart, Link.LeftStart+VOffset, 30.f, Link.Color, SDPG_World, 3.5f);
+			DrawArrowHead(PDI, Link.RightEnd, Link.LeftEnd+VOffset, 30.f, Link.Color, SDPG_World, 3.5f);
+			DrawArrowHead(PDI, Link.LeftStart, Link.RightStart+VOffset, 30.f, Link.Color, SDPG_World, 3.5f);
+			DrawArrowHead(PDI, Link.LeftEnd, Link.RightEnd+VOffset, 30.f, Link.Color, SDPG_World, 3.5f);
+			break;
+		}
+
+		// draw snap-spheres on both ends
+		for (int32 StepHeightIndex = 0; StepHeightIndex < StepHeights.Num(); ++StepHeightIndex)
+		{
+			GetCylinderMesh(Link.RightStart, FVector(1,0,0), FVector(0,1,0), FVector(0,0,1), Link.SnapRadius, StepHeights[StepHeightIndex], 10, MeshColorInstance, SDPG_World, ViewIndex, Collector);
+			GetCylinderMesh(Link.RightEnd, FVector(1,0,0), FVector(0,1,0), FVector(0,0,1), Link.SnapRadius, StepHeights[StepHeightIndex], 10, MeshColorInstance, SDPG_World, ViewIndex, Collector);
+			GetCylinderMesh(Link.LeftStart, FVector(1,0,0), FVector(0,1,0), FVector(0,0,1), Link.SnapRadius, StepHeights[StepHeightIndex], 10, MeshColorInstance, SDPG_World, ViewIndex, Collector);
+			GetCylinderMesh(Link.LeftEnd, FVector(1,0,0), FVector(0,1,0), FVector(0,0,1), Link.SnapRadius, StepHeights[StepHeightIndex], 10, MeshColorInstance, SDPG_World, ViewIndex, Collector);
+		}
+	}
+}
+
 void FNavLinkRenderingProxy::DrawLinks(FPrimitiveDrawInterface* PDI, TArray<FNavLinkDrawing>& OffMeshPointLinks, TArray<FNavLinkSegmentDrawing>& OffMeshSegmentLinks, TArray<float>& StepHeights, FMaterialRenderProxy* const MeshColorInstance)
 {
 	static const FColor LinkColor(0,0,166);
diff --git a/Engine/Source/Runtime/Engine/Private/AI/Navigation/NavTestRenderingComponent.cpp b/Engine/Source/Runtime/Engine/Private/AI/Navigation/NavTestRenderingComponent.cpp
--- a/Engine/Source/Runtime/Engine/Private/AI/Navigation/NavTestRenderingComponent.cpp
+++ b/Engine/Source/Runtime/Engine/Private/AI/Navigation/NavTestRenderingComponent.cpp
@@ -78,6 +78,112 @@ public:
 		}
 	}
 
+	virtual void GetDynamicMeshElements(const TArray<const FSceneView*>& Views, const FSceneViewFamily& ViewFamily, uint32 VisibilityMap, FMeshElementCollector& Collector) const override
+	{
+		for (int32 ViewIndex = 0; ViewIndex < Views.Num(); ViewIndex++)
+		{
+			if (VisibilityMap & (1 << ViewIndex))
+			{
+				const FSceneView* View = Views[ViewIndex];
+				FPrimitiveDrawInterface* PDI = Collector.GetPDI(ViewIndex);
+
+				if (NavTestActor)
+				{
+					//DrawArc(PDI, Link.Left, Link.Right, 0.4f, NavMeshColors[Link.AreaID], SDPG_World, 3.5f);
+					//const FVector VOffset(0,0,FVector::Dist(Link.Left, Link.Right)*1.333f);
+					//DrawArrowHead(PDI, Link.Right, Link.Left+VOffset, 30.f, NavMeshColors[Link.AreaID], SDPG_World, 3.5f);
+
+					//@todo - the rendering thread should never read from UObjects directly!  These are race conditions, the properties should be mirrored on the proxy
+					const FVector ActorLocation = NavTestActor->GetActorLocation();
+					const FVector ProjectedLocation = NavTestActor->ProjectedLocation + NavMeshDrawOffset;
+					const FColor ProjectedColor = NavTestActor->bProjectedLocationValid ? FColor(0, 255, 0, 120) : FColor(255, 0, 0, 120);
+					const FVector BoxExtent(20, 20, 20);
+
+					FMaterialRenderProxy* const ColoredMeshInstance = new(FMemStack::Get()) FColoredMaterialRenderProxy(GEngine->DebugMeshMaterial->GetRenderProxy(false), ProjectedColor);
+					//DrawBox(PDI, FTransform(ProjectedLocation).ToMatrixNoScale(),BoxExtent, ColoredMeshInstance, SDPG_World);
+					GetSphereMesh(ProjectedLocation, BoxExtent, 10, 7, ColoredMeshInstance, SDPG_World, false, ViewIndex, Collector);
+
+					//DrawWireBox(PDI, FBox(ProjectedLocation-BoxExtent, ProjectedLocation+BoxExtent), ProjectedColor, false);
+					DrawWireBox(PDI, FBox(ActorLocation - BoxExtent, ActorLocation + BoxExtent), FColor::White, false);
+					const FVector LineEnd = ProjectedLocation - (ProjectedLocation - ActorLocation).SafeNormal()*BoxExtent.X;
+					PDI->DrawLine(LineEnd, ActorLocation, ProjectedColor, SDPG_World, 2.5);
+					DrawArrowHead(PDI, LineEnd, ActorLocation, 20.f, ProjectedColor, SDPG_World, 2.5f);
+
+					// draw query extent
+					DrawWireBox(PDI, FBox(ActorLocation - NavTestActor->QueryingExtent, ActorLocation + NavTestActor->QueryingExtent), FColor::Blue, false);
+				}
+
+				// draw path
+				if (!bShowBestPath || !NodeDebug.Num())
+				{
+					for (int32 PointIndex = 1; PointIndex < PathPoints.Num(); PointIndex++)
+					{
+						PDI->DrawLine(PathPoints[PointIndex-1], PathPoints[PointIndex], FLinearColor::Red, SDPG_World, 2.0f, 0.0f, true);
+					}
+				}
+
+				// draw path debug data
+				if (bShowNodePool)
+				{
+					if (ClosedSetIndices.Num())
+					{
+						const FColoredMaterialRenderProxy *MeshColorInstance = new(FMemStack::Get()) FColoredMaterialRenderProxy(GEngine->DebugMeshMaterial->GetRenderProxy(false), NavMeshRenderColor_ClosedSet);						
+						FDynamicMeshBuilder	MeshBuilder;
+						MeshBuilder.AddVertices(ClosedSetVerts);
+						MeshBuilder.AddTriangles(ClosedSetIndices);
+						MeshBuilder.GetMesh(FMatrix::Identity, MeshColorInstance, GetDepthPriorityGroup(View), false, false, ViewIndex, Collector);
+					}
+
+					if (OpenSetIndices.Num())
+					{
+						const FColoredMaterialRenderProxy *MeshColorInstance = new(FMemStack::Get()) FColoredMaterialRenderProxy(GEngine->DebugMeshMaterial->GetRenderProxy(false), NavMeshRenderColor_OpenSet);						
+						FDynamicMeshBuilder	MeshBuilder;
+						MeshBuilder.AddVertices(OpenSetVerts);
+						MeshBuilder.AddTriangles(OpenSetIndices);
+						MeshBuilder.GetMesh(FMatrix::Identity, MeshColorInstance, GetDepthPriorityGroup(View), false, false, ViewIndex, Collector);
+					}
+				}
+
+				for (TSet<FNodeDebugData>::TConstIterator It(NodeDebug); It; ++It)
+				{
+					const FNodeDebugData& NodeData = *It;
+
+					FColor LineColor(FColor::Blue);
+					if (bShowBestPath && NodeData.bBestPath)
+					{
+						LineColor = FColor::Red;
+					}
+
+					if (bShowDiff)
+					{
+						LineColor.A = NodeData.bModified ? NavMeshRenderAlpha_Modifed : NavMeshRenderAlpha_NonModified;
+					}
+
+					FVector ParentPos(NodeData.ParentId.IsValidId() ? NodeDebug[NodeData.ParentId].Position : NodeData.Position);
+
+					if (bShowDiff && !NodeData.bModified)
+					{
+						PDI->DrawLine(NodeData.Position, ParentPos, LineColor, SDPG_World);
+					}
+					else
+					{
+						PDI->DrawLine(NodeData.Position, ParentPos, LineColor, SDPG_World, 2.0f, 0.0, true);
+					}
+
+					if (NodeData.bOffMeshLink)
+					{
+						DrawWireBox(PDI, FBox::BuildAABB(NodeData.Position, FVector(10.0f)), LineColor, SDPG_World);
+					}
+
+					if (bShowDiff && NodeData.bModified)
+					{
+						PDI->DrawLine(NodeData.Position + FVector(0,0,10), NodeData.Position + FVector(0,0,100), FColor::Green, SDPG_World);
+					}
+				}
+			}
+		}
+	}
+
 	virtual void DrawDynamicElements(FPrimitiveDrawInterface* PDI,const FSceneView* View) override
 	{
 		if (NavTestActor)
diff --git a/Engine/Source/Runtime/Engine/Private/ArrowComponent.cpp b/Engine/Source/Runtime/Engine/Private/ArrowComponent.cpp
--- a/Engine/Source/Runtime/Engine/Private/ArrowComponent.cpp
+++ b/Engine/Source/Runtime/Engine/Private/ArrowComponent.cpp
@@ -152,6 +152,76 @@ public:
 
 	// FPrimitiveSceneProxy interface.
 	
+	virtual void GetDynamicMeshElements(const TArray<const FSceneView*>& Views, const FSceneViewFamily& ViewFamily, uint32 VisibilityMap, FMeshElementCollector& Collector) const override
+	{
+		QUICK_SCOPE_CYCLE_COUNTER( STAT_ArrowSceneProxy_DrawDynamicElements );
+
+		FMatrix EffectiveLocalToWorld;
+#if WITH_EDITOR
+		if (bLightAttachment)
+		{
+			EffectiveLocalToWorld = GetLocalToWorld().GetMatrixWithoutScale();
+		}
+		else
+#endif	//WITH_EDITOR
+		{
+			EffectiveLocalToWorld = GetLocalToWorld();
+		}
+
+		auto ArrowMaterialRenderProxy = new FColoredMaterialRenderProxy(
+			GEngine->ArrowMaterial->GetRenderProxy(IsSelected(), IsHovered()),
+			ArrowColor,
+			"GizmoColor"
+			);
+
+		Collector.RegisterOneFrameMaterialProxy(ArrowMaterialRenderProxy);
+
+		for (int32 ViewIndex = 0; ViewIndex < Views.Num(); ViewIndex++)
+		{
+			if (VisibilityMap & (1 << ViewIndex))
+			{
+				const FSceneView* View = Views[ViewIndex];
+
+				// Calculate the view-dependent scaling factor.
+				float ViewScale = 1.0f;
+				if (bIsScreenSizeScaled && (View->ViewMatrices.ProjMatrix.M[3][3] != 1.0f))
+				{
+					const float ZoomFactor = FMath::Min<float>(View->ViewMatrices.ProjMatrix.M[0][0], View->ViewMatrices.ProjMatrix.M[1][1]);
+					if (ZoomFactor != 0.0f)
+					{
+						const float Radius = View->WorldToScreen(Origin).W * (ScreenSize / ZoomFactor);
+						if (Radius < 1.0f && Radius > 0)
+						{
+							ViewScale *= Radius;
+						}
+					}
+				}
+
+		#if WITH_EDITORONLY_DATA
+				ViewScale *= EditorScale;
+		#endif
+
+				// Draw the mesh.
+				FMeshBatch& Mesh = Collector.AllocateMesh();
+				FMeshBatchElement& BatchElement = Mesh.Elements[0];
+				BatchElement.IndexBuffer = &IndexBuffer;
+				Mesh.bWireframe = false;
+				Mesh.VertexFactory = &VertexFactory;
+				Mesh.MaterialRenderProxy = ArrowMaterialRenderProxy;
+				BatchElement.PrimitiveUniformBuffer = CreatePrimitiveUniformBufferImmediate(FScaleMatrix(ViewScale) * EffectiveLocalToWorld, GetBounds(), GetLocalBounds(), true, UseEditorDepthTest());
+				BatchElement.FirstIndex = 0;
+				BatchElement.NumPrimitives = IndexBuffer.Indices.Num() / 3;
+				BatchElement.MinVertexIndex = 0;
+				BatchElement.MaxVertexIndex = VertexBuffer.Vertices.Num() - 1;
+				Mesh.ReverseCulling = IsLocalToWorldDeterminantNegative();
+				Mesh.Type = PT_TriangleList;
+				Mesh.DepthPriorityGroup = SDPG_World;
+				Mesh.bCanApplyViewModeOverrides = false;
+				Collector.AddMesh(ViewIndex, Mesh);
+			}
+		}
+	}
+
 	/** 
 	* Draw the scene proxy as a dynamic element
 	*
@@ -208,7 +278,7 @@ public:
 		Mesh.bWireframe = false;
 		Mesh.VertexFactory = &VertexFactory;
 		Mesh.MaterialRenderProxy = &ArrowMaterialRenderProxy;
-		BatchElement.PrimitiveUniformBuffer = CreatePrimitiveUniformBufferImmediate(FScaleMatrix(ViewScale) * EffectiveLocalToWorld, GetBounds(), GetLocalBounds(), true);
+		BatchElement.PrimitiveUniformBuffer = CreatePrimitiveUniformBufferImmediate(FScaleMatrix(ViewScale) * EffectiveLocalToWorld, GetBounds(), GetLocalBounds(), true, UseEditorDepthTest());
 		BatchElement.FirstIndex = 0;
 		BatchElement.NumPrimitives = IndexBuffer.Indices.Num() / 3;
 		BatchElement.MinVertexIndex = 0;
diff --git a/Engine/Source/Runtime/Engine/Private/BatchedElements.cpp b/Engine/Source/Runtime/Engine/Private/BatchedElements.cpp
--- a/Engine/Source/Runtime/Engine/Private/BatchedElements.cpp
+++ b/Engine/Source/Runtime/Engine/Private/BatchedElements.cpp
@@ -15,6 +15,18 @@ DEFINE_LOG_CATEGORY_STATIC(LogBatchedElements, Log, All);
 /** The simple element vertex declaration. */
 TGlobalResource<FSimpleElementVertexDeclaration> GSimpleElementVertexDeclaration;
 
+EBlendModeFilter::Type GetBlendModeFilter(ESimpleElementBlendMode BlendMode)
+{
+	if (BlendMode == SE_BLEND_Opaque || BlendMode == SE_BLEND_Masked || BlendMode == SE_BLEND_MaskedDistanceField || BlendMode == SE_BLEND_MaskedDistanceFieldShadowed)
+	{
+		return EBlendModeFilter::OpaqueAndMasked;
+	}
+	else
+	{
+		return EBlendModeFilter::Translucent;
+	}
+}
+
 void FBatchedElements::AddLine(const FVector& Start, const FVector& End, const FLinearColor& Color, FHitProxyId HitProxyId, float Thickness, float DepthBias, bool bScreenSpace)
 {
 	// Ensure the line isn't masked out.  Some legacy code relies on Color.A being ignored.
@@ -596,7 +608,7 @@ void FBatchedElements::DrawPointElements(FRHICommandList& RHICmdList, const FMat
 }
 
 
-bool FBatchedElements::Draw(FRHICommandList& RHICmdList, bool bNeedToSwitchVerticalAxis, const FMatrix& Transform, uint32 ViewportSizeX, uint32 ViewportSizeY, bool bHitTesting, float Gamma, const FSceneView* View, FTexture2DRHIRef DepthTexture) const
+bool FBatchedElements::Draw(FRHICommandList& RHICmdList, bool bNeedToSwitchVerticalAxis, const FMatrix& Transform, uint32 ViewportSizeX, uint32 ViewportSizeY, bool bHitTesting, float Gamma, const FSceneView* View, FTexture2DRHIRef DepthTexture, EBlendModeFilter::Type Filter) const
 {
 	if( HasPrimsToDraw() )
 	{
@@ -607,7 +619,8 @@ bool FBatchedElements::Draw(FRHICommandList& RHICmdList, bool bNeedToSwitchVerti
 
 		RHICmdList.SetRasterizerState(TStaticRasterizerState<FM_Solid, CM_None>::GetRHI());
 
-		if( LineVertices.Num() > 0 || Points.Num() > 0 || ThickLines.Num() > 0 || WireTris.Num() > 0 )
+		if( (LineVertices.Num() > 0 || Points.Num() > 0 || ThickLines.Num() > 0 || WireTris.Num() > 0)
+			&& (Filter & EBlendModeFilter::OpaqueAndMasked))
 		{
 			// Lines/points don't support batched element parameters (yet!)
 			FBatchedElementParameters* BatchedElementParameters = NULL;
@@ -819,47 +832,59 @@ bool FBatchedElements::Draw(FRHICommandList& RHICmdList, bool bNeedToSwitchVerti
 			for(int32 SpriteIndex = 0;SpriteIndex < SpriteCount;SpriteIndex++)
 			{
 				const FBatchedSprite& Sprite = Sprites[SpriteIndex];
-				if (CurrentTexture != Sprite.Texture || CurrentBlendMode != Sprite.BlendMode)
-				{
-					//New batch, draw previous and clear
-					const int32 VertexCount = SpriteList.Num();
-					const int32 PrimCount = VertexCount / 3;
-					PrepareShaders(RHICmdList, CurrentBlendMode, Transform, bNeedToSwitchVerticalAxis, BatchedElementParameters, CurrentTexture, bHitTesting, Gamma, NULL, View, DepthTexture);
-					DrawPrimitiveUP(RHICmdList, PT_TriangleList, PrimCount, SpriteList.GetTypedData(), sizeof(FSimpleElementVertex));
+				const EBlendModeFilter::Type SpriteFilter = GetBlendModeFilter((ESimpleElementBlendMode)Sprite.BlendMode);
 
-					SpriteList.Empty(6);
-					CurrentTexture = Sprite.Texture;
-					CurrentBlendMode = (ESimpleElementBlendMode)Sprite.BlendMode;
-				}
+				// Only render blend modes in the filter
+				if (Filter & SpriteFilter)
+				{
+					if (CurrentTexture != Sprite.Texture || CurrentBlendMode != Sprite.BlendMode)
+					{
+						//New batch, draw previous and clear
+						const int32 VertexCount = SpriteList.Num();
+						const int32 PrimCount = VertexCount / 3;
+						PrepareShaders(RHICmdList, CurrentBlendMode, Transform, bNeedToSwitchVerticalAxis, BatchedElementParameters, CurrentTexture, bHitTesting, Gamma, NULL, View, DepthTexture);
+						DrawPrimitiveUP(RHICmdList, PT_TriangleList, PrimCount, SpriteList.GetTypedData(), sizeof(FSimpleElementVertex));
+
+						SpriteList.Empty(6);
+						CurrentTexture = Sprite.Texture;
+						CurrentBlendMode = (ESimpleElementBlendMode)Sprite.BlendMode;
+					}
 
-				int32 SpriteListIndex = SpriteList.AddUninitialized(6);
-				FSimpleElementVertex* Vertex = SpriteList.GetTypedData();
+					int32 SpriteListIndex = SpriteList.AddUninitialized(6);
+					FSimpleElementVertex* Vertex = SpriteList.GetTypedData();
 
-				// Compute the sprite vertices.
-				const FVector WorldSpriteX = CameraX * Sprite.SizeX;
-				const FVector WorldSpriteY = CameraY * -Sprite.SizeY * GProjectionSignY;
+					// Compute the sprite vertices.
+					const FVector WorldSpriteX = CameraX * Sprite.SizeX;
+					const FVector WorldSpriteY = CameraY * -Sprite.SizeY * GProjectionSignY;
 
-				const float UStart = Sprite.U/Sprite.Texture->GetSizeX();
-				const float UEnd = (Sprite.U + Sprite.UL)/Sprite.Texture->GetSizeX();
-				const float VStart = Sprite.V/Sprite.Texture->GetSizeY();
-				const float VEnd = (Sprite.V + Sprite.VL)/Sprite.Texture->GetSizeY();
+					const float UStart = Sprite.U/Sprite.Texture->GetSizeX();
+					const float UEnd = (Sprite.U + Sprite.UL)/Sprite.Texture->GetSizeX();
+					const float VStart = Sprite.V/Sprite.Texture->GetSizeY();
+					const float VEnd = (Sprite.V + Sprite.VL)/Sprite.Texture->GetSizeY();
 
-				Vertex[SpriteListIndex + 0] = FSimpleElementVertex(FVector4(Sprite.Position + WorldSpriteX - WorldSpriteY,1),FVector2D(UEnd,  VStart),Sprite.Color,Sprite.HitProxyId);
-				Vertex[SpriteListIndex + 1] = FSimpleElementVertex(FVector4(Sprite.Position + WorldSpriteX + WorldSpriteY,1),FVector2D(UEnd,  VEnd  ),Sprite.Color,Sprite.HitProxyId);
-				Vertex[SpriteListIndex + 2] = FSimpleElementVertex(FVector4(Sprite.Position - WorldSpriteX - WorldSpriteY,1),FVector2D(UStart,VStart),Sprite.Color,Sprite.HitProxyId);
+					Vertex[SpriteListIndex + 0] = FSimpleElementVertex(FVector4(Sprite.Position + WorldSpriteX - WorldSpriteY,1),FVector2D(UEnd,  VStart),Sprite.Color,Sprite.HitProxyId);
+					Vertex[SpriteListIndex + 1] = FSimpleElementVertex(FVector4(Sprite.Position + WorldSpriteX + WorldSpriteY,1),FVector2D(UEnd,  VEnd  ),Sprite.Color,Sprite.HitProxyId);
+					Vertex[SpriteListIndex + 2] = FSimpleElementVertex(FVector4(Sprite.Position - WorldSpriteX - WorldSpriteY,1),FVector2D(UStart,VStart),Sprite.Color,Sprite.HitProxyId);
 
-				Vertex[SpriteListIndex + 3] = FSimpleElementVertex(FVector4(Sprite.Position + WorldSpriteX + WorldSpriteY,1),FVector2D(UEnd,  VEnd  ),Sprite.Color,Sprite.HitProxyId);
-				Vertex[SpriteListIndex + 4] = FSimpleElementVertex(FVector4(Sprite.Position - WorldSpriteX - WorldSpriteY,1),FVector2D(UStart,VStart),Sprite.Color,Sprite.HitProxyId);
-				Vertex[SpriteListIndex + 5] = FSimpleElementVertex(FVector4(Sprite.Position - WorldSpriteX + WorldSpriteY,1),FVector2D(UStart,VEnd  ),Sprite.Color,Sprite.HitProxyId);
+					Vertex[SpriteListIndex + 3] = FSimpleElementVertex(FVector4(Sprite.Position + WorldSpriteX + WorldSpriteY,1),FVector2D(UEnd,  VEnd  ),Sprite.Color,Sprite.HitProxyId);
+					Vertex[SpriteListIndex + 4] = FSimpleElementVertex(FVector4(Sprite.Position - WorldSpriteX - WorldSpriteY,1),FVector2D(UStart,VStart),Sprite.Color,Sprite.HitProxyId);
+					Vertex[SpriteListIndex + 5] = FSimpleElementVertex(FVector4(Sprite.Position - WorldSpriteX + WorldSpriteY,1),FVector2D(UStart,VEnd  ),Sprite.Color,Sprite.HitProxyId);
+				}
 			}
 
 			if (SpriteList.Num() > 0)
 			{
-				//Draw last batch
-				const int32 VertexCount = SpriteList.Num();
-				const int32 PrimCount = VertexCount / 3;
-				PrepareShaders(RHICmdList, CurrentBlendMode, Transform, bNeedToSwitchVerticalAxis, BatchedElementParameters, CurrentTexture, bHitTesting, Gamma, NULL, View, DepthTexture);
-				DrawPrimitiveUP(RHICmdList, PT_TriangleList, PrimCount, SpriteList.GetTypedData(), sizeof(FSimpleElementVertex));
+				const EBlendModeFilter::Type SpriteFilter = GetBlendModeFilter(CurrentBlendMode);
+
+				// Only render blend modes in the filter
+				if (Filter & SpriteFilter)
+				{
+					//Draw last batch
+					const int32 VertexCount = SpriteList.Num();
+					const int32 PrimCount = VertexCount / 3;
+					PrepareShaders(RHICmdList, CurrentBlendMode, Transform, bNeedToSwitchVerticalAxis, BatchedElementParameters, CurrentTexture, bHitTesting, Gamma, NULL, View, DepthTexture);
+					DrawPrimitiveUP(RHICmdList, PT_TriangleList, PrimCount, SpriteList.GetTypedData(), sizeof(FSimpleElementVertex));
+				}
 			}
 		}
 
@@ -869,43 +894,53 @@ bool FBatchedElements::Draw(FRHICommandList& RHICmdList, bool bNeedToSwitchVerti
 			for(int32 MeshIndex = 0;MeshIndex < MeshElements.Num();MeshIndex++)
 			{
 				const FBatchedMeshElement& MeshElement = MeshElements[MeshIndex];
+				const EBlendModeFilter::Type MeshFilter = GetBlendModeFilter(MeshElement.BlendMode);
 
-				// Set the appropriate pixel shader for the mesh.
-				PrepareShaders(RHICmdList, MeshElement.BlendMode, Transform, bNeedToSwitchVerticalAxis, MeshElement.BatchedElementParameters, MeshElement.Texture, bHitTesting, Gamma, &MeshElement.GlowInfo);
-
-				// Draw the mesh.
-				DrawIndexedPrimitiveUP(
-					RHICmdList,
-					PT_TriangleList,
-					0,
-					MeshElement.MaxVertex - MeshElement.MinVertex + 1,
-					MeshElement.Indices.Num() / 3,
-					MeshElement.Indices.GetTypedData(),
-					sizeof(uint16),
-					&MeshVertices[MeshElement.MinVertex],
-					sizeof(FSimpleElementVertex)
-					);
+				// Only render blend modes in the filter
+				if (Filter & MeshFilter)
+				{
+					// Set the appropriate pixel shader for the mesh.
+					PrepareShaders(RHICmdList, MeshElement.BlendMode, Transform, bNeedToSwitchVerticalAxis, MeshElement.BatchedElementParameters, MeshElement.Texture, bHitTesting, Gamma, &MeshElement.GlowInfo);
+
+					// Draw the mesh.
+					DrawIndexedPrimitiveUP(
+						RHICmdList,
+						PT_TriangleList,
+						0,
+						MeshElement.MaxVertex - MeshElement.MinVertex + 1,
+						MeshElement.Indices.Num() / 3,
+						MeshElement.Indices.GetTypedData(),
+						sizeof(uint16),
+						&MeshVertices[MeshElement.MinVertex],
+						sizeof(FSimpleElementVertex)
+						);
+				}
 			}
 
 			// Draw the quad mesh elements.
 			for(int32 MeshIndex = 0;MeshIndex < QuadMeshElements.Num();MeshIndex++)
 			{
 				const FBatchedQuadMeshElement& MeshElement = QuadMeshElements[MeshIndex];
+				const EBlendModeFilter::Type MeshFilter = GetBlendModeFilter(MeshElement.BlendMode);
 
-				// Quads don't support batched element parameters (yet!)
-				FBatchedElementParameters* BatchedElementParameters = NULL;
-
-				// Set the appropriate pixel shader for the mesh.
-				PrepareShaders(RHICmdList, MeshElement.BlendMode, Transform, bNeedToSwitchVerticalAxis, BatchedElementParameters, MeshElement.Texture, bHitTesting, Gamma);
-
-				// Draw the mesh.
-				DrawPrimitiveUP(
-					RHICmdList,
-					PT_QuadList,
-					MeshElement.Vertices.Num() / 4,
-					MeshElement.Vertices.GetTypedData(),
-					sizeof(FSimpleElementVertex)
-					);
+				// Only render blend modes in the filter
+				if (Filter & MeshFilter)
+				{
+					// Quads don't support batched element parameters (yet!)
+					FBatchedElementParameters* BatchedElementParameters = NULL;
+
+					// Set the appropriate pixel shader for the mesh.
+					PrepareShaders(RHICmdList, MeshElement.BlendMode, Transform, bNeedToSwitchVerticalAxis, BatchedElementParameters, MeshElement.Texture, bHitTesting, Gamma);
+
+					// Draw the mesh.
+					DrawPrimitiveUP(
+						RHICmdList,
+						PT_QuadList,
+						MeshElement.Vertices.Num() / 4,
+						MeshElement.Vertices.GetTypedData(),
+						sizeof(FSimpleElementVertex)
+						);
+				}
 			}
 		}
 
diff --git a/Engine/Source/Runtime/Engine/Private/BillboardComponent.cpp b/Engine/Source/Runtime/Engine/Private/BillboardComponent.cpp
--- a/Engine/Source/Runtime/Engine/Private/BillboardComponent.cpp
+++ b/Engine/Source/Runtime/Engine/Private/BillboardComponent.cpp
@@ -23,8 +23,6 @@ public:
 		, U(InComponent->U)
 		, V(InComponent->V)
 		, Color(FLinearColor::White)
-		, LevelColor(FLinearColor::White)
-		, PropertyColor(255,255,255)
 		, bIsScreenSizeScaled(InComponent->bIsScreenSizeScaled)
 #if WITH_EDITORONLY_DATA
 		, bUseInEditorScaling(InComponent->bUseInEditorScaling)
@@ -94,11 +92,91 @@ public:
 			}
 		}
 
-		GEngine->GetPropertyColorationColor( (UObject*)InComponent, PropertyColor );
+		FColor NewPropertyColor;
+		GEngine->GetPropertyColorationColor( (UObject*)InComponent, NewPropertyColor );
+		PropertyColor = NewPropertyColor;
 	}
 
 	// FPrimitiveSceneProxy interface.
 	
+	virtual void GetDynamicMeshElements(const TArray<const FSceneView*>& Views, const FSceneViewFamily& ViewFamily, uint32 VisibilityMap, FMeshElementCollector& Collector) const override
+	{
+		QUICK_SCOPE_CYCLE_COUNTER( STAT_SpriteSceneProxy_GetDynamicMeshElements );
+
+		FTexture* TextureResource = (Texture != NULL) ? Texture->Resource : NULL;
+		if (TextureResource)
+		{
+			for (int32 ViewIndex = 0; ViewIndex < Views.Num(); ViewIndex++)
+			{
+				if (VisibilityMap & (1 << ViewIndex))
+				{
+					const FSceneView* View = Views[ViewIndex];
+
+					// Calculate the view-dependent scaling factor.
+					float ViewedSizeX = SizeX;
+					float ViewedSizeY = SizeY;
+
+					if (bIsScreenSizeScaled && (View->ViewMatrices.ProjMatrix.M[3][3] != 1.0f))
+					{
+						const float ZoomFactor	= FMath::Min<float>(View->ViewMatrices.ProjMatrix.M[0][0], View->ViewMatrices.ProjMatrix.M[1][1]);
+
+						if(ZoomFactor != 0.0f)
+						{
+							const float Radius = View->WorldToScreen(Origin).W * (ScreenSize / ZoomFactor);
+
+							if (Radius < 1.0f)
+							{
+								ViewedSizeX *= Radius;
+								ViewedSizeY *= Radius;
+							}					
+						}
+					}
+
+#if WITH_EDITORONLY_DATA
+					ViewedSizeX *= EditorScale;
+					ViewedSizeY *= EditorScale;
+#endif
+
+					FLinearColor ColorToUse = Color;
+
+					// Set the selection/hover color from the current engine setting.
+					// The color is multiplied by 10 because this value is normally expected to be blended
+					// additively, this is not how the sprites work and therefore need an extra boost
+					// to appear the same color as previously.
+					if (IsSelected())
+					{
+						ColorToUse = FLinearColor::White + (GEngine->GetSelectedMaterialColor() * GEngine->SelectionHighlightIntensityBillboards * 10);
+					}
+					else if (IsHovered())
+					{
+						ColorToUse = FLinearColor::White + (GEngine->GetHoveredMaterialColor() * GEngine->HoverHighlightIntensity * 10);
+					}
+
+					// Sprites of locked actors draw in red.
+					if (bIsActorLocked)
+					{
+						ColorToUse = FColor(255,0,0);
+					}
+					FLinearColor LevelColorToUse = IsSelected() ? ColorToUse : (FLinearColor)LevelColor;
+					FLinearColor PropertyColorToUse = PropertyColor;
+
+					const FLinearColor& SpriteColor = View->Family->EngineShowFlags.LevelColoration ? LevelColorToUse :
+						( (View->Family->EngineShowFlags.PropertyColoration) ? PropertyColorToUse : ColorToUse );
+
+					Collector.GetPDI(ViewIndex)->DrawSprite(
+						Origin,
+						ViewedSizeX,
+						ViewedSizeY,
+						TextureResource,
+						SpriteColor,
+						GetDepthPriorityGroup(View),
+						U,UL,V,VL
+						);
+				}
+			}
+		}
+	}
+
 	/** 
 	 * Draw the scene proxy as a dynamic element
 	 *
@@ -208,8 +286,6 @@ private:
 	const float V;
 	float VL;
 	FLinearColor Color;
-	FLinearColor LevelColor;
-	FColor PropertyColor;
 	const uint32 bIsScreenSizeScaled : 1;
 	uint32 bIsActorLocked : 1;
 #if WITH_EDITOR
@@ -310,4 +386,4 @@ void UBillboardComponent::SetEditorScale(float InEditorScale)
 		It->MarkRenderStateDirty();
 	}
 }
-#endif
\ No newline at end of file
+#endif
diff --git a/Engine/Source/Runtime/Engine/Private/BoxComponent.cpp b/Engine/Source/Runtime/Engine/Private/BoxComponent.cpp
--- a/Engine/Source/Runtime/Engine/Private/BoxComponent.cpp
+++ b/Engine/Source/Runtime/Engine/Private/BoxComponent.cpp
@@ -95,6 +95,22 @@ FPrimitiveSceneProxy* UBoxComponent::CreateSceneProxy()
 			bWillEverBeLit = false;
 		}
 
+		virtual void GetDynamicMeshElements(const TArray<const FSceneView*>& Views, const FSceneViewFamily& ViewFamily, uint32 VisibilityMap, FMeshElementCollector& Collector) const override
+		{
+			QUICK_SCOPE_CYCLE_COUNTER( STAT_BoxSceneProxy_GetDynamicMeshElements );
+
+			const FMatrix& LocalToWorld = GetLocalToWorld();
+			const FColor DrawColor = GetSelectionColor(BoxColor, IsSelected(), IsHovered(), /*bUseOverlayIntensity=*/false);
+
+			for (int32 ViewIndex = 0; ViewIndex < Views.Num(); ViewIndex++)
+			{
+				if (VisibilityMap & (1 << ViewIndex))
+				{
+					FPrimitiveDrawInterface* PDI = Collector.GetPDI(ViewIndex);
+					DrawOrientedWireBox(PDI, LocalToWorld.GetOrigin(), LocalToWorld.GetScaledAxis( EAxis::X ), LocalToWorld.GetScaledAxis( EAxis::Y ), LocalToWorld.GetScaledAxis( EAxis::Z ), BoxExtents, DrawColor, SDPG_World);
+				}
+			}
+		}
 
 		virtual void DrawDynamicElements(FPrimitiveDrawInterface* PDI,const FSceneView* View) override
 		{
diff --git a/Engine/Source/Runtime/Engine/Private/BrushComponent.cpp b/Engine/Source/Runtime/Engine/Private/BrushComponent.cpp
--- a/Engine/Source/Runtime/Engine/Private/BrushComponent.cpp
+++ b/Engine/Source/Runtime/Engine/Private/BrushComponent.cpp
@@ -135,8 +135,6 @@ public:
 		bSolidWhenSelected(false),
 		bInManipulation(false),
 		BrushColor(GEngine->C_BrushWire),
-		LevelColor(FLinearColor::White),
-		PropertyColor(255,255,255),
 		BodySetup(Component->BrushBodySetup),
 		CollisionResponse(Component->GetCollisionResponseToChannels())
 	{
@@ -175,8 +173,12 @@ public:
 			}
 		}
 
+		bUseEditorDepthTest = !bInManipulation;
+
 		// Get a color for property coloration.
-		GEngine->GetPropertyColorationColor( (UObject*)Component, PropertyColor );
+		FColor NewPropertyColor;
+		GEngine->GetPropertyColorationColor( (UObject*)Component, NewPropertyColor );
+		PropertyColor = NewPropertyColor;
 
 #if WITH_EDITORONLY_DATA
 		ENQUEUE_UNIQUE_RENDER_COMMAND_TWOPARAMETER(
@@ -204,13 +206,13 @@ public:
 
 	}
 
-	bool IsCollisionView(const FSceneView* View, bool & bDrawCollision) const
+	bool IsCollisionView(const FEngineShowFlags& EngineShowFlags, bool & bDrawCollision) const
 	{
-		const bool bInCollisionView = View->Family->EngineShowFlags.CollisionVisibility || View->Family->EngineShowFlags.CollisionPawn;
+		const bool bInCollisionView = EngineShowFlags.CollisionVisibility || EngineShowFlags.CollisionPawn;
 		if (bInCollisionView && IsCollisionEnabled())
 		{
-			bDrawCollision = View->Family->EngineShowFlags.CollisionPawn && (CollisionResponse.GetResponse(ECC_Pawn) != ECR_Ignore);
-			bDrawCollision |= View->Family->EngineShowFlags.CollisionVisibility && (CollisionResponse.GetResponse(ECC_Visibility) != ECR_Ignore);
+			bDrawCollision = EngineShowFlags.CollisionPawn && (CollisionResponse.GetResponse(ECC_Pawn) != ECR_Ignore);
+			bDrawCollision |= EngineShowFlags.CollisionVisibility && (CollisionResponse.GetResponse(ECC_Visibility) != ECR_Ignore);
 		}
 		else
 		{
@@ -220,6 +222,103 @@ public:
 		return bInCollisionView;
 	}
 
+	virtual void GetDynamicMeshElements(const TArray<const FSceneView*>& Views, const FSceneViewFamily& ViewFamily, uint32 VisibilityMap, FMeshElementCollector& Collector) const override
+	{
+		QUICK_SCOPE_CYCLE_COUNTER( STAT_BrushSceneProxy_GetDynamicMeshElements );
+
+		if( AllowDebugViewmodes() )
+		{
+			for (int32 ViewIndex = 0; ViewIndex < Views.Num(); ViewIndex++)
+			{
+				if (VisibilityMap & (1 << ViewIndex))
+				{
+					const FSceneView* View = Views[ViewIndex];
+
+					bool bDrawCollision = false;
+					const bool bInCollisionView = IsCollisionView(ViewFamily.EngineShowFlags, bDrawCollision);
+
+					// Draw solid if 'solid when selected' and selected, or we are in a 'collision view'
+					const bool bDrawSolid = ((bSolidWhenSelected && IsSelected()) || (bInCollisionView && bDrawCollision));
+					// Don't draw wireframe if in a collision view mode and not drawing solid
+					const bool bDrawWireframe = !bInCollisionView;
+
+					// Choose color to draw it
+					FLinearColor DrawColor = BrushColor;
+					// In a collision view mode
+					if(bInCollisionView)
+					{
+						DrawColor = BrushColor;
+					}
+					else if(View->Family->EngineShowFlags.PropertyColoration)
+					{
+						DrawColor = PropertyColor;
+					}
+					else if(View->Family->EngineShowFlags.LevelColoration)
+					{
+						DrawColor = LevelColor;
+					}
+
+
+					// SOLID
+					if(bDrawSolid)
+					{
+						if(BodySetup != NULL)
+						{
+							auto SolidMaterialInstance = new FColoredMaterialRenderProxy(
+								GEngine->ShadedLevelColorationUnlitMaterial->GetRenderProxy(IsSelected(), IsHovered()),
+								DrawColor
+								);
+
+							Collector.RegisterOneFrameMaterialProxy(SolidMaterialInstance);
+
+							FTransform GeomTransform(GetLocalToWorld());
+							BodySetup->AggGeom.GetAggGeom(GeomTransform, DrawColor, /*Material=*/SolidMaterialInstance, false, /*bSolid=*/ true, UseEditorDepthTest(), ViewIndex, Collector);
+						}
+					}
+					// WIREFRAME
+					else if(bDrawWireframe)
+					{
+						// If we have the editor data (Wire Buffers) use those for wireframe
+#if WITH_EDITOR
+						if(WireIndexBuffer.GetNumEdges() && WireVertexBuffer.GetNumVertices())
+						{
+							auto WireframeMaterial = new FColoredMaterialRenderProxy(
+								GEngine->LevelColorationUnlitMaterial->GetRenderProxy(IsSelected(), IsHovered()),
+								GetSelectionColor(DrawColor,!(GIsEditor && (View->Family->EngineShowFlags.Selection)) || IsSelected(), IsHovered(), /*bUseOverlayIntensity=*/false)
+								);
+
+							Collector.RegisterOneFrameMaterialProxy(WireframeMaterial);
+
+							FMeshBatch& Mesh = Collector.AllocateMesh();
+							FMeshBatchElement& BatchElement = Mesh.Elements[0];
+							BatchElement.IndexBuffer = &WireIndexBuffer;
+							Mesh.VertexFactory = &VertexFactory;
+							Mesh.MaterialRenderProxy = WireframeMaterial;
+							BatchElement.PrimitiveUniformBufferResource = &GetUniformBuffer();
+							BatchElement.FirstIndex = 0;
+							BatchElement.NumPrimitives = WireIndexBuffer.GetNumEdges();
+							BatchElement.MinVertexIndex = 0;
+							BatchElement.MaxVertexIndex = WireVertexBuffer.GetNumVertices() - 1;
+							Mesh.CastShadow = false;
+							Mesh.Type = PT_LineList;
+							Mesh.DepthPriorityGroup = IsSelected() ? SDPG_Foreground : SDPG_World;
+							Collector.AddMesh(ViewIndex, Mesh);
+						}
+						else 
+#endif
+						if(BodySetup != NULL)
+							// If not, use the body setup for wireframe
+						{
+							FTransform GeomTransform(GetLocalToWorld());
+							BodySetup->AggGeom.GetAggGeom(GeomTransform, GetSelectionColor(DrawColor, IsSelected(), IsHovered()), /* Material=*/ NULL, false, /* bSolid=*/ false, UseEditorDepthTest(), ViewIndex, Collector);
+						}
+
+					}
+				}
+			}
+		}
+	}
+
 	/** 
 	* Draw the scene proxy as a dynamic element
 	*
@@ -234,7 +333,7 @@ public:
 		if( AllowDebugViewmodes() )
 		{
 			bool bDrawCollision = false;
-			const bool bInCollisionView = IsCollisionView(View, bDrawCollision);
+			const bool bInCollisionView = IsCollisionView(View->Family->EngineShowFlags, bDrawCollision);
 
 			// Draw solid if 'solid when selected' and selected, or we are in a 'collision view'
 			const bool bDrawSolid = ((bSolidWhenSelected && IsSelected()) || (bInCollisionView && bDrawCollision));
@@ -269,7 +368,7 @@ public:
 						);
 
 					FTransform GeomTransform(GetLocalToWorld());
-					BodySetup->AggGeom.DrawAggGeom(PDI, GeomTransform, DrawColor, /*Material=*/&SolidMaterialInstance, false, /*bSolid=*/ true );
+					BodySetup->AggGeom.DrawAggGeom(PDI, GeomTransform, DrawColor, /*Material=*/&SolidMaterialInstance, false, /*bSolid=*/ true, UseEditorDepthTest() );
 				}
 			}
 			// WIREFRAME
@@ -305,7 +404,7 @@ public:
 				// If not, use the body setup for wireframe
 				{
 					FTransform GeomTransform(GetLocalToWorld());
-					BodySetup->AggGeom.DrawAggGeom(PDI, GeomTransform, GetSelectionColor(DrawColor, IsSelected(), IsHovered()), /* Material=*/ NULL, false, /* bSolid=*/ false );
+					BodySetup->AggGeom.DrawAggGeom(PDI, GeomTransform, GetSelectionColor(DrawColor, IsSelected(), IsHovered()), /* Material=*/ NULL, false, /* bSolid=*/ false, UseEditorDepthTest() );
 				}
 
 			}
diff --git a/Engine/Source/Runtime/Engine/Private/CapsuleComponent.cpp b/Engine/Source/Runtime/Engine/Private/CapsuleComponent.cpp
--- a/Engine/Source/Runtime/Engine/Private/CapsuleComponent.cpp
+++ b/Engine/Source/Runtime/Engine/Private/CapsuleComponent.cpp
@@ -32,6 +32,24 @@ FPrimitiveSceneProxy* UCapsuleComponent::CreateSceneProxy()
 			bWillEverBeLit = false;
 		}
 
+		virtual void GetDynamicMeshElements(const TArray<const FSceneView*>& Views, const FSceneViewFamily& ViewFamily, uint32 VisibilityMap, FMeshElementCollector& Collector) const override
+		{
+			QUICK_SCOPE_CYCLE_COUNTER( STAT_GetDynamicMeshElements_DrawDynamicElements );
+
+			const FLinearColor DrawCapsuleColor = GetSelectionColor(ShapeColor, IsSelected(), IsHovered(), /*bUseOverlayIntensity=*/false);
+			const FMatrix& LocalToWorld = GetLocalToWorld();
+			const int32 CapsuleSides =  FMath::Clamp<int32>(CapsuleRadius/4.f, 16, 64);
+
+			for (int32 ViewIndex = 0; ViewIndex < Views.Num(); ViewIndex++)
+			{
+				if (VisibilityMap & (1 << ViewIndex))
+				{
+					FPrimitiveDrawInterface* PDI = Collector.GetPDI(ViewIndex);
+					DrawWireCapsule( PDI, LocalToWorld.GetOrigin(), LocalToWorld.GetScaledAxis( EAxis::X ), LocalToWorld.GetScaledAxis( EAxis::Y ), LocalToWorld.GetScaledAxis( EAxis::Z ), DrawCapsuleColor, CapsuleRadius, CapsuleHalfHeight, CapsuleSides, SDPG_World );
+				}
+			}
+		}
+
 		/** 
 		* Draw the scene proxy as a dynamic element
 		*
diff --git a/Engine/Source/Runtime/Engine/Private/DebugRenderSceneProxy.cpp b/Engine/Source/Runtime/Engine/Private/DebugRenderSceneProxy.cpp
--- a/Engine/Source/Runtime/Engine/Private/DebugRenderSceneProxy.cpp
+++ b/Engine/Source/Runtime/Engine/Private/DebugRenderSceneProxy.cpp
@@ -71,6 +71,103 @@ void FDebugRenderSceneProxy::DrawDebugLabels(UCanvas* Canvas, APlayerController*
 	Canvas->SetDrawColor(OldDrawColor);
 }
 
+void FDebugRenderSceneProxy::GetDynamicMeshElements(const TArray<const FSceneView*>& Views, const FSceneViewFamily& ViewFamily, uint32 VisibilityMap, FMeshElementCollector& Collector) const 
+{
+	QUICK_SCOPE_CYCLE_COUNTER( STAT_DebugRenderSceneProxy_GetDynamicMeshElements );
+
+	for (int32 ViewIndex = 0; ViewIndex < Views.Num(); ViewIndex++)
+	{
+		if (VisibilityMap & (1 << ViewIndex))
+		{
+			const FSceneView* View = Views[ViewIndex];
+			FPrimitiveDrawInterface* PDI = Collector.GetPDI(ViewIndex);
+
+			// Draw Lines
+			for(int32 LineIdx=0; LineIdx<Lines.Num(); LineIdx++)
+			{
+				const FDebugLine& Line = Lines[LineIdx];
+
+				PDI->DrawLine(Line.Start, Line.End, Line.Color, SDPG_World);
+			}
+
+			// Draw Arrows
+			for(int32 LineIdx=0; LineIdx<ArrowLines.Num(); LineIdx++)
+			{
+				const FArrowLine &Line = ArrowLines[LineIdx];
+
+				DrawLineArrow(PDI, Line.Start, Line.End, Line.Color, 8.0f);
+			}
+
+			// Draw Cylinders
+			for(int32 CylinderIdx=0; CylinderIdx<Cylinders.Num(); CylinderIdx++)
+			{
+				const FWireCylinder& Cylinder = Cylinders[CylinderIdx];
+
+				DrawWireCylinder( PDI, Cylinder.Base, FVector(1,0,0), FVector(0,1,0), FVector(0,0,1),
+					Cylinder.Color, Cylinder.Radius, Cylinder.HalfHeight, 16, SDPG_World );
+			}
+
+			// Draw Stars
+			for(int32 StarIdx=0; StarIdx<Stars.Num(); StarIdx++)
+			{
+				const FWireStar& Star = Stars[StarIdx];
+
+				DrawWireStar(PDI, Star.Position, Star.Size, Star.Color, SDPG_World);
+			}
+
+			// Draw Dashed Lines
+			for(int32 DashIdx=0; DashIdx<DashedLines.Num(); DashIdx++)
+			{
+				const FDashedLine& Dash = DashedLines[DashIdx];
+
+				DrawDashedLine(PDI, Dash.Start, Dash.End, Dash.Color, Dash.DashSize, SDPG_World);
+			}
+
+			// Draw Boxes
+			for(int32 BoxIdx=0; BoxIdx<WireBoxes.Num(); BoxIdx++)
+			{
+				const FDebugBox& Box = WireBoxes[BoxIdx];
+
+				DrawWireBox( PDI, Box.Box, Box.Color, SDPG_World);
+			}
+
+			// Draw solid spheres
+			struct FMaterialCache
+			{
+				FMaterialCache() {}
+
+				FMaterialRenderProxy* operator[](FLinearColor Color) const
+				{
+					FMaterialRenderProxy* MeshColor = NULL;
+					const uint32 HashKey = GetTypeHash(Color);
+					if (MeshColorInstances.Contains(HashKey))
+					{
+						MeshColor = *MeshColorInstances.Find(HashKey);
+					}
+					else
+					{
+						MeshColor = new(FMemStack::Get()) FColoredMaterialRenderProxy(GEngine->DebugMeshMaterial->GetRenderProxy(false), Color);
+						MeshColorInstances.Add(HashKey, MeshColor);
+					}
+
+					return MeshColor;
+				}
+
+				mutable TMap<uint32, FMaterialRenderProxy*> MeshColorInstances;
+			};
+			const FMaterialCache MaterialCache;
+
+			for (auto It = SolidSpheres.CreateConstIterator(); It; ++It)
+			{
+				if (PointInView(It->Location, View))
+				{
+					GetSphereMesh(It->Location, FVector(It->Radius), 10, 7, MaterialCache[It->Color], SDPG_World, false, ViewIndex, Collector);
+				}
+			}
+		}
+	}
+}
+
 /** 
 * Draw the scene proxy as a dynamic element
 *
@@ -176,7 +273,7 @@ void FDebugRenderSceneProxy::DrawDynamicElements(FPrimitiveDrawInterface* PDI,co
 * @param Color		Color of the line.
 * @param Mag		Size of the arrow.
 */
-void FDebugRenderSceneProxy::DrawLineArrow(FPrimitiveDrawInterface* PDI,const FVector &Start,const FVector &End,const FColor &Color,float Mag)
+void FDebugRenderSceneProxy::DrawLineArrow(FPrimitiveDrawInterface* PDI,const FVector &Start,const FVector &End,const FColor &Color,float Mag) const
 {
 	// draw a pretty arrow
 	FVector Dir = End - Start;
diff --git a/Engine/Source/Runtime/Engine/Private/DrawFrustumComponent.cpp b/Engine/Source/Runtime/Engine/Private/DrawFrustumComponent.cpp
--- a/Engine/Source/Runtime/Engine/Private/DrawFrustumComponent.cpp
+++ b/Engine/Source/Runtime/Engine/Private/DrawFrustumComponent.cpp
@@ -30,6 +30,67 @@ public:
 
 	// FPrimitiveSceneProxy interface.
 
+	virtual void GetDynamicMeshElements(const TArray<const FSceneView*>& Views, const FSceneViewFamily& ViewFamily, uint32 VisibilityMap, FMeshElementCollector& Collector) const override
+	{
+		QUICK_SCOPE_CYCLE_COUNTER( STAT_DrawFrustumSceneProxy_DrawDynamicElements );
+
+		FVector Direction(1,0,0);
+		FVector LeftVector(0,1,0);
+		FVector UpVector(0,0,1);
+
+		FVector Verts[8];
+
+		// FOVAngle controls the horizontal angle.
+		float HozHalfAngle = (FrustumAngle) * ((float)PI/360.f);
+		float HozLength = FrustumStartDist * FMath::Tan(HozHalfAngle);
+		float VertLength = HozLength/FrustumAspectRatio;
+
+		// near plane verts
+		Verts[0] = (Direction * FrustumStartDist) + (UpVector * VertLength) + (LeftVector * HozLength);
+		Verts[1] = (Direction * FrustumStartDist) + (UpVector * VertLength) - (LeftVector * HozLength);
+		Verts[2] = (Direction * FrustumStartDist) - (UpVector * VertLength) - (LeftVector * HozLength);
+		Verts[3] = (Direction * FrustumStartDist) - (UpVector * VertLength) + (LeftVector * HozLength);
+
+		HozLength = FrustumEndDist * FMath::Tan(HozHalfAngle);
+		VertLength = HozLength/FrustumAspectRatio;
+
+		// far plane verts
+		Verts[4] = (Direction * FrustumEndDist) + (UpVector * VertLength) + (LeftVector * HozLength);
+		Verts[5] = (Direction * FrustumEndDist) + (UpVector * VertLength) - (LeftVector * HozLength);
+		Verts[6] = (Direction * FrustumEndDist) - (UpVector * VertLength) - (LeftVector * HozLength);
+		Verts[7] = (Direction * FrustumEndDist) - (UpVector * VertLength) + (LeftVector * HozLength);
+
+		for( int32 x = 0 ; x < 8 ; ++x )
+		{
+			Verts[x] = GetLocalToWorld().TransformPosition( Verts[x] );
+		}
+
+		for (int32 ViewIndex = 0; ViewIndex < Views.Num(); ViewIndex++)
+		{
+			if (VisibilityMap & (1 << ViewIndex))
+			{
+				FPrimitiveDrawInterface* PDI = Collector.GetPDI(ViewIndex);
+				const FSceneView* View = Views[ViewIndex];
+
+				const uint8 DepthPriorityGroup = GetDepthPriorityGroup(View);
+				PDI->DrawLine( Verts[0], Verts[1], FrustumColor, DepthPriorityGroup );
+				PDI->DrawLine( Verts[1], Verts[2], FrustumColor, DepthPriorityGroup );
+				PDI->DrawLine( Verts[2], Verts[3], FrustumColor, DepthPriorityGroup );
+				PDI->DrawLine( Verts[3], Verts[0], FrustumColor, DepthPriorityGroup );
+
+				PDI->DrawLine( Verts[4], Verts[5], FrustumColor, DepthPriorityGroup );
+				PDI->DrawLine( Verts[5], Verts[6], FrustumColor, DepthPriorityGroup );
+				PDI->DrawLine( Verts[6], Verts[7], FrustumColor, DepthPriorityGroup );
+				PDI->DrawLine( Verts[7], Verts[4], FrustumColor, DepthPriorityGroup );
+
+				PDI->DrawLine( Verts[0], Verts[4], FrustumColor, DepthPriorityGroup );
+				PDI->DrawLine( Verts[1], Verts[5], FrustumColor, DepthPriorityGroup );
+				PDI->DrawLine( Verts[2], Verts[6], FrustumColor, DepthPriorityGroup );
+				PDI->DrawLine( Verts[3], Verts[7], FrustumColor, DepthPriorityGroup );
+			}
+		}
+	}
+
 	/** 
 	* Draw the scene proxy as a dynamic element
 	*
diff --git a/Engine/Source/Runtime/Engine/Private/DynamicMeshBuilder.cpp b/Engine/Source/Runtime/Engine/Private/DynamicMeshBuilder.cpp
--- a/Engine/Source/Runtime/Engine/Private/DynamicMeshBuilder.cpp
+++ b/Engine/Source/Runtime/Engine/Private/DynamicMeshBuilder.cpp
@@ -412,6 +412,69 @@ void FDynamicMeshBuilder::AddTriangles(const TArray<int32> &InIndices)
 	IndexBuffer->Indices.Append(InIndices);
 }
 
+void FDynamicMeshBuilder::GetMesh(const FMatrix& LocalToWorld,const FMaterialRenderProxy* MaterialRenderProxy,uint8 DepthPriorityGroup,bool bDisableBackfaceCulling, bool bReceivesDecals, int32 ViewIndex, FMeshElementCollector& Collector) 
+{
+	// Only draw non-empty meshes.
+	if(VertexBuffer->Vertices.Num() > 0 && IndexBuffer->Indices.Num() > 0)
+	{
+		FPrimitiveDrawInterface* PDI = Collector.GetPDI(ViewIndex);
+		// Register the dynamic resources with the PDI.
+		PDI->RegisterDynamicResource(VertexBuffer);
+		PDI->RegisterDynamicResource(IndexBuffer);
+
+		// Create the vertex factory.
+		FDynamicMeshVertexFactory* VertexFactory = new FDynamicMeshVertexFactory(VertexBuffer);
+		PDI->RegisterDynamicResource(VertexFactory);
+
+		// Create the primitive uniform buffer.
+		FDynamicMeshPrimitiveUniformBuffer* PrimitiveUniformBuffer = new FDynamicMeshPrimitiveUniformBuffer();
+		FPrimitiveUniformShaderParameters PrimitiveParams = GetPrimitiveUniformShaderParameters(
+			LocalToWorld,
+			LocalToWorld.GetOrigin(),
+			FBoxSphereBounds(EForceInit::ForceInit),
+			FBoxSphereBounds(EForceInit::ForceInit),
+			bReceivesDecals,
+			false,
+			false,
+			1.0f		// LPV bias
+			);
+
+
+		if (IsInGameThread())
+		{
+			BeginSetUniformBufferContents(*PrimitiveUniformBuffer, PrimitiveParams);
+		}
+		else
+		{
+			PrimitiveUniformBuffer->SetContents(PrimitiveParams);
+		}
+		PDI->RegisterDynamicResource(PrimitiveUniformBuffer);
+
+		// Draw the mesh.
+		FMeshBatch& Mesh = Collector.AllocateMesh();
+		FMeshBatchElement& BatchElement = Mesh.Elements[0];
+		BatchElement.IndexBuffer = IndexBuffer;
+		Mesh.VertexFactory = VertexFactory;
+		Mesh.MaterialRenderProxy = MaterialRenderProxy;
+		BatchElement.PrimitiveUniformBufferResource = PrimitiveUniformBuffer;
+		// previous l2w not used so treat as static
+		BatchElement.FirstIndex = 0;
+		BatchElement.NumPrimitives = IndexBuffer->Indices.Num() / 3;
+		BatchElement.MinVertexIndex = 0;
+		BatchElement.MaxVertexIndex = VertexBuffer->Vertices.Num() - 1;
+		Mesh.ReverseCulling = LocalToWorld.Determinant() < 0.0f ? true : false;
+		Mesh.bDisableBackfaceCulling = bDisableBackfaceCulling;
+		Mesh.Type = PT_TriangleList;
+		Mesh.DepthPriorityGroup = DepthPriorityGroup;
+		Collector.AddMesh(ViewIndex, Mesh);
+
+		// Clear the resource pointers so they cannot be overwritten accidentally.
+		// These resources will be released by the PDI.
+		VertexBuffer = NULL;
+		IndexBuffer = NULL;
+	}
+}
+
 void FDynamicMeshBuilder::Draw(FPrimitiveDrawInterface* PDI,const FMatrix& LocalToWorld,const FMaterialRenderProxy* MaterialRenderProxy,uint8 DepthPriorityGroup,bool bDisableBackfaceCulling, bool bReceivesDecals)
 {
 	// Only draw non-empty meshes.
@@ -434,6 +497,7 @@ void FDynamicMeshBuilder::Draw(FPrimitiveDrawInterface* PDI,const FMatrix& Local
 			FBoxSphereBounds(EForceInit::ForceInit),
 			bReceivesDecals,
 			false,
+			false,
 			1.0f		// LPV bias
 			);
 
