Paper2D: PaperSpriteSheet container asset allows reimporting sprites from an imported json spritesheet, handles reasonable changes
Paper2D: Action to automatically extract flipbooks from PaperSpriteSheets grouped by name_number

diff --git a/Engine/Plugins/Experimental/Paper2D/Source/Paper2DEditor/Classes/PaperFlipbookFactory.h b/Engine/Plugins/Experimental/Paper2D/Source/Paper2DEditor/Classes/PaperFlipbookFactory.h
--- a/Engine/Plugins/Experimental/Paper2D/Source/Paper2DEditor/Classes/PaperFlipbookFactory.h
+++ b/Engine/Plugins/Experimental/Paper2D/Source/Paper2DEditor/Classes/PaperFlipbookFactory.h
@@ -4,10 +4,11 @@
  * Factory for flipbooks
  */
 
+#pragma once
 #include "PaperFlipbookFactory.generated.h"
 
 UCLASS()
-class UPaperFlipbookFactory : public UFactory
+class PAPER2DEDITOR_API UPaperFlipbookFactory : public UFactory
 {
 	GENERATED_UCLASS_BODY()
 
diff --git a/Engine/Plugins/Experimental/Paper2D/Source/PaperJsonImporter/Classes/PaperJsonImporterFactory.h b/Engine/Plugins/Experimental/Paper2D/Source/PaperJsonImporter/Classes/PaperJsonImporterFactory.h
--- a/Engine/Plugins/Experimental/Paper2D/Source/PaperJsonImporter/Classes/PaperJsonImporterFactory.h
+++ b/Engine/Plugins/Experimental/Paper2D/Source/PaperJsonImporter/Classes/PaperJsonImporterFactory.h
@@ -14,8 +14,5 @@ class UPaperJsonImporterFactory : public UFactory
 	virtual FText GetToolTip() const override;
 	virtual UObject* FactoryCreateText(UClass* InClass, UObject* InParent, FName InName, EObjectFlags Flags, UObject* Context, const TCHAR* Type, const TCHAR*& Buffer, const TCHAR* BufferEnd, FFeedbackContext* Warn) override;
 	// End of UFactory interface
-
-protected:
-	TSharedPtr<class FJsonObject> ParseJSON(const FString& FileContents, const FString& NameForErrors);
 };
 
diff --git a/Engine/Plugins/Experimental/Paper2D/Source/PaperJsonImporter/PaperJsonImporter.Build.cs b/Engine/Plugins/Experimental/Paper2D/Source/PaperJsonImporter/PaperJsonImporter.Build.cs
--- a/Engine/Plugins/Experimental/Paper2D/Source/PaperJsonImporter/PaperJsonImporter.Build.cs
+++ b/Engine/Plugins/Experimental/Paper2D/Source/PaperJsonImporter/PaperJsonImporter.Build.cs
@@ -10,10 +10,14 @@ public class PaperJsonImporter : ModuleRules
 			new string[] {
 				"Core",
 				"CoreUObject",
+				"Slate",
+				"SlateCore",
 				"Engine",
 				"Paper2D",
 				"UnrealEd",
-				"Paper2DEditor"
+				"Paper2DEditor",
+				"AssetTools",
+				"ContentBrowser"
 			});
 
 		PrivateIncludePathModuleNames.AddRange(
diff --git a/Engine/Plugins/Experimental/Paper2D/Source/PaperJsonImporter/Private/PaperJsonImporterFactory.cpp b/Engine/Plugins/Experimental/Paper2D/Source/PaperJsonImporter/Private/PaperJsonImporterFactory.cpp
--- a/Engine/Plugins/Experimental/Paper2D/Source/PaperJsonImporter/Private/PaperJsonImporterFactory.cpp
+++ b/Engine/Plugins/Experimental/Paper2D/Source/PaperJsonImporter/Private/PaperJsonImporterFactory.cpp
@@ -7,28 +7,9 @@
 #include "AssetToolsModule.h"
 #include "AssetRegistryModule.h"
 #include "PackageTools.h"
-
-//////////////////////////////////////////////////////////////////////////
-// FSpriteFrame
-
-// Represents one parsed frame in a sprite sheet
-struct FSpriteFrame
-{
-	FName FrameName;
-
-	FVector2D SpritePosInSheet;
-	FVector2D SpriteSizeInSheet;
-
-	FVector2D SpriteSourcePos;
-	FVector2D SpriteSourceSize;
-
-	FVector2D ImageSourceSize;
-
-	FVector2D Pivot;
-
-	bool bTrimmed;
-	bool bRotated;
-};
+#include "PaperSpriteSheet.h"
+#include "PaperJsonSpriteSheetImporter.h"
+#include "ReimportPaperJsonImporterFactory.h"
 
 //////////////////////////////////////////////////////////////////////////
 // UPaperJsonImporterFactory
@@ -38,7 +19,7 @@ UPaperJsonImporterFactory::UPaperJsonImporterFactory(const FPostConstructInitial
 {
 	bCreateNew = false;
 	//bEditAfterNew = true;
-	SupportedClass = UPaperFlipbook::StaticClass();
+	SupportedClass = UPaperSpriteSheet::StaticClass();
 
 	bEditorImport = true;
 	bText = true;
@@ -52,216 +33,6 @@ FText UPaperJsonImporterFactory::GetToolTip() const
 	return NSLOCTEXT("Paper2D", "PaperJsonImporterFactoryDescription", "Sprite sheets exported from Adobe Flash");
 }
 
-static bool ParseFrame(TSharedPtr<FJsonObject> &FrameData, FSpriteFrame &OutFrame)
-{
-	bool bReadFrameSuccessfully = true;
-	// An example frame:
-	//   "frame": {"x":210,"y":10,"w":190,"h":223},
-	//   "rotated": false,
-	//   "trimmed": true,
-	//   "spriteSourceSize": {"x":0,"y":11,"w":216,"h":240},
-	//   "sourceSize": {"w":216,"h":240},	
-	//   "pivot": {"x":0.5,"y":0.5}			[optional]
-
-	OutFrame.bRotated = FPaperJSONHelpers::ReadBoolean(FrameData, TEXT("rotated"), false);
-	OutFrame.bTrimmed = FPaperJSONHelpers::ReadBoolean(FrameData, TEXT("trimmed"), false);
-
-	bReadFrameSuccessfully = bReadFrameSuccessfully && FPaperJSONHelpers::ReadRectangle(FrameData, "frame", /*out*/ OutFrame.SpritePosInSheet, /*out*/ OutFrame.SpriteSizeInSheet);
-
-	if (OutFrame.bTrimmed)
-	{
-		bReadFrameSuccessfully = bReadFrameSuccessfully && FPaperJSONHelpers::ReadSize(FrameData, "sourceSize", /*out*/ OutFrame.ImageSourceSize);
-		bReadFrameSuccessfully = bReadFrameSuccessfully && FPaperJSONHelpers::ReadRectangle(FrameData, "spriteSourceSize", /*out*/ OutFrame.SpriteSourcePos, /*out*/ OutFrame.SpriteSourceSize);
-	}
-	else
-	{
-		OutFrame.SpriteSourcePos = FVector2D::ZeroVector;
-		OutFrame.SpriteSourceSize = OutFrame.SpriteSizeInSheet;
-		OutFrame.ImageSourceSize = OutFrame.SpriteSizeInSheet;
-	}
-
-	if (!FPaperJSONHelpers::ReadXY(FrameData, "pivot", /*out*/ OutFrame.Pivot))
-	{
-		OutFrame.Pivot.X = 0.5f;
-		OutFrame.Pivot.Y = 0.5f;
-	}
-
-	// A few more prerequisites to sort out before rotation can be enabled
-	if (OutFrame.bRotated)
-	{
-		// Sprite Source Pos is from top left, our pivot when rotated is bottom left
-		OutFrame.SpriteSourcePos.Y = OutFrame.ImageSourceSize.Y - OutFrame.SpriteSourcePos.Y - OutFrame.SpriteSizeInSheet.Y;
-
-		// We rotate final sprite geometry 90 deg CCW to fixup
-		// These need to be swapped to be valid in texture space.
-		Swap(OutFrame.SpriteSizeInSheet.X, OutFrame.SpriteSizeInSheet.Y);
-		Swap(OutFrame.ImageSourceSize.X, OutFrame.ImageSourceSize.Y);
-		
-		Swap(OutFrame.SpriteSourcePos.X, OutFrame.SpriteSourcePos.Y);
-		Swap(OutFrame.SpriteSourceSize.X, OutFrame.SpriteSourceSize.Y);
-	}
-
-	return bReadFrameSuccessfully;
-}
-
-static bool ParseFramesFromSpriteHash(TSharedPtr<FJsonObject> ObjectBlock, TArray<FSpriteFrame>& OutSpriteFrames, TSet<FName>& FrameNames)
-{
-	GWarn->BeginSlowTask(NSLOCTEXT("Paper2D", "PaperJsonImporterFactory_ParsingSprites", "Parsing Sprite Frame"), true, true);
-	bool bLoadedSuccessfully = true;
-
-	// Parse all of the frames
-	int32 FrameCount = 0;
-	for (auto FrameIt = ObjectBlock->Values.CreateIterator(); FrameIt; ++FrameIt)
-	{
-		GWarn->StatusUpdate(FrameCount, ObjectBlock->Values.Num(), NSLOCTEXT("Paper2D", "PaperJsonImporterFactory_ParsingSprites", "Parsing Sprite Frames"));
-
-		bool bReadFrameSuccessfully = true;
-
-		FSpriteFrame Frame;
-		Frame.FrameName = *FrameIt.Key();
-
-		if (FrameNames.Contains(Frame.FrameName))
-		{
-			bReadFrameSuccessfully = false;
-		}
-		else
-		{
-			FrameNames.Add(Frame.FrameName);
-		}
-
-		TSharedPtr<FJsonValue> FrameDataAsValue = FrameIt.Value();
-		TSharedPtr<FJsonObject> FrameData;
-		if (FrameDataAsValue->Type == EJson::Object)
-		{
-			FrameData = FrameDataAsValue->AsObject();
-			bReadFrameSuccessfully = bReadFrameSuccessfully && ParseFrame(FrameData, /*out*/Frame);
-		}
-		else
-		{
-			bReadFrameSuccessfully = false;
-		}
-
-		if (bReadFrameSuccessfully)
-		{
-			OutSpriteFrames.Add(Frame);
-		}
-		else
-		{
-			UE_LOG(LogPaperJsonImporter, Warning, TEXT("Frame %s is in an unexpected format"), *Frame.FrameName.ToString());
-			bLoadedSuccessfully = false;
-		}
-
-		FrameCount++;
-	}
-
-	GWarn->EndSlowTask();
-	return bLoadedSuccessfully;
-}
-
-static bool ParseFramesFromSpriteArray(const TArray<TSharedPtr<FJsonValue>>& ArrayBlock, TArray<FSpriteFrame>& OutSpriteFrames, TSet<FName>& FrameNames)
-{
-	GWarn->BeginSlowTask(NSLOCTEXT("Paper2D", "PaperJsonImporterFactory_ParsingSprites", "Parsing Sprite Frame"), true, true);
-	bool bLoadedSuccessfully = true;
-
-	// Parse all of the frames
-	int32 FrameCount = 0;
-	for (int32 FrameCount = 0; FrameCount < ArrayBlock.Num(); ++FrameCount)
-	{
-		GWarn->StatusUpdate(FrameCount, ArrayBlock.Num(), NSLOCTEXT("Paper2D", "PaperJsonImporterFactory_ParsingSprites", "Parsing Sprite Frames"));
-		bool bReadFrameSuccessfully = true;
-		FSpriteFrame Frame;
-
-		const TSharedPtr<FJsonValue>& FrameDataAsValue = ArrayBlock[FrameCount];
-		if (FrameDataAsValue->Type == EJson::Object)
-		{
-			TSharedPtr<FJsonObject> FrameData;
-			FrameData = FrameDataAsValue->AsObject();
-
-			FString FrameFilename = FPaperJSONHelpers::ReadString(FrameData, TEXT("filename"), TEXT(""));
-			if (!FrameFilename.IsEmpty())
-			{
-				Frame.FrameName = FName(*FrameFilename); // case insensitive
-				if (FrameNames.Contains(Frame.FrameName))
-				{
-					bReadFrameSuccessfully = false;
-				}
-				else
-				{
-					FrameNames.Add(Frame.FrameName);
-				}
-
-				bReadFrameSuccessfully = bReadFrameSuccessfully && ParseFrame(FrameData, /*out*/Frame);
-			}
-			else
-			{
-				bReadFrameSuccessfully = false;
-			}
-		}
-		else
-		{
-			bReadFrameSuccessfully = false;
-		}
-
-		if (bReadFrameSuccessfully)
-		{
-			OutSpriteFrames.Add(Frame);
-		}
-		else
-		{
-			UE_LOG(LogPaperJsonImporter, Warning, TEXT("Frame %s is in an unexpected format"), *Frame.FrameName.ToString());
-			bLoadedSuccessfully = false;
-		}
-	}
-
-	GWarn->EndSlowTask();
-	return bLoadedSuccessfully;
-}
-
-static ESpritePivotMode::Type GetBestPivotType(FVector2D JsonPivot)
-{
-	// Not assuming layout of ESpritePivotMode
-	if (JsonPivot.X == 0 && JsonPivot.Y == 0)
-	{
-		return ESpritePivotMode::Top_Left;
-	}
-	else if (JsonPivot.X == 0.5f && JsonPivot.Y == 0)
-	{
-		return ESpritePivotMode::Top_Center;
-	}
-	else if (JsonPivot.X == 1.0f && JsonPivot.Y == 0)
-	{
-		return ESpritePivotMode::Top_Right;
-	}
-	else if (JsonPivot.X == 0 && JsonPivot.Y == 0.5f)
-	{
-		return ESpritePivotMode::Center_Left;
-	}
-	else if (JsonPivot.X == 0.5f && JsonPivot.Y == 0.5f)
-	{
-		return ESpritePivotMode::Center_Center;
-	}
-	else if (JsonPivot.X == 1.0f && JsonPivot.Y == 0.5f)
-	{
-		return ESpritePivotMode::Center_Right;
-	}
-	else if (JsonPivot.X == 0 && JsonPivot.Y == 1.0f)
-	{
-		return ESpritePivotMode::Bottom_Left;
-	}
-	else if (JsonPivot.X == 0.5f && JsonPivot.Y == 1.0f)
-	{
-		return ESpritePivotMode::Bottom_Center;
-	}
-	else if (JsonPivot.X == 1.0f && JsonPivot.Y == 1.0f)
-	{
-		return ESpritePivotMode::Bottom_Right;
-	}
-	else
-	{
-		return ESpritePivotMode::Custom;
-	}
-}
-
 UObject* UPaperJsonImporterFactory::FactoryCreateText(UClass* InClass, UObject* InParent, FName InName, EObjectFlags Flags, UObject* Context, const TCHAR* Type, const TCHAR*& Buffer, const TCHAR* BufferEnd, FFeedbackContext* Warn)
 {
 	Flags |= RF_Transactional;
@@ -270,249 +41,49 @@ UObject* UPaperJsonImporterFactory::FactoryCreateText(UClass* InClass, UObject*
 
 	FAssetToolsModule& AssetToolsModule = FModuleManager::GetModuleChecked<FAssetToolsModule>("AssetTools");
 
-	bool bLoadedSuccessfully = true;
-
 	const FString CurrentFilename = UFactory::GetCurrentFilename();
 	FString CurrentSourcePath;
 	FString FilenameNoExtension;
 	FString UnusedExtension;
 	FPaths::Split(CurrentFilename, CurrentSourcePath, FilenameNoExtension, UnusedExtension);
 
-	const FString LongPackagePath = FPackageName::GetLongPackagePath(InParent->GetOutermost()->GetPathName());
+	
 
+	const FString LongPackagePath = FPackageName::GetLongPackagePath(InParent->GetOutermost()->GetPathName());
+	UPaperSpriteSheet* Result = NULL;
+	
 	const FString NameForErrors(InName.ToString());
 	const FString FileContent(BufferEnd - Buffer, Buffer);
-	TSharedPtr<FJsonObject> SpriteDescriptorObject = ParseJSON(FileContent, NameForErrors);
-
-	TSet<FName> FrameNames;
-	TArray<FSpriteFrame> ParsedFrames;
+	
+	FPaperJsonSpriteSheetImporter Importer;
 
-	UPaperFlipbook* Result = NULL;
-
-	if (SpriteDescriptorObject.IsValid())
+	// Are we reimporting?
+	if (IsA(UReimportPaperJsonImporterFactory::StaticClass()))
 	{
-		UTexture2D* ImportedTexture = NULL;
-
-		// Validate the meta block
-		TSharedPtr<FJsonObject> MetaBlock = FPaperJSONHelpers::ReadObject(SpriteDescriptorObject, TEXT("meta"));
-		if (MetaBlock.IsValid())
-		{
-			// Example contents:
-			//   "app": "Adobe Flash CS6",
-			//   "version": "12.0.0.481",        (ignored)
-			//   "image": "MySprite.png",
-			//   "format": "RGBA8888",           (ignored)
-			//   "size": {"w":2048,"h":2048},    (ignored)
-			//   "scale": "1"                    (ignored)
-
-			const FString AppName = FPaperJSONHelpers::ReadString(MetaBlock, TEXT("app"), TEXT(""));
-			const FString Image = FPaperJSONHelpers::ReadString(MetaBlock, TEXT("image"), TEXT(""));
-			
- 			if (bLoadedSuccessfully)
-			{
-				const FString FlashPrefix(TEXT("Adobe Flash"));
-				const FString TexturePackerPrefix(TEXT("http://www.codeandweb.com/texturepacker"));
-
-				if (AppName.StartsWith(FlashPrefix) || AppName.StartsWith(TexturePackerPrefix))
-				{
-					// Cool, we (mostly) know how to handle these sorts of files!
-					UE_LOG(LogPaperJsonImporter, Log, TEXT("Parsing sprite sheet exported from '%s'"), *AppName);
-				}
-				else if (!AppName.IsEmpty())
-				{
-					// It's got an app tag inside a meta block, so we'll take a crack at it
-					UE_LOG(LogPaperJsonImporter, Warning, TEXT("Unexpected 'app' named '%s' while parsing sprite descriptor file '%s'.  Parsing will continue but the format may not be fully supported"), *AppName, *NameForErrors);
-				}
-				else
-				{
-					// Probably not a sprite sheet
-					UE_LOG(LogPaperJsonImporter, Warning, TEXT("Failed to parse sprite descriptor file '%s'.  Expected 'app' key indicating the exporter (might not be a sprite sheet)"), *NameForErrors);
-					bLoadedSuccessfully = false;
-				}
- 			}
-
-			if (bLoadedSuccessfully)
-			{
-				if (Image.IsEmpty())
-				{
-					UE_LOG(LogPaperJsonImporter, Warning, TEXT("Failed to parse sprite descriptor file '%s'.  Expected valid 'image' tag"), *NameForErrors);
-					bLoadedSuccessfully = false;
-				}
-				else
-				{
-					// Import the texture
-					//@TODO: Avoid the first compression, since we're going to recompress
-					const FString TargetSubPath = LongPackagePath + TEXT("/Textures");
-
-					const FString SourceSheetImageFilename = FPaths::Combine(*CurrentSourcePath, *Image);
-					TArray<FString> SheetFileNames;
-					SheetFileNames.Add(SourceSheetImageFilename);
-
-					TArray<UObject*> ImportedSheets = AssetToolsModule.Get().ImportAssets(SheetFileNames, TargetSubPath);
-
-					ImportedTexture = (ImportedSheets.Num() > 0) ? Cast<UTexture2D>(ImportedSheets[0]) : NULL;
-
-					if (ImportedTexture == NULL)
-					{
-						UE_LOG(LogPaperJsonImporter, Warning, TEXT("Failed to import sprite sheet image '%s'."), *SourceSheetImageFilename);
-						bLoadedSuccessfully = false;
-					}
-					else
-					{
-						// Change the compression settings
-						ImportedTexture->Modify();
-						ImportedTexture->LODGroup = TEXTUREGROUP_UI;
-						ImportedTexture->CompressionSettings = TC_EditorIcon;
-						ImportedTexture->Filter = TF_Nearest;
-						ImportedTexture->PostEditChange();
-					}
-				}
-			}
-		}
-		else
-		{
-			// Missing meta tag
-			UE_LOG(LogPaperJsonImporter, Warning, TEXT("Failed to parse sprite descriptor file '%s'.  Missing meta block"), *NameForErrors);
-			bLoadedSuccessfully = false;
-		}
+		UReimportPaperJsonImporterFactory* CurrentReimportFactory = Cast<UReimportPaperJsonImporterFactory>(this);
+		Importer.SetReimportData(CurrentReimportFactory->ExistingTextureName, CurrentReimportFactory->ExistingTexture, CurrentReimportFactory->ExistingSpriteNames, CurrentReimportFactory->ExistingSprites);
+	}
 
-		// Parse the frames
-		if (bLoadedSuccessfully)
+	if (Importer.ImportFromString(FileContent, NameForErrors) &&
+		Importer.ImportTextures(LongPackagePath, CurrentSourcePath))
+	{
+		UPaperSpriteSheet* SpriteSheet = NewNamedObject<UPaperSpriteSheet>(InParent, InName, Flags);
+		if (Importer.PerformImport(LongPackagePath, Flags, SpriteSheet))
 		{
-			TSharedPtr<FJsonObject> ObjectFrameBlock = FPaperJSONHelpers::ReadObject(SpriteDescriptorObject, TEXT("frames"));
-			if (ObjectFrameBlock.IsValid())
-			{
-				bLoadedSuccessfully = ParseFramesFromSpriteHash(ObjectFrameBlock, /*out*/ParsedFrames, /*inout*/FrameNames);
-			}
-			else
+			if (SpriteSheet->AssetImportData == nullptr)
 			{
-				// Try loading as an array
-				TArray<TSharedPtr<FJsonValue>> ArrayBlock = FPaperJSONHelpers::ReadArray(SpriteDescriptorObject, TEXT("frames"));
-				if (ArrayBlock.Num() > 0)
-				{
-					bLoadedSuccessfully = ParseFramesFromSpriteArray(ArrayBlock, /*out*/ParsedFrames, /*inout*/FrameNames);
-				}
-				else
-				{
-					UE_LOG(LogPaperJsonImporter, Warning, TEXT("Failed to parse sprite descriptor file '%s'.  Missing frames block"), *NameForErrors);
-				}
+				SpriteSheet->AssetImportData = ConstructObject<UAssetImportData>(UAssetImportData::StaticClass(), SpriteSheet);
 			}
 
+			SpriteSheet->AssetImportData->SourceFilePath = FReimportManager::SanitizeImportFilename(CurrentFilename, SpriteSheet);
+			SpriteSheet->AssetImportData->SourceFileTimestamp = IFileManager::Get().GetTimeStamp(*CurrentFilename).ToString();
+			SpriteSheet->AssetImportData->bDirty = false;
 
-			// Create an animated sprite that encompasses the frames
-			UPaperFlipbook* Flipbook = NULL;
-			if (ParsedFrames.Num() > 0)
-			{
-				Flipbook = NewNamedObject<UPaperFlipbook>(InParent, InName, Flags);
-				Result = Flipbook;
-
-				GWarn->BeginSlowTask(NSLOCTEXT("Paper2D", "PaperJsonImporterFactory_ImportingSprites", "Importing Sprite Frame"), true, true);
-				FScopedFlipbookMutator EditLock(Flipbook);
-
-				// Create objects for each successfully parsed frame
-				const int32 FrameRun = 1; //@TODO: Don't make a keyframe for every single item if we can help it
-				for (int32 FrameIndex = 0; FrameIndex < ParsedFrames.Num(); ++FrameIndex)
-				{
-					GWarn->StatusUpdate(FrameIndex, ParsedFrames.Num(), NSLOCTEXT("Paper2D", "PaperJsonImporterFactory_ImportingSprites", "Importing Sprite Frames"));
-
-					// Check for user canceling the import
-					if (GWarn->ReceivedUserCancel())
-					{
-						break;
-					}
-
-					const FSpriteFrame& Frame = ParsedFrames[FrameIndex];
-
-					// Create a package for the frame
-					const FString TargetSubPath = LongPackagePath + TEXT("/Frames");
-
-					UObject* OuterForFrame = NULL; // @TODO: Use this if we don't want them to be individual assets - Flipbook;
-
-					// Create a unique package name and asset name for the frame
-					const FString TentativePackagePath = PackageTools::SanitizePackageName(TargetSubPath + TEXT("/") + Frame.FrameName.ToString());
-					FString DefaultSuffix;
-					FString AssetName;
-					FString PackageName;
-					AssetToolsModule.Get().CreateUniqueAssetName(TentativePackagePath, /*out*/ DefaultSuffix, /*out*/ PackageName, /*out*/ AssetName);
-
-					// Create a package for the frame
-					OuterForFrame = CreatePackage(NULL, *PackageName);
-
-					// Create a frame in the package
-					UPaperSprite* TargetSprite = NewNamedObject<UPaperSprite>(OuterForFrame, *AssetName, Flags);
-					FAssetRegistryModule::AssetCreated(TargetSprite);
-
-					TargetSprite->Modify();
-
-					TargetSprite->InitializeSprite(ImportedTexture, Frame.SpritePosInSheet, Frame.SpriteSizeInSheet, GetDefault<UPaperRuntimeSettings>()->DefaultPixelsPerUnrealUnit);
-
-					if (Frame.bRotated)
-					{
-						TargetSprite->SetRotated(true);
-					}
-
-					if (Frame.bTrimmed)
-					{
-						TargetSprite->SetTrim(Frame.bTrimmed, Frame.SpriteSourcePos, Frame.ImageSourceSize);
-					}
-
-					// Set up pivot on object based on Texture Packer json
-					ESpritePivotMode::Type PivotType = GetBestPivotType(Frame.Pivot);
-					FVector2D TextureSpacePivotPoint = FVector2D::ZeroVector;
-					if (PivotType == ESpritePivotMode::Custom)
-					{
-						TextureSpacePivotPoint.X = Frame.SpritePosInSheet.X - Frame.SpriteSourcePos.X + Frame.ImageSourceSize.X * Frame.Pivot.X;
-						TextureSpacePivotPoint.Y = Frame.SpritePosInSheet.Y - Frame.SpriteSourcePos.Y + Frame.ImageSourceSize.Y * Frame.Pivot.Y;
-					}
-					TargetSprite->SetPivot(PivotType, TextureSpacePivotPoint);
-
-					// Create the entry in the animation
-					FPaperFlipbookKeyFrame* DestFrame = new (EditLock.KeyFrames) FPaperFlipbookKeyFrame();
-					DestFrame->Sprite = TargetSprite;
-					DestFrame->FrameRun = FrameRun;
-
-					TargetSprite->PostEditChange();
-				}
-
-				GWarn->EndSlowTask();
-			}
-
+			Result = SpriteSheet;
 		}
 	}
-	else
-	{
-		// Failed to parse the JSON
-		bLoadedSuccessfully = false;
-	}
-
+	
 	FEditorDelegates::OnAssetPostImport.Broadcast(this, Result);
 
 	return Result;
 }
-
-TSharedPtr<FJsonObject> UPaperJsonImporterFactory::ParseJSON(const FString& FileContents, const FString& NameForErrors)
-{
-	// Load the file up (JSON format)
-
-	if (!FileContents.IsEmpty())
-	{
-		const TSharedRef< TJsonReader<> >& Reader = TJsonReaderFactory<>::Create(FileContents);
-
-		TSharedPtr<FJsonObject> SpriteDescriptorObject;
-		if (FJsonSerializer::Deserialize(Reader, SpriteDescriptorObject) && SpriteDescriptorObject.IsValid())
-		{
-			// File was loaded and deserialized OK!
-			return SpriteDescriptorObject;
-		}
-		else
-		{
-			UE_LOG(LogPaperJsonImporter, Warning, TEXT("Failed to parse sprite descriptor file '%s'.  Error: '%s'"), *NameForErrors, *Reader->GetErrorMessage());
-			return NULL;
-		}
-	}
-	else
-	{
-		UE_LOG(LogPaperJsonImporter, Warning, TEXT("Sprite descriptor file '%s' was empty.  This sprite cannot be imported."), *NameForErrors);
-		return NULL;
-	}
-}
diff --git a/Engine/Plugins/Experimental/Paper2D/Source/PaperJsonImporter/Private/PaperJsonImporterModule.cpp b/Engine/Plugins/Experimental/Paper2D/Source/PaperJsonImporter/Private/PaperJsonImporterModule.cpp
--- a/Engine/Plugins/Experimental/Paper2D/Source/PaperJsonImporter/Private/PaperJsonImporterModule.cpp
+++ b/Engine/Plugins/Experimental/Paper2D/Source/PaperJsonImporter/Private/PaperJsonImporterModule.cpp
@@ -2,6 +2,9 @@
 
 #include "PaperJsonImporterPrivatePCH.h"
 
+#include "AssetToolsModule.h"
+#include "PaperSpriteSheetAssetTypeActions.h"
+
 
 //////////////////////////////////////////////////////////////////////////
 // FPaperJsonImporterModule
@@ -11,11 +14,27 @@ class FPaperJsonImporterModule : public FDefaultModuleImpl
 public:
 	virtual void StartupModule() override
 	{
+		// Register asset types
+		IAssetTools& AssetTools = FModuleManager::LoadModuleChecked<FAssetToolsModule>("AssetTools").Get();
+
+		SpriteSheetImportAssetTypeActions = MakeShareable(new FPaperSpriteSheetAssetTypeActions);
+		AssetTools.RegisterAssetTypeActions(SpriteSheetImportAssetTypeActions.ToSharedRef());
 	}
 
 	virtual void ShutdownModule() override
 	{
+		if (FModuleManager::Get().IsModuleLoaded("AssetTools"))
+		{
+			IAssetTools& AssetTools = FModuleManager::GetModuleChecked<FAssetToolsModule>("AssetTools").Get();
+			if (SpriteSheetImportAssetTypeActions.IsValid())
+			{
+				AssetTools.UnregisterAssetTypeActions(SpriteSheetImportAssetTypeActions.ToSharedRef());
+			}
+		}
 	}
+
+private:
+	TSharedPtr<IAssetTypeActions> SpriteSheetImportAssetTypeActions;
 };
 
 //////////////////////////////////////////////////////////////////////////
diff --git a/Engine/Plugins/Experimental/Paper2D/Source/PaperJsonImporter/Private/PaperJsonImporterPrivatePCH.h b/Engine/Plugins/Experimental/Paper2D/Source/PaperJsonImporter/Private/PaperJsonImporterPrivatePCH.h
--- a/Engine/Plugins/Experimental/Paper2D/Source/PaperJsonImporter/Private/PaperJsonImporterPrivatePCH.h
+++ b/Engine/Plugins/Experimental/Paper2D/Source/PaperJsonImporter/Private/PaperJsonImporterPrivatePCH.h
@@ -3,6 +3,7 @@
 #pragma once
 
 #include "UnrealEd.h"
+#include "Paper2DClasses.h"
 
 DECLARE_LOG_CATEGORY_EXTERN(LogPaperJsonImporter, Verbose, All);
 
diff --git a/Engine/Plugins/Experimental/Paper2D/Source/PaperJsonImporter/Private/PaperJsonSpriteSheetImporter.cpp b/Engine/Plugins/Experimental/Paper2D/Source/PaperJsonImporter/Private/PaperJsonSpriteSheetImporter.cpp
new file mode 100644
--- /dev/null
+++ b/Engine/Plugins/Experimental/Paper2D/Source/PaperJsonImporter/Private/PaperJsonSpriteSheetImporter.cpp
@@ -0,0 +1,538 @@
+// Copyright 1998-2014 Epic Games, Inc. All Rights Reserved.
+
+#include "PaperJsonImporterPrivatePCH.h"
+#include "Paper2DClasses.h"
+#include "Json.h"
+#include "PaperJSONHelpers.h"
+#include "AssetToolsModule.h"
+#include "AssetRegistryModule.h"
+#include "PackageTools.h"
+#include "PaperSpriteSheet.h"
+
+TSharedPtr<FJsonObject> ParseJSON(const FString& FileContents, const FString& NameForErrors)
+{
+	// Load the file up (JSON format)
+	if (!FileContents.IsEmpty())
+	{
+		const TSharedRef< TJsonReader<> >& Reader = TJsonReaderFactory<>::Create(FileContents);
+
+		TSharedPtr<FJsonObject> SpriteDescriptorObject;
+		if (FJsonSerializer::Deserialize(Reader, SpriteDescriptorObject) && SpriteDescriptorObject.IsValid())
+		{
+			// File was loaded and deserialized OK!
+			return SpriteDescriptorObject;
+		}
+		else
+		{
+			UE_LOG(LogPaperJsonImporter, Warning, TEXT("Failed to parse sprite descriptor file '%s'.  Error: '%s'"), *NameForErrors, *Reader->GetErrorMessage());
+			return NULL;
+		}
+	}
+	else
+	{
+		UE_LOG(LogPaperJsonImporter, Warning, TEXT("Sprite descriptor file '%s' was empty.  This sprite cannot be imported."), *NameForErrors);
+		return NULL;
+	}
+}
+
+TSharedPtr<FJsonObject> ParseJSON(FArchive* const Stream, const FString& NameForErrors)
+{
+	const TSharedRef< TJsonReader<> >& Reader = TJsonReaderFactory<>::Create(Stream);
+
+	TSharedPtr<FJsonObject> SpriteDescriptorObject;
+	if (FJsonSerializer::Deserialize(Reader, SpriteDescriptorObject) && SpriteDescriptorObject.IsValid())
+	{
+		// File was loaded and deserialized OK!
+		return SpriteDescriptorObject;
+	}
+	else
+	{
+		UE_LOG(LogPaperJsonImporter, Warning, TEXT("Failed to parse sprite descriptor file '%s'.  Error: '%s'"), *NameForErrors, *Reader->GetErrorMessage());
+		return NULL;
+	}
+}
+
+bool ParseMetaBlock(const FString& NameForErrors, TSharedPtr<FJsonObject>& SpriteDescriptorObject, FString& OutImage)
+{
+	bool bLoadedSuccessfully = true;
+	TSharedPtr<FJsonObject> MetaBlock = FPaperJSONHelpers::ReadObject(SpriteDescriptorObject, TEXT("meta"));
+	if (MetaBlock.IsValid())
+	{
+		// Example contents:
+		//   "app": "Adobe Flash CS6",
+		//   "version": "12.0.0.481",        (ignored)
+		//   "image": "MySprite.png",
+		//   "format": "RGBA8888",           (ignored)
+		//   "size": {"w":2048,"h":2048},    (ignored)
+		//   "scale": "1"                    (ignored)
+
+		const FString AppName = FPaperJSONHelpers::ReadString(MetaBlock, TEXT("app"), TEXT(""));
+		OutImage = FPaperJSONHelpers::ReadString(MetaBlock, TEXT("image"), TEXT(""));
+
+		const FString FlashPrefix(TEXT("Adobe Flash"));
+		const FString TexturePackerPrefix(TEXT("http://www.codeandweb.com/texturepacker"));
+
+		if (AppName.StartsWith(FlashPrefix) || AppName.StartsWith(TexturePackerPrefix))
+		{
+			// Cool, we (mostly) know how to handle these sorts of files!
+			UE_LOG(LogPaperJsonImporter, Log, TEXT("Parsing sprite sheet exported from '%s'"), *AppName);
+		}
+		else if (!AppName.IsEmpty())
+		{
+			// It's got an app tag inside a meta block, so we'll take a crack at it
+			UE_LOG(LogPaperJsonImporter, Warning, TEXT("Unexpected 'app' named '%s' while parsing sprite descriptor file '%s'.  Parsing will continue but the format may not be fully supported"), *AppName, *NameForErrors);
+		}
+		else
+		{
+			// Probably not a sprite sheet
+			UE_LOG(LogPaperJsonImporter, Warning, TEXT("Failed to parse sprite descriptor file '%s'.  Expected 'app' key indicating the exporter (might not be a sprite sheet)"), *NameForErrors);
+			bLoadedSuccessfully = false;
+		}
+
+		if (OutImage.IsEmpty())
+		{
+			UE_LOG(LogPaperJsonImporter, Warning, TEXT("Failed to parse sprite descriptor file '%s'.  Expected valid 'image' tag"), *NameForErrors);
+			bLoadedSuccessfully = false;
+		}
+	}
+	else
+	{
+		UE_LOG(LogPaperJsonImporter, Warning, TEXT("Failed to parse sprite descriptor file '%s'.  Missing meta block"), *NameForErrors);
+		bLoadedSuccessfully = false;
+	}
+	return bLoadedSuccessfully;
+}
+
+static bool ParseFrame(TSharedPtr<FJsonObject> &FrameData, FSpriteFrame &OutFrame)
+{
+	bool bReadFrameSuccessfully = true;
+	// An example frame:
+	//   "frame": {"x":210,"y":10,"w":190,"h":223},
+	//   "rotated": false,
+	//   "trimmed": true,
+	//   "spriteSourceSize": {"x":0,"y":11,"w":216,"h":240},
+	//   "sourceSize": {"w":216,"h":240},	
+	//   "pivot": {"x":0.5,"y":0.5}			[optional]
+
+	OutFrame.bRotated = FPaperJSONHelpers::ReadBoolean(FrameData, TEXT("rotated"), false);
+	OutFrame.bTrimmed = FPaperJSONHelpers::ReadBoolean(FrameData, TEXT("trimmed"), false);
+
+	bReadFrameSuccessfully = bReadFrameSuccessfully && FPaperJSONHelpers::ReadRectangle(FrameData, "frame", /*out*/ OutFrame.SpritePosInSheet, /*out*/ OutFrame.SpriteSizeInSheet);
+
+	if (OutFrame.bTrimmed)
+	{
+		bReadFrameSuccessfully = bReadFrameSuccessfully && FPaperJSONHelpers::ReadSize(FrameData, "sourceSize", /*out*/ OutFrame.ImageSourceSize);
+		bReadFrameSuccessfully = bReadFrameSuccessfully && FPaperJSONHelpers::ReadRectangle(FrameData, "spriteSourceSize", /*out*/ OutFrame.SpriteSourcePos, /*out*/ OutFrame.SpriteSourceSize);
+	}
+	else
+	{
+		OutFrame.SpriteSourcePos = FVector2D::ZeroVector;
+		OutFrame.SpriteSourceSize = OutFrame.SpriteSizeInSheet;
+		OutFrame.ImageSourceSize = OutFrame.SpriteSizeInSheet;
+	}
+
+	if (!FPaperJSONHelpers::ReadXY(FrameData, "pivot", /*out*/ OutFrame.Pivot))
+	{
+		OutFrame.Pivot.X = 0.5f;
+		OutFrame.Pivot.Y = 0.5f;
+	}
+
+	// A few more prerequisites to sort out before rotation can be enabled
+	if (OutFrame.bRotated)
+	{
+		// Sprite Source Pos is from top left, our pivot when rotated is bottom left
+		OutFrame.SpriteSourcePos.Y = OutFrame.ImageSourceSize.Y - OutFrame.SpriteSourcePos.Y - OutFrame.SpriteSizeInSheet.Y;
+
+		// We rotate final sprite geometry 90 deg CCW to fixup
+		// These need to be swapped to be valid in texture space.
+		Swap(OutFrame.SpriteSizeInSheet.X, OutFrame.SpriteSizeInSheet.Y);
+		Swap(OutFrame.ImageSourceSize.X, OutFrame.ImageSourceSize.Y);
+
+		Swap(OutFrame.SpriteSourcePos.X, OutFrame.SpriteSourcePos.Y);
+		Swap(OutFrame.SpriteSourceSize.X, OutFrame.SpriteSourceSize.Y);
+	}
+
+	return bReadFrameSuccessfully;
+}
+
+static bool ParseFramesFromSpriteHash(TSharedPtr<FJsonObject> ObjectBlock, TArray<FSpriteFrame>& OutSpriteFrames, TSet<FName>& FrameNames)
+{
+	GWarn->BeginSlowTask(NSLOCTEXT("Paper2D", "PaperJsonImporterFactory_ParsingSprites", "Parsing Sprite Frame"), true, true);
+	bool bLoadedSuccessfully = true;
+
+	// Parse all of the frames
+	int32 FrameCount = 0;
+	for (auto FrameIt = ObjectBlock->Values.CreateIterator(); FrameIt; ++FrameIt)
+	{
+		GWarn->StatusUpdate(FrameCount, ObjectBlock->Values.Num(), NSLOCTEXT("Paper2D", "PaperJsonImporterFactory_ParsingSprites", "Parsing Sprite Frames"));
+
+		bool bReadFrameSuccessfully = true;
+
+		FSpriteFrame Frame;
+		Frame.FrameName = *FrameIt.Key();
+
+		if (FrameNames.Contains(Frame.FrameName))
+		{
+			bReadFrameSuccessfully = false;
+		}
+		else
+		{
+			FrameNames.Add(Frame.FrameName);
+		}
+
+		TSharedPtr<FJsonValue> FrameDataAsValue = FrameIt.Value();
+		TSharedPtr<FJsonObject> FrameData;
+		if (FrameDataAsValue->Type == EJson::Object)
+		{
+			FrameData = FrameDataAsValue->AsObject();
+			bReadFrameSuccessfully = bReadFrameSuccessfully && ParseFrame(FrameData, /*out*/Frame);
+		}
+		else
+		{
+			bReadFrameSuccessfully = false;
+		}
+
+		if (bReadFrameSuccessfully)
+		{
+			OutSpriteFrames.Add(Frame);
+		}
+		else
+		{
+			UE_LOG(LogPaperJsonImporter, Warning, TEXT("Frame %s is in an unexpected format"), *Frame.FrameName.ToString());
+			bLoadedSuccessfully = false;
+		}
+
+		FrameCount++;
+	}
+
+	GWarn->EndSlowTask();
+	return bLoadedSuccessfully;
+}
+
+static bool ParseFramesFromSpriteArray(const TArray<TSharedPtr<FJsonValue>>& ArrayBlock, TArray<FSpriteFrame>& OutSpriteFrames, TSet<FName>& FrameNames)
+{
+	GWarn->BeginSlowTask(NSLOCTEXT("Paper2D", "PaperJsonImporterFactory_ParsingSprites", "Parsing Sprite Frame"), true, true);
+	bool bLoadedSuccessfully = true;
+
+	// Parse all of the frames
+	int32 FrameCount = 0;
+	for (int32 FrameCount = 0; FrameCount < ArrayBlock.Num(); ++FrameCount)
+	{
+		GWarn->StatusUpdate(FrameCount, ArrayBlock.Num(), NSLOCTEXT("Paper2D", "PaperJsonImporterFactory_ParsingSprites", "Parsing Sprite Frames"));
+		bool bReadFrameSuccessfully = true;
+		FSpriteFrame Frame;
+
+		const TSharedPtr<FJsonValue>& FrameDataAsValue = ArrayBlock[FrameCount];
+		if (FrameDataAsValue->Type == EJson::Object)
+		{
+			TSharedPtr<FJsonObject> FrameData;
+			FrameData = FrameDataAsValue->AsObject();
+
+			FString FrameFilename = FPaperJSONHelpers::ReadString(FrameData, TEXT("filename"), TEXT(""));
+			if (!FrameFilename.IsEmpty())
+			{
+				Frame.FrameName = FName(*FrameFilename); // case insensitive
+				if (FrameNames.Contains(Frame.FrameName))
+				{
+					bReadFrameSuccessfully = false;
+				}
+				else
+				{
+					FrameNames.Add(Frame.FrameName);
+				}
+
+				bReadFrameSuccessfully = bReadFrameSuccessfully && ParseFrame(FrameData, /*out*/Frame);
+			}
+			else
+			{
+				bReadFrameSuccessfully = false;
+			}
+		}
+		else
+		{
+			bReadFrameSuccessfully = false;
+		}
+
+		if (bReadFrameSuccessfully)
+		{
+			OutSpriteFrames.Add(Frame);
+		}
+		else
+		{
+			UE_LOG(LogPaperJsonImporter, Warning, TEXT("Frame %s is in an unexpected format"), *Frame.FrameName.ToString());
+			bLoadedSuccessfully = false;
+		}
+	}
+
+	GWarn->EndSlowTask();
+	return bLoadedSuccessfully;
+}
+
+static ESpritePivotMode::Type GetBestPivotType(FVector2D JsonPivot)
+{
+	// Not assuming layout of ESpritePivotMode
+	if (JsonPivot.X == 0 && JsonPivot.Y == 0)
+	{
+		return ESpritePivotMode::Top_Left;
+	}
+	else if (JsonPivot.X == 0.5f && JsonPivot.Y == 0)
+	{
+		return ESpritePivotMode::Top_Center;
+	}
+	else if (JsonPivot.X == 1.0f && JsonPivot.Y == 0)
+	{
+		return ESpritePivotMode::Top_Right;
+	}
+	else if (JsonPivot.X == 0 && JsonPivot.Y == 0.5f)
+	{
+		return ESpritePivotMode::Center_Left;
+	}
+	else if (JsonPivot.X == 0.5f && JsonPivot.Y == 0.5f)
+	{
+		return ESpritePivotMode::Center_Center;
+	}
+	else if (JsonPivot.X == 1.0f && JsonPivot.Y == 0.5f)
+	{
+		return ESpritePivotMode::Center_Right;
+	}
+	else if (JsonPivot.X == 0 && JsonPivot.Y == 1.0f)
+	{
+		return ESpritePivotMode::Bottom_Left;
+	}
+	else if (JsonPivot.X == 0.5f && JsonPivot.Y == 1.0f)
+	{
+		return ESpritePivotMode::Bottom_Center;
+	}
+	else if (JsonPivot.X == 1.0f && JsonPivot.Y == 1.0f)
+	{
+		return ESpritePivotMode::Bottom_Right;
+	}
+	else
+	{
+		return ESpritePivotMode::Custom;
+	}
+}
+
+// FPaperJsonSpriteSheetImporter
+////////////////////////////////////////////
+
+FPaperJsonSpriteSheetImporter::FPaperJsonSpriteSheetImporter() : bIsReimporting(false)
+{
+}
+
+void FPaperJsonSpriteSheetImporter::SetReimportData(const FString& ExistingTextureName, UTexture2D* ExistingTexture, const TArray<FString>& ExistingSpriteNames, const TArray< TAssetPtr<class UPaperSprite> >& ExistingSpriteAssetPtrs)
+{
+	this->ExistingTextureName = ExistingTextureName;
+	this->ExistingTexture = ExistingTexture;
+
+	check(ExistingSpriteNames.Num() == ExistingSpriteAssetPtrs.Num());
+	if (ExistingSpriteNames.Num() == ExistingSpriteAssetPtrs.Num())
+	{
+		for (int i = 0; i < ExistingSpriteAssetPtrs.Num(); ++i)
+		{
+			const TAssetPtr<class UPaperSprite> SpriteAssetPtr = ExistingSpriteAssetPtrs[i];
+			FStringAssetReference SpriteStringRef = SpriteAssetPtr.ToStringReference();
+			if (!SpriteStringRef.ToString().IsEmpty())
+			{
+				UPaperSprite* LoadedSprite = Cast<UPaperSprite>(StaticLoadObject(UPaperSprite::StaticClass(), NULL, *SpriteStringRef.ToString(), NULL, LOAD_None, NULL));
+				if (LoadedSprite != nullptr)
+				{
+					ExistingSprites.Add(ExistingSpriteNames[i], LoadedSprite);
+				}
+			}
+		}
+	}
+	bIsReimporting = true;
+}
+
+bool FPaperJsonSpriteSheetImporter::Import(TSharedPtr<FJsonObject> SpriteDescriptorObject, const FString& NameForErrors)
+{
+	bool bLoadedSuccessfully = ParseMetaBlock(NameForErrors, SpriteDescriptorObject, Image);
+	if (bLoadedSuccessfully)
+	{
+		TSharedPtr<FJsonObject> ObjectFrameBlock = FPaperJSONHelpers::ReadObject(SpriteDescriptorObject, TEXT("frames"));
+		TSet<FName> FrameNames;
+
+		if (ObjectFrameBlock.IsValid())
+		{
+			bLoadedSuccessfully = ParseFramesFromSpriteHash(ObjectFrameBlock, /*out*/Frames, /*inout*/FrameNames);
+		}
+		else
+		{
+			// Try loading as an array
+			TArray<TSharedPtr<FJsonValue>> ArrayBlock = FPaperJSONHelpers::ReadArray(SpriteDescriptorObject, TEXT("frames"));
+			if (ArrayBlock.Num() > 0)
+			{
+				bLoadedSuccessfully = ParseFramesFromSpriteArray(ArrayBlock, /*out*/Frames, /*inout*/FrameNames);
+			}
+			else
+			{
+				UE_LOG(LogPaperJsonImporter, Warning, TEXT("Failed to parse sprite descriptor file '%s'.  Missing frames block"), *NameForErrors);
+				bLoadedSuccessfully = false;
+			}
+		}
+
+		if (bLoadedSuccessfully && Frames.Num() == 0)
+		{
+			UE_LOG(LogPaperJsonImporter, Warning, TEXT("Failed to parse sprite descriptor file '%s'.  No frames loaded"), *NameForErrors);
+			bLoadedSuccessfully = false;
+		}
+	}
+	return bLoadedSuccessfully;
+}
+
+bool FPaperJsonSpriteSheetImporter::ImportFromString(const FString& FileContents, const FString& NameForErrors)
+{
+	TSharedPtr<FJsonObject> SpriteDescriptorObject = ParseJSON(FileContents, NameForErrors);
+	return SpriteDescriptorObject.IsValid() &&
+		Import(SpriteDescriptorObject, NameForErrors);
+}
+
+bool FPaperJsonSpriteSheetImporter::ImportFromArchive(FArchive* Archive, const FString& NameForErrors)
+{
+	TSharedPtr<FJsonObject> SpriteDescriptorObject = ParseJSON(Archive, NameForErrors);
+	return SpriteDescriptorObject.IsValid() &&
+		Import(SpriteDescriptorObject, NameForErrors);
+}
+
+
+bool FPaperJsonSpriteSheetImporter::ImportTextures(const FString& LongPackagePath, const FString& SourcePath)
+{
+	bool bLoadedSuccessfully = true;
+	bool bFoundExistingTexture = false;
+	const FString TargetSubPath = LongPackagePath + TEXT("/Textures");
+
+	const FString SourceSheetImageFilename = FPaths::Combine(*SourcePath, *Image);
+	TArray<FString> SheetFileNames;
+	SheetFileNames.Add(SourceSheetImageFilename);
+
+	if (bIsReimporting && ExistingTextureName == Image && ExistingTexture != nullptr)
+	{
+		ImageTexture = ExistingTexture;
+		bFoundExistingTexture = true;
+	}
+
+	// Import the texture if we need to
+	if (!bIsReimporting || (bIsReimporting && !bFoundExistingTexture))
+	{
+		FAssetToolsModule& AssetToolsModule = FModuleManager::GetModuleChecked<FAssetToolsModule>("AssetTools");
+		TArray<UObject*> ImportedSheets = AssetToolsModule.Get().ImportAssets(SheetFileNames, TargetSubPath);
+
+		UTexture2D* ImportedTexture = (ImportedSheets.Num() > 0) ? Cast<UTexture2D>(ImportedSheets[0]) : NULL;
+
+		if (ImportedTexture == NULL)
+		{
+			UE_LOG(LogPaperJsonImporter, Warning, TEXT("Failed to import sprite sheet image '%s'."), *SourceSheetImageFilename);
+			bLoadedSuccessfully = false;
+		}
+		else
+		{
+			// Change the compression settings
+			ImportedTexture->Modify();
+			ImportedTexture->LODGroup = TEXTUREGROUP_UI;
+			ImportedTexture->CompressionSettings = TC_EditorIcon;
+			ImportedTexture->Filter = TF_Nearest;
+			ImportedTexture->PostEditChange();
+		}
+
+		ImageTexture = ImportedTexture;
+	}
+
+	return bLoadedSuccessfully;
+}
+
+UPaperSprite* FPaperJsonSpriteSheetImporter::FindExistingSprite(const FString& Name)
+{
+	UPaperSprite**ExistingSprite = ExistingSprites.Find(Name);
+	if (ExistingSprite != nullptr)
+	{
+		return *ExistingSprite;
+	}
+	else
+	{
+		return nullptr;
+	}
+}
+
+bool FPaperJsonSpriteSheetImporter::PerformImport(const FString& LongPackagePath, EObjectFlags Flags, UPaperSpriteSheet* SpriteSheet)
+{
+	FAssetToolsModule& AssetToolsModule = FModuleManager::GetModuleChecked<FAssetToolsModule>("AssetTools");
+
+	GWarn->BeginSlowTask(NSLOCTEXT("Paper2D", "PaperJsonImporterFactory_ImportingSprites", "Importing Sprite Frame"), true, true);
+	for (int32 FrameIndex = 0; FrameIndex < Frames.Num(); ++FrameIndex)
+	{
+		GWarn->StatusUpdate(FrameIndex, Frames.Num(), NSLOCTEXT("Paper2D", "PaperJsonImporterFactory_ImportingSprites", "Importing Sprite Frames"));
+
+		// Check for user canceling the import
+		if (GWarn->ReceivedUserCancel())
+		{
+			break;
+		}
+
+		const FSpriteFrame& Frame = Frames[FrameIndex];
+
+		// Create a package for the frame
+		const FString TargetSubPath = LongPackagePath + TEXT("/Frames");
+
+		UObject* OuterForFrame = NULL; // @TODO: Use this if we don't want them to be individual assets - Flipbook;
+
+		// Create a unique package name and asset name for the frame
+		const FString TentativePackagePath = PackageTools::SanitizePackageName(TargetSubPath + TEXT("/") + Frame.FrameName.ToString());
+		FString DefaultSuffix;
+		FString AssetName;
+		FString PackageName;
+		AssetToolsModule.Get().CreateUniqueAssetName(TentativePackagePath, /*out*/ DefaultSuffix, /*out*/ PackageName, /*out*/ AssetName);
+
+		// Create a package for the frame
+		OuterForFrame = CreatePackage(NULL, *PackageName);
+
+		// Create a frame in the package
+		UPaperSprite* TargetSprite = nullptr;
+		
+		if (bIsReimporting)
+		{
+			TargetSprite = FindExistingSprite(Frame.FrameName.ToString());
+		}
+		
+		if (TargetSprite == nullptr)
+		{
+			TargetSprite = NewNamedObject<UPaperSprite>(OuterForFrame, *AssetName, Flags);
+			FAssetRegistryModule::AssetCreated(TargetSprite);
+		}
+
+		TargetSprite->Modify();
+		TargetSprite->InitializeSprite(ImageTexture, Frame.SpritePosInSheet, Frame.SpriteSizeInSheet, GetDefault<UPaperRuntimeSettings>()->DefaultPixelsPerUnrealUnit);
+
+		if (Frame.bRotated)
+		{
+			TargetSprite->SetRotated(true);
+		}
+
+		if (Frame.bTrimmed)
+		{
+			TargetSprite->SetTrim(Frame.bTrimmed, Frame.SpriteSourcePos, Frame.ImageSourceSize);
+		}
+
+		// Set up pivot on object based on Texture Packer json
+		ESpritePivotMode::Type PivotType = GetBestPivotType(Frame.Pivot);
+		FVector2D TextureSpacePivotPoint = FVector2D::ZeroVector;
+		if (PivotType == ESpritePivotMode::Custom)
+		{
+			TextureSpacePivotPoint.X = Frame.SpritePosInSheet.X - Frame.SpriteSourcePos.X + Frame.ImageSourceSize.X * Frame.Pivot.X;
+			TextureSpacePivotPoint.Y = Frame.SpritePosInSheet.Y - Frame.SpriteSourcePos.Y + Frame.ImageSourceSize.Y * Frame.Pivot.Y;
+		}
+		TargetSprite->SetPivot(PivotType, TextureSpacePivotPoint);
+
+		// Create the entry in the animation
+		SpriteSheet->SpriteNames.Add(Frame.FrameName.ToString());
+		SpriteSheet->Sprites.Add(TargetSprite);
+		
+		TargetSprite->PostEditChange();
+	}
+
+	SpriteSheet->TextureName = Image;
+	SpriteSheet->Texture = ImageTexture;
+
+	GWarn->EndSlowTask();
+	return true;
+}
diff --git a/Engine/Plugins/Experimental/Paper2D/Source/PaperJsonImporter/Private/PaperJsonSpriteSheetImporter.h b/Engine/Plugins/Experimental/Paper2D/Source/PaperJsonImporter/Private/PaperJsonSpriteSheetImporter.h
new file mode 100644
--- /dev/null
+++ b/Engine/Plugins/Experimental/Paper2D/Source/PaperJsonImporter/Private/PaperJsonSpriteSheetImporter.h
@@ -0,0 +1,59 @@
+// Copyright 1998-2014 Epic Games, Inc. All Rights Reserved.
+
+#pragma once
+
+#include "PaperJsonImporterPrivatePCH.h"
+#include "Paper2DClasses.h"
+
+//////////////////////////////////////////////////////////////////////////
+// FSpriteFrame
+
+// Represents one parsed frame in a sprite sheet
+struct FSpriteFrame
+{
+	FName FrameName;
+
+	FVector2D SpritePosInSheet;
+	FVector2D SpriteSizeInSheet;
+
+	FVector2D SpriteSourcePos;
+	FVector2D SpriteSourceSize;
+
+	FVector2D ImageSourceSize;
+
+	FVector2D Pivot;
+
+	bool bTrimmed;
+	bool bRotated;
+};
+
+//////////////////////////////////////////////////////////////////////////
+// FJsonPaperSpriteSheetImporter
+
+// Parses a json from FileContents and imports / reimports a spritesheet
+class FPaperJsonSpriteSheetImporter
+{
+public:
+	FPaperJsonSpriteSheetImporter();
+	bool ImportFromString(const FString& FileContents, const FString& NameForErrors);
+	bool ImportFromArchive(FArchive* Archive, const FString& NameForErrors);
+
+	bool ImportTextures(const FString& LongPackagePath, const FString& SourcePath);
+	bool PerformImport(const FString& LongPackagePath, EObjectFlags Flags, UPaperSpriteSheet* SpriteSheet);
+
+	void SetReimportData(const FString& ExistingImageName, UTexture2D* ExistingTexture, const TArray<FString>& ExistingSpriteNames, const TArray< TAssetPtr<class UPaperSprite> >& ExistingSpriteAssetPtrs);
+
+protected:
+	bool Import(TSharedPtr<FJsonObject> SpriteDescriptorObject, const FString& NameForErrors);
+	UPaperSprite* FindExistingSprite(const FString& Name);
+
+	TArray<FSpriteFrame> Frames;
+	FString Image;
+	UTexture2D* ImageTexture;
+
+	bool bIsReimporting;
+ 	FString ExistingTextureName;					// Name of the last imported texture
+	UTexture2D* ExistingTexture;					// And the texture itself
+	TMap<FString, UPaperSprite*> ExistingSprites;	// Map of a sprite name (as seen in the importer) -> UPaperSprite
+};
+
diff --git a/Engine/Plugins/Experimental/Paper2D/Source/PaperJsonImporter/Private/PaperSpriteSheet.cpp b/Engine/Plugins/Experimental/Paper2D/Source/PaperJsonImporter/Private/PaperSpriteSheet.cpp
new file mode 100644
--- /dev/null
+++ b/Engine/Plugins/Experimental/Paper2D/Source/PaperJsonImporter/Private/PaperSpriteSheet.cpp
@@ -0,0 +1,10 @@
+// Copyright 1998-2014 Epic Games, Inc. All Rights Reserved.
+
+#include "PaperJsonImporterPrivatePCH.h"
+#include "PaperSpriteSheet.h"
+
+UPaperSpriteSheet::UPaperSpriteSheet(const FPostConstructInitializeProperties& PCIP)
+	: Super(PCIP)
+{
+}
+
diff --git a/Engine/Plugins/Experimental/Paper2D/Source/PaperJsonImporter/Private/PaperSpriteSheet.h b/Engine/Plugins/Experimental/Paper2D/Source/PaperJsonImporter/Private/PaperSpriteSheet.h
new file mode 100644
--- /dev/null
+++ b/Engine/Plugins/Experimental/Paper2D/Source/PaperJsonImporter/Private/PaperSpriteSheet.h
@@ -0,0 +1,32 @@
+// Copyright 1998-2014 Epic Games, Inc. All Rights Reserved.
+
+#pragma once
+
+#include "PaperSpriteSheet.generated.h"
+
+UCLASS(BlueprintType, meta = (DisplayThumbnail = "true"))
+class UPaperSpriteSheet : public UObject
+{
+	GENERATED_UCLASS_BODY()
+
+public:
+	// The names of sprites during import
+	UPROPERTY(VisibleAnywhere, Category=Data)
+	TArray<FString> SpriteNames;
+
+	UPROPERTY(VisibleAnywhere, Category = Data)
+	TArray< TAssetPtr<class UPaperSprite> > Sprites;
+
+	// The name of the texture during import
+	UPROPERTY(VisibleAnywhere, Category = Data)
+	FString TextureName;
+
+	UPROPERTY(VisibleAnywhere, Category = Data)
+	UTexture2D* Texture;
+
+#if WITH_EDITORONLY_DATA
+	// Import data for this 
+	UPROPERTY(VisibleAnywhere, EditInline, Category=ImportSettings)
+	class UAssetImportData* AssetImportData;
+#endif
+};
diff --git a/Engine/Plugins/Experimental/Paper2D/Source/PaperJsonImporter/Private/PaperSpriteSheetAssetTypeActions.cpp b/Engine/Plugins/Experimental/Paper2D/Source/PaperJsonImporter/Private/PaperSpriteSheetAssetTypeActions.cpp
new file mode 100644
--- /dev/null
+++ b/Engine/Plugins/Experimental/Paper2D/Source/PaperJsonImporter/Private/PaperSpriteSheetAssetTypeActions.cpp
@@ -0,0 +1,255 @@
+// Copyright 1998-2014 Epic Games, Inc. All Rights Reserved.
+
+#include "PaperJsonImporterPrivatePCH.h"
+#include "Paper2DClasses.h"
+#include "PaperSpriteSheetAssetTypeActions.h"
+#include "AssetToolsModule.h"
+#include "ContentBrowserModule.h"
+#include "AssetRegistryModule.h"
+#include "PaperFlipbookFactory.h"
+
+#define LOCTEXT_NAMESPACE "AssetTypeActions"
+
+//////////////////////////////////////////////////////////////////////////
+// FPaperSpriteSheetImportAssetTypeActions
+
+FText FPaperSpriteSheetAssetTypeActions::GetName() const
+{
+	return LOCTEXT("FSpriteSheetAssetTypeActionsName", "Sprite Sheet");
+}
+
+FColor FPaperSpriteSheetAssetTypeActions::GetTypeColor() const
+{
+	return FColor(0, 255, 255);
+}
+
+UClass* FPaperSpriteSheetAssetTypeActions::GetSupportedClass() const
+{
+	return UPaperSpriteSheet::StaticClass();
+}
+
+uint32 FPaperSpriteSheetAssetTypeActions::GetCategories()
+{
+	return EAssetTypeCategories::Misc;
+}
+
+void FPaperSpriteSheetAssetTypeActions::GetActions(const TArray<UObject*>& InObjects, FMenuBuilder& MenuBuilder)
+{
+	auto SpriteSheetImports = GetTypedWeakObjectPtrs<UPaperSpriteSheet>(InObjects);
+
+	MenuBuilder.AddMenuEntry(
+		LOCTEXT("SpriteSheet_Reimport", "Reimport"),
+		LOCTEXT("SpriteSheet_ReimportTooltip", "Reimports the sprite sheet."),
+		FSlateIcon(),
+		FUIAction(
+		FExecuteAction::CreateSP(this, &FPaperSpriteSheetAssetTypeActions::ExecuteReimport, SpriteSheetImports),
+		FCanExecuteAction()
+		)
+	);
+
+	MenuBuilder.AddMenuEntry(
+		LOCTEXT("SpriteSheet_CreateFlipbooks", "Create Flipbooks"),
+		LOCTEXT("SpriteSheet_CreateFlipbooksTooltip", "Creates flipbooks from sprites in this sprite sheet."),
+		FSlateIcon(),
+		FUIAction(
+		FExecuteAction::CreateSP(this, &FPaperSpriteSheetAssetTypeActions::ExecuteCreateFlipbooks, SpriteSheetImports),
+		FCanExecuteAction()
+		)
+	);
+}
+
+//////////////////////////////////////////////////////////////////////////
+
+void FPaperSpriteSheetAssetTypeActions::ExecuteReimport(TArray<TWeakObjectPtr<UPaperSpriteSheet>> Objects)
+{
+	for (auto ObjIt = Objects.CreateConstIterator(); ObjIt; ++ObjIt)
+	{
+		auto Object = (*ObjIt).Get();
+		if (Object)
+		{
+			FReimportManager::Instance()->Reimport(Object, /*bAskForNewFileIfMissing=*/true);
+		}
+	}
+}
+
+static bool ExtractSpriteNumber(const FString& String, FString& BareString, int& Number)
+{
+	bool bExtracted = false;
+	int LastCharacter = String.Len() - 1;
+	if (LastCharacter >= 0 && FChar::IsDigit(String[LastCharacter]))
+	{
+		while (LastCharacter > 0 && FChar::IsDigit(String[LastCharacter - 1]))
+		{
+			LastCharacter--;
+		}
+
+		if (LastCharacter >= 0)
+		{
+			int FirstDigit = LastCharacter;
+			int EndCharacter = FirstDigit;
+			while (EndCharacter > 0 && !FChar::IsAlnum(String[EndCharacter - 1]))
+			{
+				EndCharacter--;
+			}
+
+			if (EndCharacter == 0)
+			{
+				// This string consists of non alnum + number, eg. _42
+				// The flipbook / category name in this case will be _
+				// Otherwise, we strip out all trailing non-alnum chars
+				EndCharacter = FirstDigit;
+			}
+			
+			FString NumberString = String.Mid(FirstDigit);
+			BareString = String.Left(EndCharacter);
+			Number = FCString::Atoi(*NumberString);
+			bExtracted = true;
+		}
+	}
+
+	if (!bExtracted)
+	{
+		BareString = String;
+		Number = -1;
+	}
+	return bExtracted;
+}
+
+void FPaperSpriteSheetAssetTypeActions::ExecuteCreateFlipbooks(TArray<TWeakObjectPtr<UPaperSpriteSheet>> Objects)
+{	
+	for (int SpriteSheetIndex = 0; SpriteSheetIndex < Objects.Num(); ++SpriteSheetIndex)
+	{
+		UPaperSpriteSheet* SpriteSheet = Objects[SpriteSheetIndex].Get();
+		if (SpriteSheet != nullptr)
+		{
+			const FString PackagePath = FPackageName::GetLongPackagePath(SpriteSheet->GetOutermost()->GetPathName());
+
+			FAssetToolsModule& AssetToolsModule = FModuleManager::GetModuleChecked<FAssetToolsModule>("AssetTools");
+			FContentBrowserModule& ContentBrowserModule = FModuleManager::LoadModuleChecked<FContentBrowserModule>("ContentBrowser");
+			
+			TMap<FString, UPaperSprite*> SpriteNameMap;
+			TSet<FString> FlipbookNames;
+			TMap<FString, TArray<UPaperSprite*> > FlipbookSprites;
+			TArray<UPaperSprite*> AllValidSprites;
+
+			check(SpriteSheet->SpriteNames.Num() == SpriteSheet->Sprites.Num());
+			bool useSpriteNames = (SpriteSheet->SpriteNames.Num() == SpriteSheet->Sprites.Num());
+
+			for (int SpriteIndex = 0; SpriteIndex < SpriteSheet->Sprites.Num(); ++SpriteIndex)
+			{
+				auto SpriteAssetPtr = SpriteSheet->Sprites[SpriteIndex];
+				FStringAssetReference SpriteStringRef = SpriteAssetPtr.ToStringReference();
+				UPaperSprite* Sprite = nullptr;
+				if (!SpriteStringRef.ToString().IsEmpty())
+				{
+					Sprite = Cast<UPaperSprite>(StaticLoadObject(UPaperSprite::StaticClass(), NULL, *SpriteStringRef.ToString(), NULL, LOAD_None, NULL));
+				}
+				if (Sprite != nullptr)
+				{
+					const FString SpriteName = useSpriteNames ? SpriteSheet->SpriteNames[SpriteIndex] : Sprite->GetName();
+
+					SpriteNameMap.Add(SpriteName, Sprite);
+					
+					int SpriteNumber = 0;
+					FString SpriteBareString;
+					if (ExtractSpriteNumber(SpriteName, /*out*/SpriteBareString, /*out*/SpriteNumber))
+					{
+						FlipbookNames.Add(SpriteBareString);
+						if (FlipbookSprites.Contains(SpriteBareString))
+						{
+							FlipbookSprites[SpriteBareString].Add(Sprite);
+						}
+						else
+						{
+							FlipbookSprites.Add(SpriteBareString, TArray<UPaperSprite*>());
+						}
+					}
+
+					AllValidSprites.Add(Sprite);
+				}
+			}
+
+			// Natural sort using the same method as above
+			struct FSpriteSortPredicate
+			{
+				FSpriteSortPredicate() {}
+
+				// Sort predicate operator
+				bool operator()(UPaperSprite& LHS, UPaperSprite& RHS) const
+				{
+					FString LeftString;
+					int LeftNumber;
+					ExtractSpriteNumber(LHS.GetName(), /*out*/LeftString, /*out*/LeftNumber);
+
+					FString RightString;
+					int RightNumber;
+					ExtractSpriteNumber(RHS.GetName(), /*out*/RightString, /*out*/RightNumber);
+
+					return (LeftString == RightString) ? LeftNumber < RightNumber : LeftString < RightString;
+				}
+			};
+
+			// Haven't matched any existing names, so create one flipbook with all sprites in it
+			if (FlipbookNames.Num() == 0 &&
+				EAppReturnType::Yes == FMessageDialog::Open(EAppMsgType::YesNo, NSLOCTEXT("Paper2D", "Paper2D_SpriteSheetCreateOneFlipbook", "Unable to detect flipbooks by name in this sprite sheet. Proceed to create a single flipbook containing all sprites in this sprite sheet?")) )
+			{
+				const FString& DefaultFlipbookName = SpriteSheet->GetName() + TEXT(" Flipbook");
+				FlipbookNames.Add(DefaultFlipbookName);
+				FlipbookSprites.Add(DefaultFlipbookName, AllValidSprites);
+			}
+
+			// Create one flipbook for every grouped flipbook name
+			if (FlipbookNames.Num() > 0)
+			{
+				UPaperFlipbookFactory* FlipbookFactory = ConstructObject<UPaperFlipbookFactory>(UPaperFlipbookFactory::StaticClass());
+
+				GWarn->BeginSlowTask(NSLOCTEXT("Paper2D", "Paper2D_SpriteSheetCreateFlipbooks", "Creating flipbooks from Sprite Sheet"), true, true);
+
+				TArray<FString> FlipbookNameArray = FlipbookNames.Array();
+				TArray<UObject*> ObjectsToSync;
+				for (int FlipbookNameIndex = 0; FlipbookNameIndex < FlipbookNames.Num(); ++FlipbookNameIndex)
+				{
+					GWarn->UpdateProgress(FlipbookNameIndex, FlipbookNames.Num());
+
+					const FString& FlipbookName = FlipbookNameArray[FlipbookNameIndex];
+					TArray<UPaperSprite*>& Sprites = FlipbookSprites[FlipbookName];
+					Sprites.Sort(FSpriteSortPredicate());
+
+					const FString TentativePackagePath = PackageTools::SanitizePackageName(PackagePath + TEXT("/") + FlipbookName);
+					FString DefaultSuffix;
+					FString AssetName;
+					FString PackageName;
+					AssetToolsModule.Get().CreateUniqueAssetName(TentativePackagePath, /*out*/ DefaultSuffix, /*out*/ PackageName, /*out*/ AssetName);
+
+					FlipbookFactory->KeyFrames.Empty();
+					for (int32 SpriteIndex = 0; SpriteIndex < Sprites.Num(); ++SpriteIndex)
+					{
+						FPaperFlipbookKeyFrame* KeyFrame = new (FlipbookFactory->KeyFrames) FPaperFlipbookKeyFrame();
+						KeyFrame->Sprite = Sprites[SpriteIndex];
+						KeyFrame->FrameRun = 1;
+					}
+
+					if (UObject* NewAsset = AssetToolsModule.Get().CreateAsset(AssetName, PackagePath, UPaperFlipbook::StaticClass(), FlipbookFactory))
+					{
+						ObjectsToSync.Add(NewAsset);
+					}
+
+					if (GWarn->ReceivedUserCancel())
+					{
+						break;
+					}
+				}
+
+				GWarn->EndSlowTask();
+
+				if (ObjectsToSync.Num() > 0)
+				{
+					ContentBrowserModule.Get().SyncBrowserToAssets(ObjectsToSync);
+				}
+			}
+		}
+	}
+}
+
+#undef LOCTEXT_NAMESPACE
+
diff --git a/Engine/Plugins/Experimental/Paper2D/Source/PaperJsonImporter/Private/PaperSpriteSheetAssetTypeActions.h b/Engine/Plugins/Experimental/Paper2D/Source/PaperJsonImporter/Private/PaperSpriteSheetAssetTypeActions.h
new file mode 100644
--- /dev/null
+++ b/Engine/Plugins/Experimental/Paper2D/Source/PaperJsonImporter/Private/PaperSpriteSheetAssetTypeActions.h
@@ -0,0 +1,22 @@
+// Copyright 1998-2014 Epic Games, Inc. All Rights Reserved.
+
+#pragma once
+
+#include "AssetTypeActions_Base.h"
+
+class FPaperSpriteSheetAssetTypeActions : public FAssetTypeActions_Base
+{
+public:
+	// IAssetTypeActions interface
+	virtual FText GetName() const override;
+	virtual FColor GetTypeColor() const override;
+	virtual UClass* GetSupportedClass() const override;
+	virtual bool HasActions(const TArray<UObject*>& InObjects) const override { return true; }
+	virtual void GetActions(const TArray<UObject*>& InObjects, FMenuBuilder& MenuBuilder) override;
+	virtual uint32 GetCategories() override;
+	// End of IAssetTypeActions interface
+
+private:
+	void ExecuteReimport(TArray<TWeakObjectPtr<UPaperSpriteSheet>> Objects);
+	void ExecuteCreateFlipbooks(TArray<TWeakObjectPtr<UPaperSpriteSheet>> Objects);
+};
\ No newline at end of file
diff --git a/Engine/Plugins/Experimental/Paper2D/Source/PaperJsonImporter/Private/ReimportPaperJsonImporterFactory.cpp b/Engine/Plugins/Experimental/Paper2D/Source/PaperJsonImporter/Private/ReimportPaperJsonImporterFactory.cpp
new file mode 100644
--- /dev/null
+++ b/Engine/Plugins/Experimental/Paper2D/Source/PaperJsonImporter/Private/ReimportPaperJsonImporterFactory.cpp
@@ -0,0 +1,83 @@
+// Copyright 1998-2014 Epic Games, Inc. All Rights Reserved.
+
+#pragma once
+#include "PaperJsonImporterPrivatePCH.h"
+#include "ReimportPaperJsonImporterFactory.h"
+
+#define LOCTEXT_NAMESPACE "PaperJsonImporter"
+
+
+UReimportPaperJsonImporterFactory::UReimportPaperJsonImporterFactory(const class FPostConstructInitializeProperties& PCIP)
+	: Super(PCIP)
+{
+	SupportedClass = UPaperSpriteSheet::StaticClass();
+
+	bCreateNew = false;
+	bText = true;
+}
+
+bool UReimportPaperJsonImporterFactory::CanReimport(UObject* Obj, TArray<FString>& OutFilenames)
+{
+	UPaperSpriteSheet* SpriteSheet = Cast<UPaperSpriteSheet>(Obj);
+	if (SpriteSheet && SpriteSheet->AssetImportData)
+	{
+		OutFilenames.Add(FReimportManager::ResolveImportFilename(SpriteSheet->AssetImportData->SourceFilePath, SpriteSheet));
+		return true;
+	}
+	return false;
+}
+
+void UReimportPaperJsonImporterFactory::SetReimportPaths(UObject* Obj, const TArray<FString>& NewReimportPaths)
+{
+	UPaperSpriteSheet* SpriteSheet = Cast<UPaperSpriteSheet>(Obj);
+	if (SpriteSheet && ensure(NewReimportPaths.Num() == 1))
+	{
+		SpriteSheet->AssetImportData->SourceFilePath = FReimportManager::ResolveImportFilename(NewReimportPaths[0], SpriteSheet);
+	}
+}
+
+EReimportResult::Type UReimportPaperJsonImporterFactory::Reimport(UObject* Obj)
+{
+	UPaperSpriteSheet* SpriteSheet = Cast<UPaperSpriteSheet>(Obj);
+	if (!SpriteSheet)
+	{
+		return EReimportResult::Failed;
+	}
+
+	// Make sure file is valid and exists
+	const FString Filename = FReimportManager::ResolveImportFilename(SpriteSheet->AssetImportData->SourceFilePath, SpriteSheet);
+	if (!Filename.Len() || IFileManager::Get().FileSize(*Filename) == INDEX_NONE)
+	{
+		return EReimportResult::Failed;
+	}
+
+	ExistingSpriteNames = SpriteSheet->SpriteNames;
+	ExistingSprites = SpriteSheet->Sprites;
+	ExistingTextureName = SpriteSheet->TextureName;
+	ExistingTexture = SpriteSheet->Texture;
+
+	EReimportResult::Type Result = EReimportResult::Failed;
+	if (UFactory::StaticImportObject(SpriteSheet->GetClass(), SpriteSheet->GetOuter(), *SpriteSheet->GetName(), RF_Public | RF_Standalone, *Filename, NULL, this))
+	{
+		UE_LOG(LogPaperJsonImporter, Log, TEXT("Imported successfully"));
+		// Try to find the outer package so we can dirty it up
+		if (SpriteSheet->GetOuter())
+		{
+			SpriteSheet->GetOuter()->MarkPackageDirty();
+		}
+		else
+		{
+			SpriteSheet->MarkPackageDirty();
+		}
+		Result = EReimportResult::Succeeded;
+	}
+	else
+	{
+		UE_LOG(LogPaperJsonImporter, Warning, TEXT("-- import failed"));
+		Result = EReimportResult::Failed;
+	}
+
+	return Result;
+}
+
+#undef LOCTEXT_NAMESPACE
diff --git a/Engine/Plugins/Experimental/Paper2D/Source/PaperJsonImporter/Private/ReimportPaperJsonImporterFactory.h b/Engine/Plugins/Experimental/Paper2D/Source/PaperJsonImporter/Private/ReimportPaperJsonImporterFactory.h
new file mode 100644
--- /dev/null
+++ b/Engine/Plugins/Experimental/Paper2D/Source/PaperJsonImporter/Private/ReimportPaperJsonImporterFactory.h
@@ -0,0 +1,24 @@
+// Copyright 1998-2014 Epic Games, Inc. All Rights Reserved.
+
+#pragma once
+#include "ReimportPaperJsonImporterFactory.generated.h"
+
+// Imports a paper animated sprite (and associated paper sprites and textures) from an Adobe Flash CS6 exported JSON file
+UCLASS()
+class UReimportPaperJsonImporterFactory : public UPaperJsonImporterFactory, public FReimportHandler
+{
+	GENERATED_UCLASS_BODY()
+
+	// Begin FReimportHandler interface
+	virtual bool CanReimport(UObject* Obj, TArray<FString>& OutFilenames) override;
+	virtual void SetReimportPaths(UObject* Obj, const TArray<FString>& NewReimportPaths) override;
+	virtual EReimportResult::Type Reimport(UObject* Obj) override;
+	// End FReimportHandler interface
+
+	TArray<FString> ExistingSpriteNames;
+	TArray< TAssetPtr<class UPaperSprite> > ExistingSprites;
+	
+	FString ExistingTextureName;
+	UTexture2D* ExistingTexture;
+};
+
