Paper2D: Fixed sprite thumbnail renderer assert when drawing an invalid sprite.

diff --git a/Engine/Plugins/Experimental/Paper2D/Source/Paper2DEditor/Private/PaperSpriteThumbnailRenderer.cpp b/Engine/Plugins/Experimental/Paper2D/Source/Paper2DEditor/Private/PaperSpriteThumbnailRenderer.cpp
--- a/Engine/Plugins/Experimental/Paper2D/Source/Paper2DEditor/Private/PaperSpriteThumbnailRenderer.cpp
+++ b/Engine/Plugins/Experimental/Paper2D/Source/Paper2DEditor/Private/PaperSpriteThumbnailRenderer.cpp
@@ -96,18 +96,21 @@ void UPaperSpriteThumbnailRenderer::DrawFrame(class UPaperSprite* Sprite, int32
 		}
 
 		// Draw triangles
-		TArray<FCanvasUVTri> Triangles;
-		const FLinearColor SpriteColor(FLinearColor::White);
-		for (int Vertex = 0; Vertex < CanvasPositions.Num(); Vertex += 3)
+		if (CanvasPositions.Num() > 0)
 		{
-			FCanvasUVTri *Triangle = new(Triangles)FCanvasUVTri();
-			Triangle->V0_Pos = CanvasPositions[Vertex + 0]; Triangle->V0_UV = CanvasUVs[Vertex + 0]; Triangle->V0_Color = SpriteColor;
-			Triangle->V1_Pos = CanvasPositions[Vertex + 1]; Triangle->V1_UV = CanvasUVs[Vertex + 1]; Triangle->V1_Color = SpriteColor;
-			Triangle->V2_Pos = CanvasPositions[Vertex + 2]; Triangle->V2_UV = CanvasUVs[Vertex + 2]; Triangle->V2_Color = SpriteColor;
+			TArray<FCanvasUVTri> Triangles;
+			const FLinearColor SpriteColor(FLinearColor::White);
+			for (int Vertex = 0; Vertex < CanvasPositions.Num(); Vertex += 3)
+			{
+				FCanvasUVTri *Triangle = new(Triangles)FCanvasUVTri();
+				Triangle->V0_Pos = CanvasPositions[Vertex + 0]; Triangle->V0_UV = CanvasUVs[Vertex + 0]; Triangle->V0_Color = SpriteColor;
+				Triangle->V1_Pos = CanvasPositions[Vertex + 1]; Triangle->V1_UV = CanvasUVs[Vertex + 1]; Triangle->V1_Color = SpriteColor;
+				Triangle->V2_Pos = CanvasPositions[Vertex + 2]; Triangle->V2_UV = CanvasUVs[Vertex + 2]; Triangle->V2_Color = SpriteColor;
+			}
+			FCanvasTriangleItem CanvasTriangle(Triangles, SourceTexture->Resource);
+			CanvasTriangle.BlendMode = bUseTranslucentBlend ? ESimpleElementBlendMode::SE_BLEND_Translucent : ESimpleElementBlendMode::SE_BLEND_Opaque;
+			Canvas->DrawItem(CanvasTriangle);
 		}
-		FCanvasTriangleItem CanvasTriangle(Triangles, SourceTexture->Resource);
-		CanvasTriangle.BlendMode = bUseTranslucentBlend ? ESimpleElementBlendMode::SE_BLEND_Translucent : ESimpleElementBlendMode::SE_BLEND_Opaque;
-		Canvas->DrawItem(CanvasTriangle);
 	}
 	else
 	{
