Attempt to fix Paper2D plugin nonunity build

diff --git a/Engine/Plugins/Experimental/Paper2D/Source/PaperJsonImporter/Private/PaperSpriteSheetAssetTypeActions.cpp b/Engine/Plugins/Experimental/Paper2D/Source/PaperJsonImporter/Private/PaperSpriteSheetAssetTypeActions.cpp
--- a/Engine/Plugins/Experimental/Paper2D/Source/PaperJsonImporter/Private/PaperSpriteSheetAssetTypeActions.cpp
+++ b/Engine/Plugins/Experimental/Paper2D/Source/PaperJsonImporter/Private/PaperSpriteSheetAssetTypeActions.cpp
@@ -2,6 +2,7 @@
 
 #include "PaperJsonImporterPrivatePCH.h"
 #include "Paper2DClasses.h"
+#include "PaperSpriteSheet.h"
 #include "PaperSpriteSheetAssetTypeActions.h"
 #include "AssetToolsModule.h"
 #include "ContentBrowserModule.h"
diff --git a/Engine/Plugins/Experimental/Paper2D/Source/PaperJsonImporter/Private/PaperSpriteSheetAssetTypeActions.h b/Engine/Plugins/Experimental/Paper2D/Source/PaperJsonImporter/Private/PaperSpriteSheetAssetTypeActions.h
--- a/Engine/Plugins/Experimental/Paper2D/Source/PaperJsonImporter/Private/PaperSpriteSheetAssetTypeActions.h
+++ b/Engine/Plugins/Experimental/Paper2D/Source/PaperJsonImporter/Private/PaperSpriteSheetAssetTypeActions.h
@@ -4,6 +4,8 @@
 
 #include "AssetTypeActions_Base.h"
 
+class UPaperSpriteSheet;
+
 class FPaperSpriteSheetAssetTypeActions : public FAssetTypeActions_Base
 {
 public:
