Fix non-unity for Paper2D

diff --git a/Engine/Plugins/Experimental/Paper2D/Source/PaperJsonImporter/Private/PaperJsonSpriteSheetImporter.cpp b/Engine/Plugins/Experimental/Paper2D/Source/PaperJsonImporter/Private/PaperJsonSpriteSheetImporter.cpp
--- a/Engine/Plugins/Experimental/Paper2D/Source/PaperJsonImporter/Private/PaperJsonSpriteSheetImporter.cpp
+++ b/Engine/Plugins/Experimental/Paper2D/Source/PaperJsonImporter/Private/PaperJsonSpriteSheetImporter.cpp
@@ -4,6 +4,7 @@
 #include "Paper2DClasses.h"
 #include "Json.h"
 #include "PaperJSONHelpers.h"
+#include "PaperJsonSpriteSheetImporter.h"
 #include "AssetToolsModule.h"
 #include "AssetRegistryModule.h"
 #include "PackageTools.h"
diff --git a/Engine/Plugins/Experimental/Paper2D/Source/PaperJsonImporter/Private/PaperJsonSpriteSheetImporter.h b/Engine/Plugins/Experimental/Paper2D/Source/PaperJsonImporter/Private/PaperJsonSpriteSheetImporter.h
--- a/Engine/Plugins/Experimental/Paper2D/Source/PaperJsonImporter/Private/PaperJsonSpriteSheetImporter.h
+++ b/Engine/Plugins/Experimental/Paper2D/Source/PaperJsonImporter/Private/PaperJsonSpriteSheetImporter.h
@@ -5,6 +5,8 @@
 #include "PaperJsonImporterPrivatePCH.h"
 #include "Paper2DClasses.h"
 
+class UPaperSpriteSheet;
+
 //////////////////////////////////////////////////////////////////////////
 // FSpriteFrame
 
diff --git a/Engine/Plugins/Experimental/Paper2D/Source/PaperJsonImporter/Private/PaperSpriteSheetAssetTypeActions.cpp b/Engine/Plugins/Experimental/Paper2D/Source/PaperJsonImporter/Private/PaperSpriteSheetAssetTypeActions.cpp
--- a/Engine/Plugins/Experimental/Paper2D/Source/PaperJsonImporter/Private/PaperSpriteSheetAssetTypeActions.cpp
+++ b/Engine/Plugins/Experimental/Paper2D/Source/PaperJsonImporter/Private/PaperSpriteSheetAssetTypeActions.cpp
@@ -8,6 +8,7 @@
 #include "ContentBrowserModule.h"
 #include "AssetRegistryModule.h"
 #include "PaperFlipbookFactory.h"
+#include "PackageTools.h"
 
 #define LOCTEXT_NAMESPACE "AssetTypeActions"
 
diff --git a/Engine/Plugins/Experimental/Paper2D/Source/PaperJsonImporter/Private/ReimportPaperJsonImporterFactory.cpp b/Engine/Plugins/Experimental/Paper2D/Source/PaperJsonImporter/Private/ReimportPaperJsonImporterFactory.cpp
--- a/Engine/Plugins/Experimental/Paper2D/Source/PaperJsonImporter/Private/ReimportPaperJsonImporterFactory.cpp
+++ b/Engine/Plugins/Experimental/Paper2D/Source/PaperJsonImporter/Private/ReimportPaperJsonImporterFactory.cpp
@@ -3,6 +3,7 @@
 #pragma once
 #include "PaperJsonImporterPrivatePCH.h"
 #include "ReimportPaperJsonImporterFactory.h"
+#include "PaperSpriteSheet.h"
 
 #define LOCTEXT_NAMESPACE "PaperJsonImporter"
 
