clang build fix, it detects #pragma once in .cpp files and raises a '#pragma once in main file' error

diff --git a/Engine/Plugins/Experimental/Paper2D/Source/PaperJsonImporter/Private/ReimportPaperJsonImporterFactory.cpp b/Engine/Plugins/Experimental/Paper2D/Source/PaperJsonImporter/Private/ReimportPaperJsonImporterFactory.cpp
--- a/Engine/Plugins/Experimental/Paper2D/Source/PaperJsonImporter/Private/ReimportPaperJsonImporterFactory.cpp
+++ b/Engine/Plugins/Experimental/Paper2D/Source/PaperJsonImporter/Private/ReimportPaperJsonImporterFactory.cpp
@@ -1,6 +1,5 @@
 // Copyright 1998-2014 Epic Games, Inc. All Rights Reserved.
 
-#pragma once
 #include "PaperJsonImporterPrivatePCH.h"
 #include "ReimportPaperJsonImporterFactory.h"
 #include "PaperSpriteSheet.h"
