Paper2D: Support for sprites with holes (disabled until UI is improved a bit)

diff --git a/Engine/Plugins/Experimental/Paper2D/Source/Paper2D/Classes/SpriteEditorOnlyTypes.h b/Engine/Plugins/Experimental/Paper2D/Source/Paper2D/Classes/SpriteEditorOnlyTypes.h
--- a/Engine/Plugins/Experimental/Paper2D/Source/Paper2D/Classes/SpriteEditorOnlyTypes.h
+++ b/Engine/Plugins/Experimental/Paper2D/Source/Paper2D/Classes/SpriteEditorOnlyTypes.h
@@ -5,11 +5,17 @@
 #include "SpriteEditorOnlyTypes.generated.h"
 
 // A single polygon, may be convex or concave, etc...
+// bNegativeWinding tells us if the winding should be negative (CW) regardless of the order in Vertices
 USTRUCT()
 struct FSpritePolygon
 {
 	GENERATED_USTRUCT_BODY()
 
+	FSpritePolygon() 
+	: bNegativeWinding(false)
+	{
+	}
+
 	UPROPERTY(Category=Physics, EditAnywhere)
 	TArray<FVector2D> Vertices;
 
@@ -18,6 +24,9 @@ struct FSpritePolygon
 	
 	UPROPERTY()//Category=Physics, EditAnywhere)
 	FVector2D BoxPosition;
+
+	UPROPERTY()
+	bool bNegativeWinding;
 };
 
 
diff --git a/Engine/Plugins/Experimental/Paper2D/Source/Paper2D/Private/PaperSprite.cpp b/Engine/Plugins/Experimental/Paper2D/Source/Paper2D/Private/PaperSprite.cpp
--- a/Engine/Plugins/Experimental/Paper2D/Source/Paper2D/Private/PaperSprite.cpp
+++ b/Engine/Plugins/Experimental/Paper2D/Source/Paper2D/Private/PaperSprite.cpp
@@ -10,6 +10,7 @@
 #include "PhysicsEngine/BodySetup2D.h"
 #include "PaperSpriteAtlas.h"
 #include "GeomTools.h"
+#include "PaperGeomTools.h"
 #include "BitmapUtils.h"
 #include "ComponentReregisterContext.h"
 
@@ -963,48 +964,51 @@ void UPaperSprite::CreatePolygonFromBoundingBox(FSpritePolygonCollection& GeomOw
 void UPaperSprite::Triangulate(const FSpritePolygonCollection& Source, TArray<FVector2D>& Target)
 {
 	Target.Empty();
+	TArray<FVector2D> AllGeneratedTriangles;
 
-	TArray<FClipSMTriangle> AllGeneratedTriangles;
-
-	for (int32 PolygonIndex = 0; PolygonIndex < Source.Polygons.Num(); ++PolygonIndex)
+	// Correct polygon winding for additive and subtractive polygons
+	// Invalid polygons (< 3 verts) removed from this list
+	TArray<FSpritePolygon> ValidPolygons = PaperGeomTools::CorrectPolygonWinding(Source.Polygons);
+	
+	// Check if polygons overlap, or have inconsistent winding, or edges overlap
+	if (!PaperGeomTools::ArePolygonsValid(ValidPolygons))
 	{
-		const FSpritePolygon& SourcePoly = Source.Polygons[PolygonIndex];
-		if (SourcePoly.Vertices.Num() >= 3)
-		{
-			// Convert our format into one the triangulation library supports
-			FClipSMPolygon ClipPolygon(0);
-			for (int32 VertexIndex = 0; VertexIndex < SourcePoly.Vertices.Num(); ++VertexIndex)
-			{
-				FClipSMVertex* ClipVertex = new (ClipPolygon.Vertices) FClipSMVertex;
-				FMemory::Memzero(ClipVertex, sizeof(FClipSMVertex));
+		return;
+	}
 
-				const FVector2D& SourceVertex = SourcePoly.Vertices[VertexIndex];
-				ClipVertex->Pos.X = SourceVertex.X;
-				ClipVertex->Pos.Z = SourceVertex.Y;
-			}
-			ClipPolygon.FaceNormal = FVector(0.0f, -1.0f, 0.0f);
+	// Merge each additive and associated subtractive polygons to form a list of polygons in CCW winding
+	ValidPolygons = PaperGeomTools::ReducePolygons(ValidPolygons);
 
-			// Attempt to triangulate this polygon
-			TArray<FClipSMTriangle> GeneratedTriangles;
-			if (TriangulatePoly(/*out*/ GeneratedTriangles, ClipPolygon, Source.bAvoidVertexMerging))
-			{
-				AllGeneratedTriangles.Append(GeneratedTriangles);
-			}
+	// Triangulate the polygons
+	for (int32 PolygonIndex = 0; PolygonIndex < ValidPolygons.Num(); ++PolygonIndex)
+	{
+		const FSpritePolygon& SourcePoly = ValidPolygons[PolygonIndex];
+		TArray<FVector2D> Generated2DTriangles;
+		if (PaperGeomTools::TriangulatePoly(Generated2DTriangles, SourcePoly.Vertices, Source.bAvoidVertexMerging))
+		{
+			AllGeneratedTriangles.Append(Generated2DTriangles);
 		}
 	}
 
-	if (!Source.bAvoidVertexMerging && (Source.Polygons.Num() > 1) && (AllGeneratedTriangles.Num() > 1))
+	bool bSourcePolygonHasHoles = false;
+	for (int32 PolygonIndex = 0; PolygonIndex < Source.Polygons.Num(); ++PolygonIndex)
 	{
-		RemoveRedundantTriangles(AllGeneratedTriangles);
+		if (Source.Polygons[PolygonIndex].bNegativeWinding)
+		{
+			bSourcePolygonHasHoles = true;
+			break;
+		}
 	}
 
-	// Convert the triangles back to our 2D data structure
-	for (const FClipSMTriangle& Triangle : AllGeneratedTriangles)
+	// This doesn't work when polys have holes as edges will likely form a loop around the poly
+	if (!bSourcePolygonHasHoles && !Source.bAvoidVertexMerging && (Source.Polygons.Num() > 1 && AllGeneratedTriangles.Num() > 1))
 	{
-		new (Target) FVector2D(Triangle.Vertices[0].Pos.X, Triangle.Vertices[0].Pos.Z);
-		new (Target) FVector2D(Triangle.Vertices[1].Pos.X, Triangle.Vertices[1].Pos.Z);
-		new (Target) FVector2D(Triangle.Vertices[2].Pos.X, Triangle.Vertices[2].Pos.Z);
+		TArray<FVector2D> TrianglesCopy = AllGeneratedTriangles;
+		AllGeneratedTriangles.Empty();
+		PaperGeomTools::RemoveRedundantTriangles(/*out*/AllGeneratedTriangles, TrianglesCopy);
 	}
+
+	Target.Append(AllGeneratedTriangles);
 }
 
 void UPaperSprite::ExtractRectsFromTexture(UTexture2D* Texture, TArray<FIntRect>& OutRects)
diff --git a/Engine/Plugins/Experimental/Paper2D/Source/Paper2DEditor/Private/PaperStyle.cpp b/Engine/Plugins/Experimental/Paper2D/Source/Paper2DEditor/Private/PaperStyle.cpp
--- a/Engine/Plugins/Experimental/Paper2D/Source/Paper2DEditor/Private/PaperStyle.cpp
+++ b/Engine/Plugins/Experimental/Paper2D/Source/Paper2DEditor/Private/PaperStyle.cpp
@@ -62,6 +62,8 @@ void FPaperStyle::Initialize()
 
 		StyleSet->Set("SpriteEditor.AddPolygon", new IMAGE_PLUGIN_BRUSH(TEXT("Icons/icon_Paper2D_AddPolygon_40x"), Icon40x40));
 		StyleSet->Set("SpriteEditor.AddPolygon.Small", new IMAGE_PLUGIN_BRUSH(TEXT("Icons/icon_Paper2D_AddPolygon_40x"), Icon20x20));
+		StyleSet->Set("SpriteEditor.AddSubtractivePolygon", new IMAGE_PLUGIN_BRUSH(TEXT("Icons/icon_Paper2D_AddSubtractivePolygon_40x"), Icon40x40));
+		StyleSet->Set("SpriteEditor.AddSubtractivePolygon.Small", new IMAGE_PLUGIN_BRUSH(TEXT("Icons/icon_Paper2D_AddSubtractivePolygon_40x"), Icon20x20));
 
 		StyleSet->Set("SpriteEditor.SnapAllVertices", new IMAGE_PLUGIN_BRUSH(TEXT("Icons/icon_Paper2D_SnapToPixelGrid_40x"), Icon40x40));
 		StyleSet->Set("SpriteEditor.SnapAllVertices.Small", new IMAGE_PLUGIN_BRUSH(TEXT("Icons/icon_Paper2D_SnapToPixelGrid_40x"), Icon20x20));
diff --git a/Engine/Plugins/Experimental/Paper2D/Source/Paper2DEditor/Private/SpriteEditor/SpriteEditor.cpp b/Engine/Plugins/Experimental/Paper2D/Source/Paper2DEditor/Private/SpriteEditor/SpriteEditor.cpp
--- a/Engine/Plugins/Experimental/Paper2D/Source/Paper2DEditor/Private/SpriteEditor/SpriteEditor.cpp
+++ b/Engine/Plugins/Experimental/Paper2D/Source/Paper2DEditor/Private/SpriteEditor/SpriteEditor.cpp
@@ -186,8 +186,12 @@ void SSpriteEditorViewport::BindCommands()
 		FCanExecuteAction::CreateSP( EditorViewportClientRef, &FSpriteEditorViewportClient::CanSplitEdge ));
 	CommandList->MapAction(
 		Commands.AddPolygon,
-		FExecuteAction::CreateSP( EditorViewportClientRef, &FSpriteEditorViewportClient::AddPolygon ),
-		FCanExecuteAction::CreateSP( EditorViewportClientRef, &FSpriteEditorViewportClient::CanAddPolygon ));
+		FExecuteAction::CreateSP(EditorViewportClientRef, &FSpriteEditorViewportClient::AddPolygon),
+		FCanExecuteAction::CreateSP(EditorViewportClientRef, &FSpriteEditorViewportClient::CanAddPolygon));
+	CommandList->MapAction(
+		Commands.AddSubtractivePolygon,
+		FExecuteAction::CreateSP(EditorViewportClientRef, &FSpriteEditorViewportClient::AddSubtractivePolygon),
+		FCanExecuteAction::CreateSP(EditorViewportClientRef, &FSpriteEditorViewportClient::CanAddSubtractivePolygon));
 	CommandList->MapAction(
 		Commands.SnapAllVertices,
 		FExecuteAction::CreateSP( EditorViewportClientRef, &FSpriteEditorViewportClient::SnapAllVerticesToPixelGrid ),
@@ -514,6 +518,7 @@ void FSpriteEditor::ExtendToolbar()
 			ToolbarBuilder.BeginSection("Tools");
 			{
 				ToolbarBuilder.AddToolBarButton(FSpriteEditorCommands::Get().AddPolygon);
+//				ToolbarBuilder.AddToolBarButton(FSpriteEditorCommands::Get().AddSubtractivePolygon); // disabled until polygon editing is improved
 				ToolbarBuilder.AddToolBarButton(FSpriteEditorCommands::Get().SnapAllVertices);
 			}
 			ToolbarBuilder.EndSection();
diff --git a/Engine/Plugins/Experimental/Paper2D/Source/Paper2DEditor/Private/SpriteEditor/SpriteEditorCommands.cpp b/Engine/Plugins/Experimental/Paper2D/Source/Paper2DEditor/Private/SpriteEditor/SpriteEditorCommands.cpp
--- a/Engine/Plugins/Experimental/Paper2D/Source/Paper2DEditor/Private/SpriteEditor/SpriteEditorCommands.cpp
+++ b/Engine/Plugins/Experimental/Paper2D/Source/Paper2DEditor/Private/SpriteEditor/SpriteEditorCommands.cpp
@@ -34,5 +34,6 @@ void FSpriteEditorCommands::RegisterCommands()
 	UI_COMMAND(DeleteSelection, "Delete", "Delete the selection.", EUserInterfaceActionType::Button, FInputGesture(EKeys::Platform_Delete, EModifierKey::None));
 	UI_COMMAND(SplitEdge, "Split", "Split edge.", EUserInterfaceActionType::Button, FInputGesture(EKeys::Insert, EModifierKey::None));
 	UI_COMMAND(AddPolygon, "Add Polygon", "Adds a new polygon.", EUserInterfaceActionType::Button, FInputGesture());
+	UI_COMMAND(AddSubtractivePolygon, "Add Subtractive Poly", "Adds a subtractive polygon.", EUserInterfaceActionType::Button, FInputGesture());
 	UI_COMMAND(SnapAllVertices, "Snap to pixel grid", "Snaps all vertices to the pixel grid.", EUserInterfaceActionType::Button, FInputGesture());
 }
diff --git a/Engine/Plugins/Experimental/Paper2D/Source/Paper2DEditor/Private/SpriteEditor/SpriteEditorCommands.h b/Engine/Plugins/Experimental/Paper2D/Source/Paper2DEditor/Private/SpriteEditor/SpriteEditorCommands.h
--- a/Engine/Plugins/Experimental/Paper2D/Source/Paper2DEditor/Private/SpriteEditor/SpriteEditorCommands.h
+++ b/Engine/Plugins/Experimental/Paper2D/Source/Paper2DEditor/Private/SpriteEditor/SpriteEditorCommands.h
@@ -48,5 +48,6 @@ public:
 	TSharedPtr<FUICommandInfo> DeleteSelection;
 	TSharedPtr<FUICommandInfo> SplitEdge;
 	TSharedPtr<FUICommandInfo> AddPolygon;
+	TSharedPtr<FUICommandInfo> AddSubtractivePolygon;
 	TSharedPtr<FUICommandInfo> SnapAllVertices;
 };
\ No newline at end of file
diff --git a/Engine/Plugins/Experimental/Paper2D/Source/Paper2DEditor/Private/SpriteEditor/SpriteEditorViewportClient.cpp b/Engine/Plugins/Experimental/Paper2D/Source/Paper2DEditor/Private/SpriteEditor/SpriteEditorViewportClient.cpp
--- a/Engine/Plugins/Experimental/Paper2D/Source/Paper2DEditor/Private/SpriteEditor/SpriteEditorViewportClient.cpp
+++ b/Engine/Plugins/Experimental/Paper2D/Source/Paper2DEditor/Private/SpriteEditor/SpriteEditorViewportClient.cpp
@@ -291,7 +291,7 @@ void FSpriteEditorViewportClient::DrawMarquee(FViewport& InViewport, FSceneView&
 	}
 }
 
-void FSpriteEditorViewportClient::DrawGeometry(FViewport& InViewport, FSceneView& View, FCanvas& Canvas, const FSpritePolygonCollection& Geometry, const FLinearColor& GeometryVertexColor, bool bIsRenderGeometry)
+void FSpriteEditorViewportClient::DrawGeometry(FViewport& InViewport, FSceneView& View, FCanvas& Canvas, const FSpritePolygonCollection& Geometry, const FLinearColor& GeometryVertexColor, const FLinearColor& NegativeGeometryVertexColor, bool bIsRenderGeometry)
 {
 	const bool bIsHitTesting = Canvas.IsHitTesting();
 	UPaperSprite* Sprite = GetSpriteBeingEdited();
@@ -299,13 +299,14 @@ void FSpriteEditorViewportClient::DrawGeometry(FViewport& InViewport, FSceneView
 	const float CollisionVertexSize = 8.0f;
 
 	const float NormalLength = 15.0f;
-	const FLinearColor GeometryLineColor(GeometryVertexColor.R, GeometryVertexColor.G, GeometryVertexColor.B, 0.5f * GeometryVertexColor.A);
 	const FLinearColor GeometryNormalColor(0.0f, 1.0f, 0.0f, 0.5f);
 
 	// Run thru the custom collision vertices and draw hit proxies for them
 	for (int32 PolygonIndex = 0; PolygonIndex < Geometry.Polygons.Num(); ++PolygonIndex)
 	{
 		const FSpritePolygon& Polygon = Geometry.Polygons[PolygonIndex];
+		const FLinearColor LineColor = Polygon.bNegativeWinding ? NegativeGeometryVertexColor : GeometryVertexColor;
+		const FLinearColor VertexColor = Polygon.bNegativeWinding ? NegativeGeometryVertexColor : GeometryVertexColor;
 
 		// Draw lines connecting the vertices of the polygon
 		for (int32 VertexIndex = 0; VertexIndex < Polygon.Vertices.Num(); ++VertexIndex)
@@ -341,7 +342,7 @@ void FSpriteEditorViewportClient::DrawGeometry(FViewport& InViewport, FSceneView
 				}
 
 				FCanvasLineItem LineItem(ScreenPos, NextScreenPos);
-				LineItem.SetColor(GeometryLineColor);
+				LineItem.SetColor(LineColor);
 				Canvas.DrawItem(LineItem);
 
 				if (bIsHitTesting)
@@ -368,7 +369,7 @@ void FSpriteEditorViewportClient::DrawGeometry(FViewport& InViewport, FSceneView
 				Canvas.SetHitProxy(new HSpriteSelectableObjectHitProxy(Data));
 			}
 
-			Canvas.DrawTile(ScreenPos.X - CollisionVertexSize*0.5f, ScreenPos.Y - CollisionVertexSize*0.5f, CollisionVertexSize, CollisionVertexSize, 0.f, 0.f, 1.f, 1.f, GeometryVertexColor, GWhiteTexture);
+			Canvas.DrawTile(ScreenPos.X - CollisionVertexSize*0.5f, ScreenPos.Y - CollisionVertexSize*0.5f, CollisionVertexSize, CollisionVertexSize, 0.f, 0.f, 1.f, 1.f, VertexColor, GWhiteTexture);
 			
 			if (bIsHitTesting)
 			{
@@ -649,7 +650,8 @@ void FSpriteEditorViewportClient::DrawCanvas(FViewport& Viewport, FSceneView& Vi
 
 			// Draw the custom collision geometry
 			const FLinearColor CollisionColor(1.0f, 1.0f, 0.0f, 1.0f);
-			DrawGeometry(Viewport, View, Canvas, Sprite->CollisionGeometry, CollisionColor, false);
+			const FLinearColor White(1, 1, 1, 1);
+			DrawGeometry(Viewport, View, Canvas, Sprite->CollisionGeometry, CollisionColor, White, false);
 			if (Sprite->BodySetup != nullptr)
 			{
 				DrawGeometryStats(Viewport, View, Canvas, Sprite->CollisionGeometry, false, /*inout*/ YPos);
@@ -668,7 +670,8 @@ void FSpriteEditorViewportClient::DrawCanvas(FViewport& Viewport, FSceneView& Vi
 
 			// Draw the custom render geometry
 			const FLinearColor RenderGeomColor(1.0f, 0.2f, 0.0f, 1.0f);
-			DrawGeometry(Viewport, View, Canvas, Sprite->RenderGeometry, RenderGeomColor, true);
+			const FLinearColor NegativeRenderGeomColor(0.0f, 0.2f, 1.0f, 1.0f);
+			DrawGeometry(Viewport, View, Canvas, Sprite->RenderGeometry, RenderGeomColor, NegativeRenderGeomColor, true);
 			DrawGeometryStats(Viewport, View, Canvas, Sprite->RenderGeometry, true, /*inout*/ YPos);
 			DrawRenderStats(Viewport, View, Canvas, Sprite, /*inout*/ YPos);
 
@@ -1244,6 +1247,33 @@ void FSpriteEditorViewportClient::AddPolygon()
 	if (FSpritePolygonCollection* Geometry = GetGeometryBeingEdited())
 	{
 		FSpritePolygon& NewPoly = *new (Geometry->Polygons) FSpritePolygon;
+		NewPoly.bNegativeWinding = false;
+
+		//@TODO: Should make this more awesome (or even just regular awesome)
+		// Hack adding verts since we don't have 'line drawing mode' where you can just click-click-click yet
+		UPaperSprite* Sprite = GetSpriteBeingEdited();
+		const FVector2D BaseUV = Sprite->GetSourceUV();
+		const float Size = 10.0f;
+
+		new (NewPoly.Vertices) FVector2D(BaseUV.X, BaseUV.Y);
+		new (NewPoly.Vertices) FVector2D(BaseUV.X + Size, BaseUV.Y + Size);
+		new (NewPoly.Vertices) FVector2D(BaseUV.X, BaseUV.Y + Size);
+		bManipulationDirtiedSomething = true;
+
+		Geometry->GeometryType = ESpritePolygonMode::FullyCustom;
+	}
+
+	EndTransaction();
+}
+
+void FSpriteEditorViewportClient::AddSubtractivePolygon()
+{
+	BeginTransaction(LOCTEXT("AddPolygonTransaction", "Add Subtractive Polygon"));
+
+	if (FSpritePolygonCollection* Geometry = GetGeometryBeingEdited())
+	{
+		FSpritePolygon& NewPoly = *new (Geometry->Polygons) FSpritePolygon;
+		NewPoly.bNegativeWinding = true;
 
 		//@TODO: Should make this more awesome (or even just regular awesome)
 		// Hack adding verts since we don't have 'line drawing mode' where you can just click-click-click yet
diff --git a/Engine/Plugins/Experimental/Paper2D/Source/Paper2DEditor/Private/SpriteEditor/SpriteEditorViewportClient.h b/Engine/Plugins/Experimental/Paper2D/Source/Paper2DEditor/Private/SpriteEditor/SpriteEditorViewportClient.h
--- a/Engine/Plugins/Experimental/Paper2D/Source/Paper2DEditor/Private/SpriteEditor/SpriteEditorViewportClient.h
+++ b/Engine/Plugins/Experimental/Paper2D/Source/Paper2DEditor/Private/SpriteEditor/SpriteEditorViewportClient.h
@@ -106,7 +106,9 @@ public:
 	bool CanSplitEdge() const { return IsEditingGeometry(); } //@TODO: Need an edge
 
 	void AddPolygon();
+	void AddSubtractivePolygon();
 	bool CanAddPolygon() const { return IsEditingGeometry(); }
+	bool CanAddSubtractivePolygon() const { return CanAddPolygon() && IsInRenderingEditMode(); }
 
 	void SnapAllVerticesToPixelGrid();
 	bool CanSnapVerticesToPixelGrid() const { return IsEditingGeometry(); }
@@ -196,7 +198,7 @@ private:
 
 	void DrawTriangleList(FViewport& InViewport, FSceneView& View, FCanvas& Canvas, const TArray<FVector2D>& Triangles);
 	void DrawBoundsAsText(FViewport& InViewport, FSceneView& View, FCanvas& Canvas, int32& YPos);
-	void DrawGeometry(FViewport& InViewport, FSceneView& View, FCanvas& Canvas, const FSpritePolygonCollection& Geometry, const FLinearColor& GeometryVertexColor, bool bIsRenderGeometry);
+	void DrawGeometry(FViewport& InViewport, FSceneView& View, FCanvas& Canvas, const FSpritePolygonCollection& Geometry, const FLinearColor& GeometryVertexColor, const FLinearColor& NegativeGeometryVertexColor, bool bIsRenderGeometry);
 	void DrawGeometryStats(FViewport& InViewport, FSceneView& View, FCanvas& Canvas, const FSpritePolygonCollection& Geometry, bool bIsRenderGeometry, int32& YPos);
 	void DrawRenderStats(FViewport& InViewport, FSceneView& View, FCanvas& Canvas, class UPaperSprite* Sprite, int32& YPos);
 	void DrawSockets(const FSceneView* View, FPrimitiveDrawInterface* PDI);
