Paper2D: Added missing files (PaperGeomTools.h/cpp)

diff --git a/Engine/Plugins/Experimental/Paper2D/Source/Paper2D/Private/PaperGeomTools.cpp b/Engine/Plugins/Experimental/Paper2D/Source/Paper2D/Private/PaperGeomTools.cpp
new file mode 100644
--- /dev/null
+++ b/Engine/Plugins/Experimental/Paper2D/Source/Paper2D/Private/PaperGeomTools.cpp
@@ -0,0 +1,561 @@
+// Copyright 1998-2014 Epic Games, Inc. All Rights Reserved.
+
+#if WITH_EDITOR
+
+#include "Paper2DPrivatePCH.h"
+#include "PaperGeomTools.h"
+#include "GeomTools.h"
+
+static bool IsPolygonWindingCCW(const TArray<FVector2D>& Points)
+{
+	float Sum = 0.0f;
+	const int PointCount = Points.Num();
+	for (int PointIndex = 0; PointIndex < PointCount; ++PointIndex)
+	{
+		const FVector2D& A = Points[PointIndex];
+		const FVector2D& B = Points[(PointIndex + 1) % PointCount];
+		Sum += (B.X - A.X) * (B.Y + A.Y);
+	}
+	return (Sum < 0.0f);
+}
+
+static bool IsPointInPolygon(const FVector2D& TestPoint, const TArray<FVector2D>& PolygonPoints)
+{
+	const int NumPoints = PolygonPoints.Num();
+	float AngleSum = 0.0f;
+	for (int PointIndex = 0; PointIndex < NumPoints; ++PointIndex)
+	{
+		const FVector2D& VecAB = PolygonPoints[PointIndex] - TestPoint;
+		const FVector2D& VecAC = PolygonPoints[(PointIndex + 1) % NumPoints] - TestPoint;
+		const float Angle = FMath::Sign(FVector2D::CrossProduct(VecAB, VecAC)) * FMath::Acos(FMath::Clamp(FVector2D::DotProduct(VecAB, VecAC) / (VecAB.Size() * VecAC.Size()), -1.0f, 1.0f));
+		AngleSum += Angle;
+	}
+	return (FMath::Abs(AngleSum) > 0.001f);
+}
+
+static bool IsPointInTriangle(const FVector2D& TestPoint, const FVector2D& A, const FVector2D& B, const FVector2D& C)
+{
+	const FVector2D AP = TestPoint - A;
+	const FVector2D BP = TestPoint - B;
+	const FVector2D AB = B - A;
+	const FVector2D AC = C - A;
+	const FVector2D BC = C - B;
+	if (FVector2D::CrossProduct(AB, AP) <= 0.0f) return false;
+	if (FVector2D::CrossProduct(AP, AC) <= 0.0f) return false;
+	if (FVector2D::CrossProduct(BC, BP) <= 0.0f) return false;
+	if (AP.SizeSquared() < 2.0f) return false;
+	if (BP.SizeSquared() < 2.0f) return false;
+	const FVector2D CP = TestPoint - C;
+	if (CP.SizeSquared() < 2.0f) return false;
+	return true;
+}
+
+
+static void JoinSubtractiveToAdditive(TArray<FVector2D>& AdditivePoly, const TArray<FVector2D>& SubtractivePoly, const int AdditiveJoinIndex, const int SubtractiveJoinIndex)
+{
+	TArray<FVector2D> NewAdditivePoly;
+	for (int AdditiveIndex = 0; AdditiveIndex < AdditivePoly.Num(); ++AdditiveIndex)
+	{
+		NewAdditivePoly.Add(AdditivePoly[AdditiveIndex]);
+		if (AdditiveIndex == AdditiveJoinIndex)
+		{
+			for (int SubtractiveIndex = SubtractiveJoinIndex; SubtractiveIndex < SubtractivePoly.Num(); ++SubtractiveIndex)
+			{
+				NewAdditivePoly.Add(SubtractivePoly[SubtractiveIndex]);
+			}
+			for (int SubtractiveIndex = 0; SubtractiveIndex <= SubtractiveJoinIndex; ++SubtractiveIndex)
+			{
+				NewAdditivePoly.Add(SubtractivePoly[SubtractiveIndex]);
+			}
+			NewAdditivePoly.Add(AdditivePoly[AdditiveIndex]);
+		}
+	}
+	AdditivePoly = NewAdditivePoly;
+}
+
+
+static void JoinMutuallyVisible(TArray<FVector2D>& AdditivePoly, const TArray<FVector2D>& SubtractivePoly)
+{
+	const int NumAdditivePoly = AdditivePoly.Num();
+	const int NumSubtractivePoly = SubtractivePoly.Num();
+	if (NumSubtractivePoly == 0)
+	{
+		return;
+	}
+
+	// Search the inner (subtractive) polygon for the point of maximum x-value
+	int IndexMaxX = 0;
+	{
+		float MaxX = SubtractivePoly[0].X;
+		for (int Index = 1; Index < NumSubtractivePoly; ++Index)
+		{
+			if (SubtractivePoly[Index].X > MaxX)
+			{
+				MaxX = SubtractivePoly[Index].X;
+				IndexMaxX = Index;
+			}
+		}
+	}
+	const FVector2D PointMaxX = SubtractivePoly[IndexMaxX];
+
+	// Intersect a ray from point M facing to the right (a, ab) with the additive shape edges (c, d)
+	// Find the closest intersecting point to M (left-most point)
+	int EdgeStartPointIndex = 0;
+	int EdgeEndPointIndex = 0;
+	bool bIntersectedAtVertex = false;
+	float LeftMostIntersectX = MAX_FLT;
+	const FVector2D A = PointMaxX;
+	const FVector2D AB = FVector2D(1.0f, 0.0f);
+	for (int AdditiveIndex = 0; AdditiveIndex < NumAdditivePoly; ++AdditiveIndex)
+	{
+		const FVector2D& C = AdditivePoly[AdditiveIndex];
+		const FVector2D& D = AdditivePoly[(AdditiveIndex + 1) % NumAdditivePoly];
+		const FVector2D CD = D - C;
+		// Only check edges from the inside, as edges will overlap after mutually visible points are merged.
+		if (CD.Y > 0.0f)
+		{
+			const float DetS = AB.X * CD.Y - AB.Y * CD.X;
+			const float DetT = CD.X * AB.Y - CD.Y * AB.X;
+			if (DetS != 0.0f && DetT != 0.0f)
+			{
+				const float S = (A.Y * CD.X - C.Y * CD.X - A.X * CD.Y + C.X * CD.Y) / DetS;
+				const float T = (C.Y * AB.X - A.Y * AB.X - C.X * AB.Y + A.X * AB.Y) / DetT;
+				if (S >= 0.0f && T >= 0.0f && T <= 1.0f)
+				{
+					const float IntersectX = A.X + AB.X * S;
+					if (IntersectX < LeftMostIntersectX)
+					{
+						LeftMostIntersectX = IntersectX;
+						EdgeStartPointIndex = AdditiveIndex;
+						EdgeEndPointIndex = (AdditiveIndex + 1) % NumAdditivePoly;
+						if (T < FLT_EPSILON)
+						{
+							bIntersectedAtVertex = true;
+						}
+						else if (T > 1.0f - FLT_EPSILON)
+						{
+							bIntersectedAtVertex = true;
+							EdgeStartPointIndex = EdgeEndPointIndex;
+						}
+					}
+				}
+			}
+		}
+	}
+
+	// If the ray intersected a vertex, points are mutually visible
+	if (bIntersectedAtVertex)
+	{
+		JoinSubtractiveToAdditive(AdditivePoly, SubtractivePoly, EdgeStartPointIndex, IndexMaxX);
+		return;
+	}
+
+	// Otherwise, set P to be the edge endpoint with maximum x value
+	const FVector2D Intersect(LeftMostIntersectX, PointMaxX.Y);
+	const int IndexP = (AdditivePoly[EdgeStartPointIndex].X > AdditivePoly[EdgeEndPointIndex].X) ? EdgeStartPointIndex : EdgeEndPointIndex;
+	const FVector2D P = AdditivePoly[IndexP];
+
+	// Search the vertices of the additive shape. If all of these are outside the triangle (M, intersect, P) then M and P are mutually visible
+	// If any vertices lie inside the triangle, find the one that minimizes the angle between (1, 0) and the line M R. (If multiple vertices minimize the angle, choose closest to M)
+	const FVector2D TriA = PointMaxX;
+	const FVector2D TriB = (P.Y < Intersect.Y) ? P : Intersect;
+	const FVector2D TriC = (P.Y < Intersect.Y) ? Intersect : P;
+	float CosAngleMax = 0.0f;
+	float DistanceMin = MAX_FLT;
+	int IndexR = -1;
+	for (int AdditiveIndex = 0; AdditiveIndex < NumAdditivePoly; ++AdditiveIndex)
+	{
+		// Ignore point P
+		if (AdditiveIndex == IndexP)
+		{
+			continue;
+		}
+
+		if (IsPointInTriangle(AdditivePoly[AdditiveIndex], TriA, TriB, TriC))
+		{
+			const FVector2D MR = AdditivePoly[AdditiveIndex] - PointMaxX;
+			const float CosAngle = MR.X / MR.Size();
+			const float DSq = MR.SizeSquared();
+			if (CosAngle > CosAngleMax || (CosAngle == CosAngleMax && DSq < DistanceMin))
+			{
+				CosAngleMax = CosAngle;
+				DistanceMin = DSq;
+				IndexR = AdditiveIndex;
+			}
+		}
+	}
+
+	JoinSubtractiveToAdditive(AdditivePoly, SubtractivePoly, (IndexR == -1) ? IndexP : IndexR, IndexMaxX);
+}
+
+TArray<FSpritePolygon> PaperGeomTools::ReducePolygons(const TArray<FSpritePolygon>& Polygons)
+{
+	TArray<FSpritePolygon> ReturnPolygons;
+
+	int NumPolygons = Polygons.Num();
+	
+	TArray<float> MaxXValues; // per polygon
+	for (int PolyIndex = 0; PolyIndex < NumPolygons; ++PolyIndex)
+	{
+		float MaxX = -BIG_NUMBER;
+		const TArray<FVector2D> &Vertices = Polygons[PolyIndex].Vertices;
+		for (int VertexIndex = 0; VertexIndex < Vertices.Num(); ++VertexIndex)
+		{
+			MaxX = FMath::Max(Vertices[VertexIndex].X, MaxX);
+		}
+		MaxXValues.Add(MaxX);
+	}
+
+	// Iterate through additive shapes
+	for (int PolyIndex = 0; PolyIndex < NumPolygons; ++PolyIndex)
+	{
+		TArray<FVector2D> Verts = Polygons[PolyIndex].Vertices;
+		if (!Polygons[PolyIndex].bNegativeWinding)
+		{
+			// Store indexes of subtractive shapes that lie inside the additive shape
+			TArray<int> SubtractiveShapeIndices;
+			for (int J = 0; J < NumPolygons; ++J)
+			{
+				const TArray<FVector2D>& CandidateShape = Polygons[J].Vertices;
+				if (Polygons[J].bNegativeWinding)
+				{
+					if (IsPointInPolygon(CandidateShape[0], Verts))
+					{
+						SubtractiveShapeIndices.Add(J);
+					}
+				}
+			}
+
+			// Remove subtractive shapes that lie inside the subtractive shapes we've found
+			for (int J = 0; J < SubtractiveShapeIndices.Num();)
+			{
+				const TArray<FVector2D>& ourShape = Polygons[SubtractiveShapeIndices[J]].Vertices;
+				bool bRemoveOurShape = false;
+				for (int K = 0; K < SubtractiveShapeIndices.Num(); ++K)
+				{
+					if (J == K)
+					{
+						continue;
+					}
+					if (IsPointInPolygon(ourShape[0], Polygons[SubtractiveShapeIndices[K]].Vertices))
+					{
+						bRemoveOurShape = true;
+						break;
+					}
+				}
+				if (bRemoveOurShape)
+				{
+					SubtractiveShapeIndices.RemoveAt(J);
+				}
+				else
+				{
+					++J;
+				}
+			}
+
+			// Sort subtractive shapes from right to left by their points' maximum x value
+			const int NumSubtractiveShapes = SubtractiveShapeIndices.Num();
+			for (int J = 0; J < NumSubtractiveShapes; ++J)
+			{
+				for (int K = J + 1; K < NumSubtractiveShapes; ++K)
+				{
+					if (MaxXValues[SubtractiveShapeIndices[J]] < MaxXValues[SubtractiveShapeIndices[K]])
+					{
+						int tempIndex = SubtractiveShapeIndices[J];
+						SubtractiveShapeIndices[J] = SubtractiveShapeIndices[K];
+						SubtractiveShapeIndices[K] = tempIndex;
+					}
+				}
+			}
+
+			for (int SubtractiveIndex = 0; SubtractiveIndex < NumSubtractiveShapes; ++SubtractiveIndex)
+			{
+				JoinMutuallyVisible(Verts, Polygons[SubtractiveShapeIndices[SubtractiveIndex]].Vertices);
+			}
+
+			// Add new hole-less polygon to our output shapes
+			FSpritePolygon *NewPolygon = new(ReturnPolygons)FSpritePolygon();
+			NewPolygon->bNegativeWinding = false;
+			NewPolygon->Vertices = Verts;
+		}
+	}
+
+	return ReturnPolygons;
+}
+
+
+TArray<FSpritePolygon> PaperGeomTools::CorrectPolygonWinding(const TArray<FSpritePolygon>& Polygons)
+{
+	TArray<FSpritePolygon> ReturnPolygons;
+	for (int32 PolygonIndex = 0; PolygonIndex < Polygons.Num(); ++PolygonIndex)
+	{
+		const FSpritePolygon& SourcePolygon = Polygons[PolygonIndex];
+		if (SourcePolygon.Vertices.Num() >= 3)
+		{
+			FSpritePolygon* FixedPolygon = new (ReturnPolygons)FSpritePolygon();
+			// Make sure the polygon winding is correct
+			if ((!SourcePolygon.bNegativeWinding && !IsPolygonWindingCCW(SourcePolygon.Vertices)) ||
+				(SourcePolygon.bNegativeWinding && IsPolygonWindingCCW(SourcePolygon.Vertices)))
+			{
+				// Reverse vertices
+				for (int32 VertexIndex = SourcePolygon.Vertices.Num() - 1; VertexIndex >= 0; --VertexIndex)
+				{
+					new (FixedPolygon->Vertices)FVector2D(SourcePolygon.Vertices[VertexIndex]);
+				}
+			}
+			else
+			{
+				// Copy vertices
+				for (int32 VertexIndex = 0; VertexIndex < SourcePolygon.Vertices.Num(); ++VertexIndex)
+				{
+					new (FixedPolygon->Vertices)FVector2D(SourcePolygon.Vertices[VertexIndex]);
+				}
+			}
+			FixedPolygon->bNegativeWinding = SourcePolygon.bNegativeWinding;
+		}
+	}
+	return ReturnPolygons;
+}
+
+// Assumes polygons are valid, closed and the winding is correct and matches bWindingCCW
+bool PaperGeomTools::ArePolygonsValid(const TArray<FSpritePolygon>& Polygons)
+{
+	const int PolygonCount = Polygons.Num();
+
+	// Find every pair of shapes that overlap
+	for (int PolygonIndexA = 0; PolygonIndexA < PolygonCount; ++PolygonIndexA)
+	{
+		const TArray<FVector2D>& PolygonA = Polygons[PolygonIndexA].Vertices;
+		const bool bIsWindingA_CCW = IsPolygonWindingCCW(PolygonA);
+		for (int PolygonIndexB = PolygonIndexA + 1; PolygonIndexB < PolygonCount; ++PolygonIndexB)
+		{
+			const TArray<FVector2D>& PolygonB = Polygons[PolygonIndexB].Vertices;
+			const bool bIsWindingB_CCW = IsPolygonWindingCCW(PolygonB);
+
+			// Additive polygons are allowed to intersect
+			if (bIsWindingA_CCW && bIsWindingB_CCW)
+			{
+				continue;
+			}
+
+			// Test each shapeA edge against each shapeB edge for intersection
+			for (int VertexA = 0; VertexA < PolygonA.Num(); ++VertexA)
+			{
+				const FVector2D& A0 = PolygonA[VertexA];
+				const FVector2D& A1 = PolygonA[(VertexA + 1) % PolygonA.Num()];
+				const FVector2D A10 = A1 - A0;
+
+				for (int VertexB = 0; VertexB < PolygonB.Num(); ++VertexB)
+				{
+					const FVector2D B0 = PolygonB[VertexB];
+					const FVector2D B1 = PolygonB[(VertexB + 1) % PolygonB.Num()];
+					const FVector2D B10 = B1 - B0;
+
+					const float DetS = A10.X * B10.Y - A10.Y * B10.X;
+					const float DetT = B10.X * A10.Y - B10.Y * A10.X;
+					if (DetS != 0.0f && DetT != 0.0f)
+					{
+						const float S = (A0.Y * B10.X - B0.Y * B10.X - A0.X * B10.Y + B0.X * B10.Y) / DetS;
+						const float T = (B0.Y * A10.X - A0.Y * A10.X - B0.X * A10.Y + A0.X * A10.Y) / DetT;
+						if (S >= 0.0f && S <= 1.0f && T >= 0.0f && T <= 1.0f)
+						{
+							// Edges intersect
+							UE_LOG(LogPaper2D, Log, TEXT("Edges in polygon %d and %d intersect"), PolygonIndexA, PolygonIndexB);
+							return false;
+						}
+					}
+				}
+			}
+		}
+	}
+
+	// Make sure no polygons are self-intersecting
+	for (int PolygonIndex = 0; PolygonIndex < PolygonCount; ++PolygonIndex)
+	{
+		const TArray<FVector2D>& Polygon = Polygons[PolygonIndex].Vertices;
+		const int PointCount = Polygon.Num();
+		for (int PointIndexA = 0; PointIndexA < PointCount; ++PointIndexA)
+		{
+			const FVector2D& A0 = Polygon[PointIndexA];
+			const FVector2D& A1 = Polygon[(PointIndexA + 1) % PointCount];
+			const FVector2D A10 = A1 - A0;
+			for (int PointIndexB = 0; PointIndexB < PointCount - 3; ++PointIndexB)
+			{
+				const FVector2D& B0 = Polygon[(PointIndexA + 2 + PointIndexB) % PointCount];
+				const FVector2D& B1 = Polygon[(PointIndexA + 3 + PointIndexB) % PointCount];
+				const FVector2D B10 = B1 - B0;
+
+				const float DetS = A10.X * B10.Y - A10.Y * B10.X;
+				const float DetT = B10.X * A10.Y - B10.Y * A10.X;
+				if (DetS != 0.0f && DetT != 0.0f)
+				{
+					const float S = (A0.Y * B10.X - B0.Y * B10.X - A0.X * B10.Y + B0.X * B10.Y) / DetS;
+					const float T = (B0.Y * A10.X - A0.Y * A10.X - B0.X * A10.Y + A0.X * A10.Y) / DetT;
+					if (S >= 0.0f && S <= 1.0f && T >= 0.0f && T <= 1.0f)
+					{
+						// Edges intersect
+						UE_LOG(LogPaper2D, Log, TEXT("Polygon %d is self intersecting"), PolygonIndex);
+						return false;
+					}
+				}
+			}
+		}
+	}
+
+	return true;
+}
+
+// Determines whether two edges may be merged.
+// Very similar to the one in GeomTools.cpp, but only considers positions (as opposed to other vertex attributes)
+static bool AreEdgesMergeable(const FVector2D& V0, const FVector2D& V1, const FVector2D& V2)
+{
+	const FVector2D MergedEdgeVector = V2 - V0;
+	const float MergedEdgeLengthSquared = MergedEdgeVector.SizeSquared();
+	if(MergedEdgeLengthSquared > DELTA)
+	{
+		// Find the vertex closest to A1/B0 that is on the hypothetical merged edge formed by A0-B1.
+		const float IntermediateVertexEdgeFraction = ((V2 - V0) | (V1 - V0)) / MergedEdgeLengthSquared;
+		const FVector2D InterpolatedVertex = V0 + (V2 - V0) * IntermediateVertexEdgeFraction;
+
+		// The edges are merge-able if the interpolated vertex is close enough to the intermediate vertex.
+		return InterpolatedVertex.Equals(V1, THRESH_POINTS_ARE_SAME);
+	}
+	else
+	{
+		return true;
+	}
+}
+
+// Nearly identical function in GeomTools.cpp - Identical points (exact same fp values, not nearlyEq) are not considered to be inside 
+// the triangle when ear clipping. These are generated by ReducePolygons when additive and subtractive polygons are merged.
+// Expected input - PolygonVertices in CCW order, not overlapping
+bool PaperGeomTools::TriangulatePoly(TArray<FVector2D>& OutTris, const TArray<FVector2D>& InPolyVerts, bool bKeepColinearVertices)
+{
+	// Can't work if not enough verts for 1 triangle
+	if (InPolyVerts.Num() < 3)
+	{
+		// Return true because poly is already a tri
+		return true;
+	}
+
+	// Vertices of polygon in order - make a copy we are going to modify.
+	TArray<FVector2D> PolyVerts = InPolyVerts;
+
+	// Keep iterating while there are still vertices
+	while (true)
+	{
+		if (!bKeepColinearVertices)
+		{
+			// Cull redundant vertex edges from the polygon.
+			for (int32 VertexIndex = 0; VertexIndex < PolyVerts.Num(); VertexIndex++)
+			{
+				const int32 I0 = (VertexIndex + 0) % PolyVerts.Num();
+				const int32 I1 = (VertexIndex + 1) % PolyVerts.Num();
+				const int32 I2 = (VertexIndex + 2) % PolyVerts.Num();
+				if (AreEdgesMergeable(PolyVerts[I0], PolyVerts[I1], PolyVerts[I2]))
+				{
+					PolyVerts.RemoveAt(I1);
+					VertexIndex--;
+				}
+			}
+		}
+
+		if (PolyVerts.Num() < 3)
+		{
+			break;
+		}
+		else
+		{
+			// Look for an 'ear' triangle
+			bool bFoundEar = false;
+			for (int32 EarVertexIndex = 0; EarVertexIndex < PolyVerts.Num(); EarVertexIndex++)
+			{
+				// Triangle is 'this' vert plus the one before and after it
+				const int32 AIndex = (EarVertexIndex == 0) ? PolyVerts.Num() - 1 : EarVertexIndex - 1;
+				const int32 BIndex = EarVertexIndex;
+				const int32 CIndex = (EarVertexIndex + 1) % PolyVerts.Num();
+
+				// Check that this vertex is convex (cross product must be positive)
+				const FVector2D ABEdge = PolyVerts[BIndex] - PolyVerts[AIndex];
+				const FVector2D ACEdge = PolyVerts[CIndex] - PolyVerts[AIndex];
+				if (FMath::IsNegativeFloat( ABEdge ^ ACEdge ))
+				{
+					continue;
+				}
+
+				bool bFoundVertInside = false;
+				// Look through all verts before this in array to see if any are inside triangle
+				for (int32 VertexIndex = 0; VertexIndex < PolyVerts.Num(); VertexIndex++)
+				{
+					const FVector2D& CurrentVertex = PolyVerts[VertexIndex];
+					// Test indices first, then make sure we arent comparing identical points
+					// These might have been added by the convex / concave splitting
+					if (VertexIndex != AIndex && VertexIndex != BIndex && VertexIndex != CIndex &&
+						CurrentVertex != PolyVerts[AIndex] && CurrentVertex != PolyVerts[BIndex] && CurrentVertex != PolyVerts[CIndex] &&
+						IsPointInTriangle(CurrentVertex, PolyVerts[AIndex], PolyVerts[BIndex], PolyVerts[CIndex]))
+					{
+						bFoundVertInside = true;
+						break;
+					}
+				}
+
+				// Triangle with no verts inside - its an 'ear'! 
+				if (!bFoundVertInside)
+				{
+					// Add to output list..
+					OutTris.Add(PolyVerts[AIndex]);
+					OutTris.Add(PolyVerts[BIndex]);
+					OutTris.Add(PolyVerts[CIndex]);
+
+					// And remove vertex from polygon
+					PolyVerts.RemoveAt(EarVertexIndex);
+
+					bFoundEar = true;
+					break;
+				}
+			}
+
+			// If we couldn't find an 'ear' it indicates something is bad with this polygon - discard triangles and return.
+			if (!bFoundEar)
+			{
+				UE_LOG(LogPaper2D, Log, TEXT("Triangulation of poly failed."));
+				OutTris.Empty();
+				return false;
+			}
+		}
+	}
+
+	return true;
+}
+
+// Wrap around GeomTools RemoveRedundantTriangles
+void PaperGeomTools::RemoveRedundantTriangles(TArray<FVector2D>& OutTriangles, const TArray<FVector2D>& InTriangles)
+{
+	// Convert to ClipSMTriangles
+	TArray<FClipSMTriangle> ClipSMTriangles;
+	for (int TriangleVertex = 0; TriangleVertex < InTriangles.Num(); TriangleVertex += 3)
+	{
+		FClipSMTriangle* Tri = new(ClipSMTriangles)FClipSMTriangle(0); // Zeroed
+		for (int TriVertIndex = 0; TriVertIndex < 3; ++TriVertIndex)
+		{
+			Tri->Vertices[TriVertIndex].Pos.X = InTriangles[TriangleVertex + TriVertIndex].X;
+			Tri->Vertices[TriVertIndex].Pos.Z = InTriangles[TriangleVertex + TriVertIndex].Y;
+		}
+		Tri->FaceNormal = FVector(0.0f, -1.0f, 0.0f);
+	}
+
+	// GeomTools.cpp
+	::RemoveRedundantTriangles(ClipSMTriangles);
+
+	// Convert back to FVector2Ds
+	OutTriangles.Empty();
+	for (int TriangleIndex = 0; TriangleIndex < ClipSMTriangles.Num(); ++TriangleIndex)
+	{
+		FClipSMTriangle& Triangle = ClipSMTriangles[TriangleIndex];
+		new (OutTriangles)FVector2D(Triangle.Vertices[0].Pos.X, Triangle.Vertices[0].Pos.Z);
+		new (OutTriangles)FVector2D(Triangle.Vertices[1].Pos.X, Triangle.Vertices[1].Pos.Z);
+		new (OutTriangles)FVector2D(Triangle.Vertices[2].Pos.X, Triangle.Vertices[2].Pos.Z);
+	}
+}
+
+#endif // WITH_EDITOR
diff --git a/Engine/Plugins/Experimental/Paper2D/Source/Paper2D/Private/PaperGeomTools.h b/Engine/Plugins/Experimental/Paper2D/Source/Paper2D/Private/PaperGeomTools.h
new file mode 100644
--- /dev/null
+++ b/Engine/Plugins/Experimental/Paper2D/Source/Paper2D/Private/PaperGeomTools.h
@@ -0,0 +1,32 @@
+// Copyright 1998-2014 Epic Games, Inc. All Rights Reserved.
+
+#pragma once
+
+#include "SpriteEditorOnlyTypes.h"
+
+#if WITH_EDITOR
+
+// PaperGeomTools is editor only, used only by PaperSprite.cpp
+class PaperGeomTools
+{
+public:
+	// Returns a list of polygons with the points in the corresponding winding to the winding flag on FSpritePolygon
+	// Ie. If the vertex order doesn't match the winding, it is reversed
+	// Entries in the array that aren't polygons (< 3 points) are removed
+	static TArray<FSpritePolygon> CorrectPolygonWinding(const TArray<FSpritePolygon>& Polygons);
+
+	// Checks that these polygons can be successfully triangulated	
+	static bool ArePolygonsValid(const TArray<FSpritePolygon>& Polygons);
+
+	// Merge additive and subtractive polygons, split them up into additive polygons
+	// Assumes all polygons and overlapping polygons are valid, and the windings match the setting on the polygon
+	static TArray<FSpritePolygon> PaperGeomTools::ReducePolygons(const TArray<FSpritePolygon>& Polygons);
+
+	// Triangulate a polygon. Check notes in implementation
+	static bool TriangulatePoly(TArray<FVector2D>& OutTris, const TArray<FVector2D>& PolygonVertices, bool bKeepColinearVertices);
+
+	// Wrapper around GeomTools.cpp, RemoveRedundantTriangles
+	static void RemoveRedundantTriangles(TArray<FVector2D>& OutTriangles, const TArray<FVector2D>& InTriangles);
+};
+
+#endif
