Fix PaperGeomTools compiling on Mac

diff --git a/Engine/Plugins/Experimental/Paper2D/Source/Paper2D/Private/PaperGeomTools.h b/Engine/Plugins/Experimental/Paper2D/Source/Paper2D/Private/PaperGeomTools.h
--- a/Engine/Plugins/Experimental/Paper2D/Source/Paper2D/Private/PaperGeomTools.h
+++ b/Engine/Plugins/Experimental/Paper2D/Source/Paper2D/Private/PaperGeomTools.h
@@ -20,7 +20,7 @@ public:
 
 	// Merge additive and subtractive polygons, split them up into additive polygons
 	// Assumes all polygons and overlapping polygons are valid, and the windings match the setting on the polygon
-	static TArray<FSpritePolygon> PaperGeomTools::ReducePolygons(const TArray<FSpritePolygon>& Polygons);
+	static TArray<FSpritePolygon> ReducePolygons(const TArray<FSpritePolygon>& Polygons);
 
 	// Triangulate a polygon. Check notes in implementation
 	static bool TriangulatePoly(TArray<FVector2D>& OutTris, const TArray<FVector2D>& PolygonVertices, bool bKeepColinearVertices);
