Paper2D: Null check in PaperSpriteThumbnailRenderer, don't draw if texture is invalid.
#TTP 345436

diff --git a/Engine/Plugins/Experimental/Paper2D/Source/Paper2DEditor/Private/PaperSpriteThumbnailRenderer.cpp b/Engine/Plugins/Experimental/Paper2D/Source/Paper2DEditor/Private/PaperSpriteThumbnailRenderer.cpp
--- a/Engine/Plugins/Experimental/Paper2D/Source/Paper2DEditor/Private/PaperSpriteThumbnailRenderer.cpp
+++ b/Engine/Plugins/Experimental/Paper2D/Source/Paper2DEditor/Private/PaperSpriteThumbnailRenderer.cpp
@@ -96,7 +96,7 @@ void UPaperSpriteThumbnailRenderer::DrawFrame(class UPaperSprite* Sprite, int32
 		}
 
 		// Draw triangles
-		if (CanvasPositions.Num() > 0)
+		if (CanvasPositions.Num() > 0 && SourceTexture->Resource != nullptr)
 		{
 			TArray<FCanvasUVTri> Triangles;
 			const FLinearColor SpriteColor(FLinearColor::White);
