Paper2D: Geometry editing improvements
- Moved delete out of SpriteEditorSelections, its now possible to delete a multi selection, and edges.
- Double click on a vertex to select the entire polygon.
- Ctrl+Click,Click,Click to create vertices on the closest edge. If a vertex is already selected, the vertex will be created on that polygon, otherwise the closest edge is found.

diff --git a/Engine/Plugins/Experimental/Paper2D/Source/Paper2DEditor/Private/SpriteEditor/SpriteEditorSelections.h b/Engine/Plugins/Experimental/Paper2D/Source/Paper2DEditor/Private/SpriteEditor/SpriteEditorSelections.h
--- a/Engine/Plugins/Experimental/Paper2D/Source/Paper2DEditor/Private/SpriteEditor/SpriteEditorSelections.h
+++ b/Engine/Plugins/Experimental/Paper2D/Source/Paper2DEditor/Private/SpriteEditor/SpriteEditorSelections.h
@@ -17,6 +17,8 @@ private:
 	FSelectionTypes() {}
 };
 
+class FSpriteSelectedVertex;
+
 //////////////////////////////////////////////////////////////////////////
 // FSelectedItem
 
@@ -50,11 +52,6 @@ public:
 	{
 	}
 
-	//@TODO: Doesn't belong here in base!
-	virtual void Delete()
-	{
-	}
-
 	//@TODO: Doesn't belong here in base!
 	virtual void SplitEdge()
 	{
@@ -65,6 +62,11 @@ public:
 		return FVector::ZeroVector;
 	}
 
+	virtual const FSpriteSelectedVertex* CastSelectedVertex() const
+	{
+		return nullptr;
+	}
+
 	virtual ~FSelectedItem() {}
 };
 
@@ -239,12 +241,7 @@ public:
 		return GetWorldPosIndexed(VertexIndex);
 	}
 
-	virtual void Delete() override
-	{
-		// Cant delete anything on a bounds object
-	}
-
-		virtual void SplitEdge() override
+	virtual void SplitEdge() override
 	{
 		// Nonsense operation on a vertex, do nothing
 	}
@@ -273,6 +270,27 @@ public:
 	{
 	}
 
+	virtual bool IsValidInEditor(UPaperSprite* Sprite, bool bRenderData) const
+	{
+		if (SpritePtr == Sprite && this->bRenderData == bRenderData)
+		{
+			if (UPaperSprite* Sprite = SpritePtr.Get())
+			{
+				FSpritePolygonCollection& Geometry = bRenderData ? Sprite->RenderGeometry : Sprite->CollisionGeometry;
+				if (Geometry.Polygons.IsValidIndex(PolygonIndex) && Geometry.Polygons[PolygonIndex].Vertices.IsValidIndex(VertexIndex))
+				{
+					return true;
+				}
+			}
+		}
+		return false;
+	}
+
+	virtual const FSpriteSelectedVertex* CastSelectedVertex() const
+	{
+		return (FSpriteSelectedVertex*)this;
+	}
+
 	virtual uint32 GetTypeHash() const override
 	{
 		return VertexIndex + (PolygonIndex * 311) + (bRenderData ? 1063 : 0);
@@ -346,29 +364,6 @@ public:
 		return GetWorldPosIndexed(VertexIndex);
 	}
 
-	virtual void Delete() override
-	{
-		if (UPaperSprite* Sprite = SpritePtr.Get())
-		{
-			FSpritePolygonCollection& Geometry = bRenderData ? Sprite->RenderGeometry : Sprite->CollisionGeometry;
-			Geometry.GeometryType = ESpritePolygonMode::FullyCustom;
-
-			if (Geometry.Polygons.IsValidIndex(PolygonIndex))
-			{
-				TArray<FVector2D>& Vertices = Geometry.Polygons[PolygonIndex].Vertices;
-				if (Vertices.IsValidIndex(VertexIndex))
-				{
-					Vertices.RemoveAt(VertexIndex);
-				}
-
-				if (Vertices.Num() < 1)
-				{
-					Geometry.Polygons.RemoveAt(PolygonIndex);
-				}
-			}
-		}
-	}
-
 	virtual void SplitEdge() override
 	{
 		// Nonsense operation on a vertex, do nothing
@@ -419,11 +414,6 @@ public:
 		return (Pos1 + Pos2) * 0.5f;
 	}
 
-	virtual void Delete() override
-	{
-		//@TODO: Support deleting edges
-	}
-
 	virtual void SplitEdge() override
 	{
 		if (UPaperSprite* Sprite = SpritePtr.Get())
@@ -506,11 +496,6 @@ public:
 		return FVector::ZeroVector;
 	}
 
-	virtual void Delete() override
-	{
-		//@TODO: Support deleting edges
-	}
-
 	virtual void SplitEdge() override
 	{
 		// Nonsense operation on a socket, do nothing
diff --git a/Engine/Plugins/Experimental/Paper2D/Source/Paper2DEditor/Private/SpriteEditor/SpriteEditorViewportClient.cpp b/Engine/Plugins/Experimental/Paper2D/Source/Paper2DEditor/Private/SpriteEditor/SpriteEditorViewportClient.cpp
--- a/Engine/Plugins/Experimental/Paper2D/Source/Paper2DEditor/Private/SpriteEditor/SpriteEditorViewportClient.cpp
+++ b/Engine/Plugins/Experimental/Paper2D/Source/Paper2DEditor/Private/SpriteEditor/SpriteEditorViewportClient.cpp
@@ -804,7 +804,7 @@ void FSpriteEditorViewportClient::ToggleShowSourceTexture()
 void FSpriteEditorViewportClient::ToggleShowMeshEdges()
 {
 	EngineShowFlags.MeshEdges = !EngineShowFlags.MeshEdges;
-	Invalidate();
+	Invalidate(); 
 }
 
 bool FSpriteEditorViewportClient::IsShowMeshEdgesChecked() const
@@ -876,7 +876,22 @@ void FSpriteEditorViewportClient::ProcessClick(FSceneView& View, HHitProxy* HitP
 			ClearSelectionSet();
 		}
 
-		SelectionSet.Add(SelectedItemProxy->Data);
+		if (Event == EInputEvent::IE_DoubleClick)
+		{
+			// Double click to select a polygon
+			if (const FSpriteSelectedVertex* SelectedVertex = SelectedItemProxy->Data->CastSelectedVertex())
+			{
+				if (SelectedVertex->IsValidInEditor(GetSpriteBeingEdited(), IsInRenderingEditMode()))
+				{
+					SelectPolygon(SelectedVertex->PolygonIndex);
+				}
+			}
+		}
+		else
+		{
+			SelectionSet.Add(SelectedItemProxy->Data);
+		}
+
 		bHandled = true;
 	}
 // 	else if (HWidgetUtilProxy* PivotProxy = HitProxyCast<HWidgetUtilProxy>(HitProxy))
@@ -893,10 +908,10 @@ void FSpriteEditorViewportClient::ProcessClick(FSceneView& View, HHitProxy* HitP
 // 	}
 	else
 	{
-		ClearSelectionSet();
-
-		if (CurrentMode == ESpriteEditorMode::EditSourceRegionMode)
+		if (IsInSourceRegionEditMode())
 		{
+			ClearSelectionSet();
+
 			if (Event == EInputEvent::IE_DoubleClick)
 			{
 				FVector4 WorldPoint = View.PixelToWorld(HitX, HitY, 0);
@@ -942,6 +957,36 @@ void FSpriteEditorViewportClient::ProcessClick(FSceneView& View, HHitProxy* HitP
 				}
 			}
 		}
+		else if (IsInRenderingEditMode() || IsInCollisionEditMode())
+		{
+			if (bIsCtrlKeyDown)
+			{
+				FVector4 WorldPoint = View.PixelToWorld(HitX, HitY, 0);
+				UPaperSprite* Sprite = RenderSpriteComponent->GetSprite();
+				if (Sprite != nullptr)
+				{
+					FVector2D SpriteSpaceClickPoint = Sprite->ConvertWorldSpaceToTextureSpace(WorldPoint);
+
+					// find a polygon to add vert to
+					if (SelectionSet.Num() > 0)
+					{
+						TSharedPtr<FSelectedItem> SelectedItemPtr = *SelectionSet.CreateIterator();
+						if (const FSpriteSelectedVertex* SelectedVertex = SelectedItemPtr->CastSelectedVertex())
+						{
+							if (SelectedVertex->IsValidInEditor(Sprite, IsInRenderingEditMode()))
+							{
+								AddPointToGeometry(SpriteSpaceClickPoint, SelectedVertex->PolygonIndex);
+							}
+						}
+					}
+					else
+					{
+						AddPointToGeometry(SpriteSpaceClickPoint);
+					}
+					bHandled = true;
+				}
+			}
+		}
 
 		if (!bHandled)
 		{
@@ -1198,25 +1243,78 @@ void FSpriteEditorViewportClient::FocusOnSprite()
 	FocusViewportOnBox(RenderSpriteComponent->Bounds.GetBox());
 }
 
+// Be sure to call this with polygonIndex and VertexIndex in descending order
+static void DeleteVertexInPolygonInternal(UPaperSprite* Sprite, FSpritePolygonCollection* Geometry, const int PolygonIndex, const int VertexIndex)
+{
+	if (Geometry->Polygons.IsValidIndex(PolygonIndex))
+	{
+		FSpritePolygon& Polygon = Geometry->Polygons[PolygonIndex];
+		if (Polygon.Vertices.IsValidIndex(VertexIndex))
+		{
+			Polygon.Vertices.RemoveAt(VertexIndex);
+			
+			// Delete polygon
+			if (Polygon.Vertices.Num() == 0)
+			{
+				Geometry->Polygons.RemoveAt(PolygonIndex);
+			}
+		}
+	}
+}
+
 void FSpriteEditorViewportClient::DeleteSelection()
 {
-	BeginTransaction(LOCTEXT("DeleteSelectionTransaction", "Delete Selection"));
+	// Make a set of polygon-vertices "pairs" to delete
+	TSet<int32> CompositeIndicesSet;
+	const int32 MaxVerticesPerPolygon = 4096;
 
-	//@TODO: Really need a proper data structure!!!
-	if (SelectionSet.Num() == 1)
+	UPaperSprite* Sprite = GetSpriteBeingEdited();
+	FSpritePolygonCollection* Geometry = GetGeometryBeingEdited();
+	for (TSharedPtr<FSelectedItem> SelectionIt : SelectionSet)
 	{
-		// Remove this item
-		TSharedPtr<FSelectedItem> SelectedItem = *SelectionSet.CreateIterator();
-		SelectedItem->Delete();
-		ClearSelectionSet();
-		bManipulationDirtiedSomething = true;
+		if (const FSpriteSelectedVertex* SelectedVertex = SelectionIt->CastSelectedVertex())
+		{
+			bool bValidSelectedVertex = SelectedVertex->IsValidInEditor(Sprite, IsInRenderingEditMode());
+			if (bValidSelectedVertex && ensure(SelectedVertex->PolygonIndex < MAX_int32))
+			{
+				// ensure we're not going to overflow the pair hash above
+				if (ensure(SelectedVertex->VertexIndex < MaxVerticesPerPolygon))
+				{
+					CompositeIndicesSet.Add(SelectedVertex->PolygonIndex * MaxVerticesPerPolygon + SelectedVertex->VertexIndex);
+				}
+				if (SelectedVertex->IsA(FSelectionTypes::Edge)) // add the "next" point for the edge
+				{
+					int32 NextIndex = (SelectedVertex->VertexIndex + 1) % Geometry->Polygons[SelectedVertex->PolygonIndex].Vertices.Num();
+					if (ensure(NextIndex < MaxVerticesPerPolygon))
+					{
+						CompositeIndicesSet.Add(SelectedVertex->PolygonIndex * MaxVerticesPerPolygon + NextIndex);
+					}
+				}
+			}
+		}
 	}
-	else
+
+	if (CompositeIndicesSet.Num() > 0)
 	{
-		// Can't handle this without a better DS, as the indices will get screwed up
-	}
+		BeginTransaction(LOCTEXT("DeleteSelectionTransaction", "Delete Selection"));
 
-	EndTransaction();
+		// Must delete vertices and polygons in descending order to preserve index validity
+		TArray<int32> CompositeIndices = CompositeIndicesSet.Array();
+		CompositeIndices.Sort();
+		for (int32 SelectedCompositeIndex = CompositeIndices.Num() - 1; SelectedCompositeIndex >= 0; --SelectedCompositeIndex)
+		{
+			int32 CompositeIndex = CompositeIndices[SelectedCompositeIndex];
+			int32 PolygonIndex = CompositeIndex / MaxVerticesPerPolygon;
+			int32 VertexIndex = CompositeIndex % MaxVerticesPerPolygon;
+			DeleteVertexInPolygonInternal(Sprite, Geometry, PolygonIndex, VertexIndex);
+		}
+
+		Geometry->GeometryType = ESpritePolygonMode::FullyCustom;
+		bManipulationDirtiedSomething = true;
+
+		ClearSelectionSet();
+		EndTransaction();
+	}
 }
 
 void FSpriteEditorViewportClient::SplitEdge()
@@ -1256,10 +1354,14 @@ void FSpriteEditorViewportClient::AddPolygon()
 		const float Size = 10.0f;
 
 		new (NewPoly.Vertices) FVector2D(BaseUV.X, BaseUV.Y);
+		new (NewPoly.Vertices) FVector2D(BaseUV.X + Size, BaseUV.Y);
 		new (NewPoly.Vertices) FVector2D(BaseUV.X + Size, BaseUV.Y + Size);
 		new (NewPoly.Vertices) FVector2D(BaseUV.X, BaseUV.Y + Size);
 		bManipulationDirtiedSomething = true;
 
+		// Select the newly added polygon
+		SelectPolygon(Geometry->Polygons.Num() - 1);
+
 		Geometry->GeometryType = ESpritePolygonMode::FullyCustom;
 	}
 
@@ -1268,7 +1370,7 @@ void FSpriteEditorViewportClient::AddPolygon()
 
 void FSpriteEditorViewportClient::AddSubtractivePolygon()
 {
-	BeginTransaction(LOCTEXT("AddPolygonTransaction", "Add Subtractive Polygon"));
+	BeginTransaction(LOCTEXT("AddSubtractivePolygonTransaction", "Add Subtractive Polygon"));
 
 	if (FSpritePolygonCollection* Geometry = GetGeometryBeingEdited())
 	{
@@ -1282,10 +1384,14 @@ void FSpriteEditorViewportClient::AddSubtractivePolygon()
 		const float Size = 10.0f;
 
 		new (NewPoly.Vertices) FVector2D(BaseUV.X, BaseUV.Y);
+		new (NewPoly.Vertices) FVector2D(BaseUV.X + Size, BaseUV.Y);
 		new (NewPoly.Vertices) FVector2D(BaseUV.X + Size, BaseUV.Y + Size);
 		new (NewPoly.Vertices) FVector2D(BaseUV.X, BaseUV.Y + Size);
 		bManipulationDirtiedSomething = true;
 
+		// Select the newly added polygon
+		SelectPolygon(Geometry->Polygons.Num() - 1);
+
 		Geometry->GeometryType = ESpritePolygonMode::FullyCustom;
 	}
 
@@ -1331,11 +1437,129 @@ FSpritePolygonCollection* FSpriteEditorViewportClient::GetGeometryBeingEdited()
 	}
 }
 
+// Selects a polygon in the currently selected mode
+void FSpriteEditorViewportClient::SelectPolygon(const int PolygonIndex)
+{
+	UPaperSprite* Sprite = GetSpriteBeingEdited();
+	FSpritePolygonCollection* Geometry = GetGeometryBeingEdited();
+	if (ensure(Geometry != nullptr && PolygonIndex >= 0 && PolygonIndex < Geometry->Polygons.Num()))
+	{
+		FSpritePolygon& Polygon = Geometry->Polygons[PolygonIndex];
+		for (int VertexIndex = 0; VertexIndex < Polygon.Vertices.Num(); ++VertexIndex)
+		{
+			TSharedPtr<FSpriteSelectedVertex> Vertex = MakeShareable(new FSpriteSelectedVertex());
+			Vertex->SpritePtr = Sprite;
+			Vertex->PolygonIndex = PolygonIndex;
+			Vertex->VertexIndex = VertexIndex;
+			Vertex->bRenderData = IsInRenderingEditMode();
+			SelectionSet.Add(Vertex);
+		}
+	}
+}
+
 void FSpriteEditorViewportClient::ClearSelectionSet()
 {
 	SelectionSet.Empty();
 }
 
+static bool ClosestPointOnLine(const FVector2D& Point, const FVector2D& LineStart, const FVector2D& LineEnd, FVector2D& OutClosestPoint)
+{
+	FVector2D DL = LineEnd - LineStart;
+	FVector2D DP = Point - LineStart;
+	float DotL = FVector2D::DotProduct(DP, DL);
+	float DotP = FVector2D::DotProduct(DL, DL);
+	if (DotP > 0.0001f)
+	{
+		float T = DotL / DotP;
+		if (T >= -0.0001f && T <= 1.0001f)
+		{
+			T = FMath::Clamp(T, 0.0f, 1.0f);
+			OutClosestPoint = LineStart + T * DL;
+			return true;
+		}
+	}
+
+	return false;
+}
+
+// Adds a point to the currently edited index
+// If SelectedPolygonIndex is set, only that polygon is considered
+void FSpriteEditorViewportClient::AddPointToGeometry(const FVector2D& TextureSpacePoint, int SelectedPolygonIndex)
+{
+	FSpritePolygonCollection* Geometry = GetGeometryBeingEdited();
+
+	int ClosestPolygonIndex = -1;
+	int ClosestVertexInsertIndex = -1;
+	float ClosestDistanceSquared = MAX_FLT;
+
+	int StartPolygonIndex = 0;
+	int EndPolygonIndex = Geometry->Polygons.Num();
+	if (SelectedPolygonIndex >= 0 && SelectedPolygonIndex < Geometry->Polygons.Num())
+	{
+		StartPolygonIndex = SelectedPolygonIndex;
+		EndPolygonIndex = SelectedPolygonIndex + 1;
+	}
+
+	for (int PolygonIndex = StartPolygonIndex; PolygonIndex < EndPolygonIndex; ++PolygonIndex)
+	{
+		FSpritePolygon& Polygon = Geometry->Polygons[PolygonIndex];
+		if (Polygon.Vertices.Num() >= 3)
+		{
+			for (int VertexIndex = 0; VertexIndex < Polygon.Vertices.Num(); ++VertexIndex)
+			{
+				FVector2D ClosestPoint;
+				FVector2D LineStart = Polygon.Vertices[VertexIndex];
+				int NextVertexIndex = (VertexIndex + 1) % Polygon.Vertices.Num();
+				FVector2D LineEnd = Polygon.Vertices[NextVertexIndex];
+				if (ClosestPointOnLine(TextureSpacePoint, LineStart, LineEnd, /*out*/ClosestPoint))
+				{
+					float CurrentDistanceSquared = FVector2D::DistSquared(ClosestPoint, TextureSpacePoint);
+					if (CurrentDistanceSquared < ClosestDistanceSquared)
+					{
+						ClosestPolygonIndex = PolygonIndex;
+						ClosestDistanceSquared = CurrentDistanceSquared;
+						ClosestVertexInsertIndex = NextVertexIndex;
+					}
+				}
+			}
+		}
+		else
+		{
+			// Simply insert the vertex
+			for (int VertexIndex = 0; VertexIndex < Polygon.Vertices.Num(); ++VertexIndex)
+			{
+				float CurrentDistanceSquared = FVector2D::DistSquared(Polygon.Vertices[VertexIndex], TextureSpacePoint);
+				if (CurrentDistanceSquared < ClosestDistanceSquared)
+				{
+					ClosestPolygonIndex = PolygonIndex;
+					ClosestDistanceSquared = CurrentDistanceSquared;
+					ClosestVertexInsertIndex = VertexIndex + 1;
+				}
+			}
+		}
+	}
+
+	if (ClosestVertexInsertIndex != -1 && ClosestPolygonIndex != -1)
+	{
+		BeginTransaction(LOCTEXT("AddPolygonVertexTransaction", "Add Vertex to Polygon"));
+		FSpritePolygon& Polygon = Geometry->Polygons[ClosestPolygonIndex];
+		Polygon.Vertices.Insert(TextureSpacePoint, ClosestVertexInsertIndex);
+		Geometry->GeometryType = ESpritePolygonMode::FullyCustom;
+		bManipulationDirtiedSomething = true;
+		EndTransaction();
+
+		// Select this vertex
+		ClearSelectionSet();
+		TSharedPtr<FSpriteSelectedVertex> SelectedVertex = MakeShareable(new FSpriteSelectedVertex());
+		SelectedVertex->SpritePtr = GetSpriteBeingEdited();
+		SelectedVertex->PolygonIndex = ClosestPolygonIndex;
+		SelectedVertex->VertexIndex = ClosestVertexInsertIndex;
+		SelectedVertex->bRenderData = IsInRenderingEditMode();
+		SelectionSet.Add(SelectedVertex);
+
+	}
+}
+
 void FSpriteEditorViewportClient::ResetMarqueeTracking()
 {
 	bIsMarqueeTracking = false;
diff --git a/Engine/Plugins/Experimental/Paper2D/Source/Paper2DEditor/Private/SpriteEditor/SpriteEditorViewportClient.h b/Engine/Plugins/Experimental/Paper2D/Source/Paper2DEditor/Private/SpriteEditor/SpriteEditorViewportClient.h
--- a/Engine/Plugins/Experimental/Paper2D/Source/Paper2DEditor/Private/SpriteEditor/SpriteEditorViewportClient.h
+++ b/Engine/Plugins/Experimental/Paper2D/Source/Paper2DEditor/Private/SpriteEditor/SpriteEditorViewportClient.h
@@ -211,6 +211,10 @@ private:
 	void EndTransaction();
 
 	void UpdateSourceTextureSpriteFromSprite(UPaperSprite* SourceSprite);
+	
+	// Selection handling
+	void SelectPolygon(const int PolygonIndex);
+	void AddPointToGeometry(const FVector2D& TextureSpacePoint, const int SelectedPolygonIndex = -1);
 	void ClearSelectionSet();
 
 
