Core - Simplified paths used for ObjectFinder

diff --git a/Engine/Plugins/Experimental/Paper2D/Source/Paper2D/Private/PaperSprite.cpp b/Engine/Plugins/Experimental/Paper2D/Source/Paper2D/Private/PaperSprite.cpp
--- a/Engine/Plugins/Experimental/Paper2D/Source/Paper2D/Private/PaperSprite.cpp
+++ b/Engine/Plugins/Experimental/Paper2D/Source/Paper2D/Private/PaperSprite.cpp
@@ -271,10 +271,10 @@ UPaperSprite::UPaperSprite(const FPostConstructInitializeProperties& PCIP)
 	bRotatedInSourceImage = false;
 #endif
 
-	static ConstructorHelpers::FObjectFinder<UMaterialInterface> MaskedMaterialRef(TEXT("/Paper2D/MaskedUnlitSpriteMaterial.MaskedUnlitSpriteMaterial"));
+	static ConstructorHelpers::FObjectFinder<UMaterialInterface> MaskedMaterialRef(TEXT("/Paper2D/MaskedUnlitSpriteMaterial"));
 	DefaultMaterial = MaskedMaterialRef.Object;
 
-	static ConstructorHelpers::FObjectFinder<UMaterialInterface> OpaqueMaterialRef(TEXT("/Paper2D/OpaqueUnlitSpriteMaterial.OpaqueUnlitSpriteMaterial"));
+	static ConstructorHelpers::FObjectFinder<UMaterialInterface> OpaqueMaterialRef(TEXT("/Paper2D/OpaqueUnlitSpriteMaterial"));
 	AlternateMaterial = OpaqueMaterialRef.Object;
 }
 
diff --git a/Engine/Plugins/Experimental/Paper2D/Source/Paper2D/Private/Terrain/PaperTerrainComponent.cpp b/Engine/Plugins/Experimental/Paper2D/Source/Paper2D/Private/Terrain/PaperTerrainComponent.cpp
--- a/Engine/Plugins/Experimental/Paper2D/Source/Paper2D/Private/Terrain/PaperTerrainComponent.cpp
+++ b/Engine/Plugins/Experimental/Paper2D/Source/Paper2D/Private/Terrain/PaperTerrainComponent.cpp
@@ -86,7 +86,7 @@ UPaperTerrainComponent::UPaperTerrainComponent(const FPostConstructInitializePro
 {
 	TestScaleFactor = 1.0f;
 
-	static ConstructorHelpers::FObjectFinder<UPaperTerrainMaterial> DefaultMaterialRef(TEXT("/Paper2D/DefaultPaperTerrainMaterial.DefaultPaperTerrainMaterial"));
+	static ConstructorHelpers::FObjectFinder<UPaperTerrainMaterial> DefaultMaterialRef(TEXT("/Paper2D/DefaultPaperTerrainMaterial"));
 	TerrainMaterial = DefaultMaterialRef.Object;
 }
 
diff --git a/Engine/Source/Runtime/Engine/Private/DocumentationActor.cpp b/Engine/Source/Runtime/Engine/Private/DocumentationActor.cpp
--- a/Engine/Source/Runtime/Engine/Private/DocumentationActor.cpp
+++ b/Engine/Source/Runtime/Engine/Private/DocumentationActor.cpp
@@ -16,7 +16,7 @@ ADocumentationActor::ADocumentationActor(const class FPostConstructInitializePro
 	RootComponent = SceneComponent;	
 
 #if WITH_EDITORONLY_DATA
- 	static ConstructorHelpers::FObjectFinder<UMaterial> MaterialAsset(TEXT("/Engine/EditorMaterials/HelpActorMaterial.HelpActorMaterial"));
+ 	static ConstructorHelpers::FObjectFinder<UMaterial> MaterialAsset(TEXT("/Engine/EditorMaterials/HelpActorMaterial"));
  	// Create a Material billboard to represent our actor
 	Billboard = PCIP.CreateDefaultSubobject<UMaterialBillboardComponent>(this, TEXT("BillboardComponent"));
 	Billboard->AddElement(MaterialAsset.Object, nullptr, false, 32.0f, 32.0f, nullptr);
diff --git a/Engine/Source/Runtime/Engine/Private/Vehicles/VehicleWheel.cpp b/Engine/Source/Runtime/Engine/Private/Vehicles/VehicleWheel.cpp
--- a/Engine/Source/Runtime/Engine/Private/Vehicles/VehicleWheel.cpp
+++ b/Engine/Source/Runtime/Engine/Private/Vehicles/VehicleWheel.cpp
@@ -19,7 +19,7 @@ UVehicleWheel::UVehicleWheel(const class FPostConstructInitializeProperties& PCI
 	static ConstructorHelpers::FObjectFinder<UStaticMesh> CollisionMeshObj(TEXT("/Engine/EngineMeshes/Cylinder"));
 	CollisionMesh = CollisionMeshObj.Object;
 
-	static ConstructorHelpers::FObjectFinder<UTireType> TireTypeObj(TEXT("/Engine/EngineTireTypes/DefaultTireType.DefaultTireType"));
+	static ConstructorHelpers::FObjectFinder<UTireType> TireTypeObj(TEXT("/Engine/EngineTireTypes/DefaultTireType"));
 	TireType = TireTypeObj.Object;
 
 	ShapeRadius = 30.0f;
