Paper2D: Allow sprites to generate navmesh geometry

diff --git a/Engine/Plugins/Experimental/Paper2D/Source/Paper2D/Private/PaperSpriteComponent.cpp b/Engine/Plugins/Experimental/Paper2D/Source/Paper2D/Private/PaperSpriteComponent.cpp
--- a/Engine/Plugins/Experimental/Paper2D/Source/Paper2D/Private/PaperSpriteComponent.cpp
+++ b/Engine/Plugins/Experimental/Paper2D/Source/Paper2D/Private/PaperSpriteComponent.cpp
@@ -21,7 +21,6 @@ UPaperSpriteComponent::UPaperSpriteComponent(const FPostConstructInitializePrope
 
 	CastShadow = false;
 	bUseAsOccluder = false;
-	bCanEverAffectNavigation = false;
 }
 
 #if WITH_EDITOR
