Paper2D: Fix sockets being drawn at the wrong location in the sprite editor when pixels/uu is not 1.0

diff --git a/Engine/Plugins/Experimental/Paper2D/Source/Paper2D/Private/PaperSprite.cpp b/Engine/Plugins/Experimental/Paper2D/Source/Paper2D/Private/PaperSprite.cpp
--- a/Engine/Plugins/Experimental/Paper2D/Source/Paper2D/Private/PaperSprite.cpp
+++ b/Engine/Plugins/Experimental/Paper2D/Source/Paper2D/Private/PaperSprite.cpp
@@ -1158,7 +1158,8 @@ FVector2D UPaperSprite::ConvertWorldSpaceDeltaToTextureSpace(const FVector& Worl
 FTransform UPaperSprite::GetPivotToWorld() const
 {
 	const FVector Translation(0, 0, 0);
-	return FTransform(Translation);
+	const FVector Scale3D(GetUnrealUnitsPerPixel());
+	return FTransform(FRotator::ZeroRotator, Translation, Scale3D);
 }
 
 FVector2D UPaperSprite::GetRawPivotPosition() const
