Paper2D: Fix socket gizmo manipulation feeling sluggish when pixels/uu is not 1.0

diff --git a/Engine/Plugins/Experimental/Paper2D/Source/Paper2DEditor/Private/SpriteEditor/SpriteEditorSelections.h b/Engine/Plugins/Experimental/Paper2D/Source/Paper2DEditor/Private/SpriteEditor/SpriteEditorSelections.h
--- a/Engine/Plugins/Experimental/Paper2D/Source/Paper2DEditor/Private/SpriteEditor/SpriteEditorSelections.h
+++ b/Engine/Plugins/Experimental/Paper2D/Source/Paper2DEditor/Private/SpriteEditor/SpriteEditorSelections.h
@@ -476,7 +476,8 @@ public:
 		{
 			if (FPaperSpriteSocket* Socket = Sprite->FindSocket(SocketName))
 			{
-				const FVector Delta3D = (PaperAxisX * Delta.X) + (PaperAxisY * -Delta.Y);
+				const FVector Delta3D_UU = (PaperAxisX * Delta.X) + (PaperAxisY * -Delta.Y);
+				const FVector Delta3D = Delta3D_UU * Sprite->GetPixelsPerUnrealUnit();
 				Socket->LocalTransform.SetLocation(Socket->LocalTransform.GetLocation() + Delta3D);
 			}
 		}
