Paper2D: Preliminary support for sockets on flipbooks (exposing sockets defined in source sprites)

diff --git a/Engine/Plugins/Experimental/Paper2D/Source/Paper2D/Classes/PaperFlipbook.h b/Engine/Plugins/Experimental/Paper2D/Source/Paper2D/Classes/PaperFlipbook.h
--- a/Engine/Plugins/Experimental/Paper2D/Source/Paper2D/Classes/PaperFlipbook.h
+++ b/Engine/Plugins/Experimental/Paper2D/Source/Paper2D/Classes/PaperFlipbook.h
@@ -83,6 +83,15 @@ public:
 		return KeyFrames[Index];
 	}
 
+	// Search for a socket at the specified frame
+	bool FindSocket(FName SocketName, int32 KeyFrameIndex, FTransform& OutLocalTransform);
+
+	// Returns true if the flipbook has any sockets
+	bool HasAnySockets() const;
+
+	// Returns a list of all of the sockets
+	void QuerySupportedSockets(TArray<FComponentSocketDescription>& OutSockets) const;
+
 	// Rebuilds cached data about the animation (such as total number of frames that the keyframes span, etc...)
 	void InvalidateCachedData();
 
diff --git a/Engine/Plugins/Experimental/Paper2D/Source/Paper2D/Classes/PaperFlipbookComponent.h b/Engine/Plugins/Experimental/Paper2D/Source/Paper2D/Classes/PaperFlipbookComponent.h
--- a/Engine/Plugins/Experimental/Paper2D/Source/Paper2D/Classes/PaperFlipbookComponent.h
+++ b/Engine/Plugins/Experimental/Paper2D/Source/Paper2D/Classes/PaperFlipbookComponent.h
@@ -156,6 +156,12 @@ public:
 	virtual const UObject* AdditionalStatObject() const override;
 	// End of UActorComponent interface
 
+	// USceneComponent interface
+	virtual bool HasAnySockets() const override;
+	virtual FTransform GetSocketTransform(FName InSocketName, ERelativeTransformSpace TransformSpace = RTS_World) const override;
+	virtual void QuerySupportedSockets(TArray<FComponentSocketDescription>& OutSockets) const override;
+	// End of USceneComponent interface
+
 	// UPrimitiveComponent interface
 	virtual FPrimitiveSceneProxy* CreateSceneProxy() override;
 	virtual FBoxSphereBounds CalcBounds(const FTransform & LocalToWorld) const override;
diff --git a/Engine/Plugins/Experimental/Paper2D/Source/Paper2D/Private/PaperFlipbook.cpp b/Engine/Plugins/Experimental/Paper2D/Source/Paper2D/Private/PaperFlipbook.cpp
--- a/Engine/Plugins/Experimental/Paper2D/Source/Paper2D/Private/PaperFlipbook.cpp
+++ b/Engine/Plugins/Experimental/Paper2D/Source/Paper2D/Private/PaperFlipbook.cpp
@@ -143,6 +143,65 @@ FBoxSphereBounds UPaperFlipbook::GetRenderBounds() const
 	return MergedBoundingBox;
 }
 
+bool UPaperFlipbook::FindSocket(FName SocketName, int32 KeyFrameIndex, FTransform& OutLocalTransform)
+{
+	if (KeyFrames.IsValidIndex(KeyFrameIndex))
+	{
+		if (UPaperSprite* SpriteFrame = KeyFrames[KeyFrameIndex].Sprite)
+		{
+			if (FPaperSpriteSocket* SocketInfo = SpriteFrame->FindSocket(SocketName))
+			{
+				OutLocalTransform = SocketInfo->LocalTransform;
+				OutLocalTransform.ScaleTranslation(SpriteFrame->GetUnrealUnitsPerPixel());
+
+				return true;
+			}
+		}
+	}
+
+	return false;
+}
+
+bool UPaperFlipbook::HasAnySockets() const
+{
+	for (const FPaperFlipbookKeyFrame& KeyFrame : KeyFrames)
+	{
+		if (KeyFrame.Sprite != nullptr)
+		{
+			if (KeyFrame.Sprite->HasAnySockets())
+			{
+				return true;
+			}
+		}
+	}
+
+	return false;
+}
+
+void UPaperFlipbook::QuerySupportedSockets(TArray<FComponentSocketDescription>& OutSockets) const
+{
+	TSet<FName> SocketNames;
+	TArray<FComponentSocketDescription> FrameDescriptions;
+
+	for (const FPaperFlipbookKeyFrame& KeyFrame : KeyFrames)
+	{
+		if (KeyFrame.Sprite != nullptr)
+		{
+			FrameDescriptions.Reset();
+			KeyFrame.Sprite->QuerySupportedSockets(/*out*/ FrameDescriptions);
+
+			for (const FComponentSocketDescription& FrameSocket : FrameDescriptions)
+			{
+				if (!SocketNames.Contains(FrameSocket.Name))
+				{
+					SocketNames.Add(FrameSocket.Name);
+					OutSockets.Add(FrameSocket);
+				}
+			}
+		}
+	}
+}
+
 void UPaperFlipbook::InvalidateCachedData()
 {
 	// No cached data yet, but the functions that currently have to iterate over all frames can use cached data in the future
diff --git a/Engine/Plugins/Experimental/Paper2D/Source/Paper2D/Private/PaperFlipbookComponent.cpp b/Engine/Plugins/Experimental/Paper2D/Source/Paper2D/Private/PaperFlipbookComponent.cpp
--- a/Engine/Plugins/Experimental/Paper2D/Source/Paper2D/Private/PaperFlipbookComponent.cpp
+++ b/Engine/Plugins/Experimental/Paper2D/Source/Paper2D/Private/PaperFlipbookComponent.cpp
@@ -96,6 +96,9 @@ void UPaperFlipbookComponent::CalculateCurrentFrame()
 
 	if (CachedFrameIndex != LastCachedFrame)
 	{
+		// Update children transforms in case we have anything attached to an animated socket
+		UpdateChildTransforms();
+
 		// Indicate we need to send new dynamic data.
 		MarkRenderDynamicDataDirty();
 	}
@@ -458,3 +461,53 @@ float UPaperFlipbookComponent::GetFlipbookFramerate() const
 {
 	return (SourceFlipbook != nullptr) ? SourceFlipbook->GetFramesPerSecond() : 15.0f;
 }
+
+bool UPaperFlipbookComponent::HasAnySockets() const
+{
+	if (SourceFlipbook != nullptr)
+	{
+		return SourceFlipbook->HasAnySockets();
+	}
+
+	return false;
+}
+
+FTransform UPaperFlipbookComponent::GetSocketTransform(FName InSocketName, ERelativeTransformSpace TransformSpace) const
+{
+	if (SourceFlipbook != nullptr)
+	{
+		FTransform SocketLocalTransform;
+		if (SourceFlipbook->FindSocket(InSocketName, (CachedFrameIndex != INDEX_NONE) ? CachedFrameIndex : 0, /*out*/ SocketLocalTransform))
+		{
+			switch (TransformSpace)
+			{
+			case RTS_World:
+				return SocketLocalTransform * ComponentToWorld;
+
+			case RTS_Actor:
+				if (const AActor* Actor = GetOwner())
+				{
+					const FTransform SocketTransform = SocketLocalTransform * ComponentToWorld;
+					return SocketTransform.GetRelativeTransform(Actor->GetTransform());
+				}
+				break;
+
+			case RTS_Component:
+				return SocketLocalTransform;
+
+			default:
+				check(false);
+			}
+		}
+	}
+
+	return Super::GetSocketTransform(InSocketName, TransformSpace);
+}
+
+void UPaperFlipbookComponent::QuerySupportedSockets(TArray<FComponentSocketDescription>& OutSockets) const
+{
+	if (SourceFlipbook != nullptr)
+	{
+		return SourceFlipbook->QuerySupportedSockets(OutSockets);
+	}
+}
