Paper2D: Reinstated "Show Source Texture" mode, the texture is shown rotated appropriately to match the sprite.
- When tracing contours, negative contours are set appropriately to cut out holes. Still clips pretty badly though.

diff --git a/Engine/Plugins/Experimental/Paper2D/Source/Paper2D/Classes/PaperSprite.h b/Engine/Plugins/Experimental/Paper2D/Source/Paper2D/Classes/PaperSprite.h
--- a/Engine/Plugins/Experimental/Paper2D/Source/Paper2D/Classes/PaperSprite.h
+++ b/Engine/Plugins/Experimental/Paper2D/Source/Paper2D/Classes/PaperSprite.h
@@ -206,7 +206,11 @@ public:
 
 	void SetTrim(bool bTrimmed, const FVector2D& OriginInSourceImage, const FVector2D& SourceImageDimension);
 	void SetRotated(bool bRotated);
-	void SetPivot(ESpritePivotMode::Type PivotMode, FVector2D CustomTextureSpacePivot);
+	void SetPivotMode(ESpritePivotMode::Type PivotMode, FVector2D CustomTextureSpacePivot);
+	
+	bool IsRotatedInSourceImage() const { return bRotatedInSourceImage; }
+	ESpritePivotMode::Type GetPivotMode(FVector2D& OutCustomTextureSpacePivot) const { OutCustomTextureSpacePivot = CustomPivotPoint; return PivotMode; }
+
 
 	FVector2D GetSourceUV() const { return SourceUV; }
 	FVector2D GetSourceSize() const { return SourceDimension; }
diff --git a/Engine/Plugins/Experimental/Paper2D/Source/Paper2D/Private/PaperGeomTools.cpp b/Engine/Plugins/Experimental/Paper2D/Source/Paper2D/Private/PaperGeomTools.cpp
--- a/Engine/Plugins/Experimental/Paper2D/Source/Paper2D/Private/PaperGeomTools.cpp
+++ b/Engine/Plugins/Experimental/Paper2D/Source/Paper2D/Private/PaperGeomTools.cpp
@@ -6,7 +6,7 @@
 #include "PaperGeomTools.h"
 #include "GeomTools.h"
 
-static bool IsPolygonWindingCCW(const TArray<FVector2D>& Points)
+bool PaperGeomTools::IsPolygonWindingCCW(const TArray<FVector2D>& Points)
 {
 	float Sum = 0.0f;
 	const int PointCount = Points.Num();
@@ -370,37 +370,38 @@ bool PaperGeomTools::ArePolygonsValid(const TArray<FSpritePolygon>& Polygons)
 	}
 
 	// Make sure no polygons are self-intersecting
-	for (int PolygonIndex = 0; PolygonIndex < PolygonCount; ++PolygonIndex)
-	{
-		const TArray<FVector2D>& Polygon = Polygons[PolygonIndex].Vertices;
-		const int PointCount = Polygon.Num();
-		for (int PointIndexA = 0; PointIndexA < PointCount; ++PointIndexA)
-		{
-			const FVector2D& A0 = Polygon[PointIndexA];
-			const FVector2D& A1 = Polygon[(PointIndexA + 1) % PointCount];
-			const FVector2D A10 = A1 - A0;
-			for (int PointIndexB = 0; PointIndexB < PointCount - 3; ++PointIndexB)
-			{
-				const FVector2D& B0 = Polygon[(PointIndexA + 2 + PointIndexB) % PointCount];
-				const FVector2D& B1 = Polygon[(PointIndexA + 3 + PointIndexB) % PointCount];
-				const FVector2D B10 = B1 - B0;
-
-				const float DetS = A10.X * B10.Y - A10.Y * B10.X;
-				const float DetT = B10.X * A10.Y - B10.Y * A10.X;
-				if (DetS != 0.0f && DetT != 0.0f)
-				{
-					const float S = (A0.Y * B10.X - B0.Y * B10.X - A0.X * B10.Y + B0.X * B10.Y) / DetS;
-					const float T = (B0.Y * A10.X - A0.Y * A10.X - B0.X * A10.Y + A0.X * A10.Y) / DetT;
-					if (S >= 0.0f && S <= 1.0f && T >= 0.0f && T <= 1.0f)
-					{
-						// Edges intersect
-						UE_LOG(LogPaper2D, Log, TEXT("Polygon %d is self intersecting"), PolygonIndex);
-						return false;
-					}
-				}
-			}
-		}
-	}
+	// Disabled for now until we decide what to do with this - contour tracing can generate invalid polys & degenerate edges
+	//for (int PolygonIndex = 0; PolygonIndex < PolygonCount; ++PolygonIndex)
+	//{
+	//	const TArray<FVector2D>& Polygon = Polygons[PolygonIndex].Vertices;
+	//	const int PointCount = Polygon.Num();
+	//	for (int PointIndexA = 0; PointIndexA < PointCount; ++PointIndexA)
+	//	{
+	//		const FVector2D& A0 = Polygon[PointIndexA];
+	//		const FVector2D& A1 = Polygon[(PointIndexA + 1) % PointCount];
+	//		const FVector2D A10 = A1 - A0;
+	//		for (int PointIndexB = 0; PointIndexB < PointCount - 3; ++PointIndexB)
+	//		{
+	//			const FVector2D& B0 = Polygon[(PointIndexA + 2 + PointIndexB) % PointCount];
+	//			const FVector2D& B1 = Polygon[(PointIndexA + 3 + PointIndexB) % PointCount];
+	//			const FVector2D B10 = B1 - B0;
+
+	//			const float DetS = A10.X * B10.Y - A10.Y * B10.X;
+	//			const float DetT = B10.X * A10.Y - B10.Y * A10.X;
+	//			if (DetS != 0.0f && DetT != 0.0f)
+	//			{
+	//				const float S = (A0.Y * B10.X - B0.Y * B10.X - A0.X * B10.Y + B0.X * B10.Y) / DetS;
+	//				const float T = (B0.Y * A10.X - A0.Y * A10.X - B0.X * A10.Y + A0.X * A10.Y) / DetT;
+	//				if (S >= 0.0f && S <= 1.0f && T >= 0.0f && T <= 1.0f)
+	//				{
+	//					// Edges intersect
+	//					UE_LOG(LogPaper2D, Log, TEXT("Polygon %d is self intersecting"), PolygonIndex);
+	//					return false;
+	//				}
+	//			}
+	//		}
+	//	}
+	//}
 
 	return true;
 }
diff --git a/Engine/Plugins/Experimental/Paper2D/Source/Paper2D/Private/PaperGeomTools.h b/Engine/Plugins/Experimental/Paper2D/Source/Paper2D/Private/PaperGeomTools.h
--- a/Engine/Plugins/Experimental/Paper2D/Source/Paper2D/Private/PaperGeomTools.h
+++ b/Engine/Plugins/Experimental/Paper2D/Source/Paper2D/Private/PaperGeomTools.h
@@ -15,6 +15,10 @@ public:
 	// Entries in the array that aren't polygons (< 3 points) are removed
 	static TArray<FSpritePolygon> CorrectPolygonWinding(const TArray<FSpritePolygon>& Polygons);
 
+	// Returns true if the points forming a polygon have CCW winding
+	// Returns true if the polygon isn't valid
+	static bool IsPolygonWindingCCW(const TArray<FVector2D>& Points);
+
 	// Checks that these polygons can be successfully triangulated	
 	static bool ArePolygonsValid(const TArray<FSpritePolygon>& Polygons);
 
diff --git a/Engine/Plugins/Experimental/Paper2D/Source/Paper2D/Private/PaperSprite.cpp b/Engine/Plugins/Experimental/Paper2D/Source/Paper2D/Private/PaperSprite.cpp
--- a/Engine/Plugins/Experimental/Paper2D/Source/Paper2D/Private/PaperSprite.cpp
+++ b/Engine/Plugins/Experimental/Paper2D/Source/Paper2D/Private/PaperSprite.cpp
@@ -790,6 +790,9 @@ void UPaperSprite::BuildGeometryFromContours(FSpritePolygonCollection& GeomOwner
 			{
 				new (Polygon.Vertices) FVector2D(Contour[PointIndex]);
 			}
+
+			// Get intended winding
+			Polygon.bNegativeWinding = !PaperGeomTools::IsPolygonWindingCCW(Polygon.Vertices);
 		}
 	}
 }
@@ -1066,10 +1069,12 @@ void UPaperSprite::SetRotated(bool bRotated)
 	RebuildCollisionData();
 }
 
-void UPaperSprite::SetPivot(ESpritePivotMode::Type PivotMode, FVector2D CustomTextureSpacePivot)
+void UPaperSprite::SetPivotMode(ESpritePivotMode::Type PivotMode, FVector2D CustomTextureSpacePivot)
 {
 	this->PivotMode = PivotMode;
 	this->CustomPivotPoint = CustomTextureSpacePivot;
+	RebuildRenderData();
+	RebuildCollisionData();
 }
 
 FVector2D UPaperSprite::ConvertTextureSpaceToPivotSpace(FVector2D Input) const
diff --git a/Engine/Plugins/Experimental/Paper2D/Source/Paper2DEditor/Private/SpriteEditor/SpriteEditor.cpp b/Engine/Plugins/Experimental/Paper2D/Source/Paper2DEditor/Private/SpriteEditor/SpriteEditor.cpp
--- a/Engine/Plugins/Experimental/Paper2D/Source/Paper2DEditor/Private/SpriteEditor/SpriteEditor.cpp
+++ b/Engine/Plugins/Experimental/Paper2D/Source/Paper2DEditor/Private/SpriteEditor/SpriteEditor.cpp
@@ -104,7 +104,7 @@ void SSpriteEditorViewport::BindCommands()
 	CommandList->MapAction(
 		Commands.SetShowSourceTexture,
 		FExecuteAction::CreateSP( EditorViewportClientRef, &FSpriteEditorViewportClient::ToggleShowSourceTexture ),
-		FCanExecuteAction(),
+		FCanExecuteAction::CreateSP(EditorViewportClientRef, &FSpriteEditorViewportClient::CanShowSourceTexture),
 		FIsActionChecked::CreateSP( EditorViewportClientRef, &FSpriteEditorViewportClient::IsShowSourceTextureChecked ) );
 
 	CommandList->MapAction(
@@ -508,12 +508,11 @@ void FSpriteEditor::ExtendToolbar()
 // 			}
 // 			ToolbarBuilder.EndSection();
 
-			// Temporarily disabled until we decide what to do with it. Its not very useful as it is right now.
-//			ToolbarBuilder.BeginSection("Command");
-//			{
-//				ToolbarBuilder.AddToolBarButton(FSpriteEditorCommands::Get().SetShowSourceTexture);
-//			}
-//			ToolbarBuilder.EndSection();
+			ToolbarBuilder.BeginSection("Command");
+			{
+				ToolbarBuilder.AddToolBarButton(FSpriteEditorCommands::Get().SetShowSourceTexture);
+			}
+			ToolbarBuilder.EndSection();
 
 			ToolbarBuilder.BeginSection("Tools");
 			{
diff --git a/Engine/Plugins/Experimental/Paper2D/Source/Paper2DEditor/Private/SpriteEditor/SpriteEditorViewportClient.cpp b/Engine/Plugins/Experimental/Paper2D/Source/Paper2DEditor/Private/SpriteEditor/SpriteEditorViewportClient.cpp
--- a/Engine/Plugins/Experimental/Paper2D/Source/Paper2DEditor/Private/SpriteEditor/SpriteEditorViewportClient.cpp
+++ b/Engine/Plugins/Experimental/Paper2D/Source/Paper2DEditor/Private/SpriteEditor/SpriteEditorViewportClient.cpp
@@ -110,6 +110,47 @@ void FSpriteEditorViewportClient::UpdateSourceTextureSpriteFromSprite(UPaperSpri
 			FComponentReregisterContext ReregisterSprite(SourceTextureViewComponent);
 			TargetSprite->InitializeSprite(SourceSprite->SourceTexture, SourceSprite->PixelsPerUnrealUnit);
 		}
+
+
+		// Position the sprite for the mode its meant to be in
+		FVector2D CurrentPivotPosition;
+		ESpritePivotMode::Type CurrentPivotMode = TargetSprite->GetPivotMode(/*out*/CurrentPivotPosition);
+
+		FVector Translation(1.0f * PaperAxisZ);
+		if (IsInSourceRegionEditMode())
+		{
+			if (CurrentPivotMode != ESpritePivotMode::Bottom_Left)
+			{
+				TargetSprite->SetPivotMode(ESpritePivotMode::Bottom_Left, FVector2D::ZeroVector);
+				TargetSprite->PostEditChange();
+			}
+			SourceTextureViewComponent->SetSpriteColor(FLinearColor::White);
+			SourceTextureViewComponent->SetWorldTransform(FTransform(Translation));
+		}
+		else
+		{
+			// Tint the source texture darker to help distinguish the two
+			FLinearColor DarkTintColor(0.05f, 0.05f, 0.05f, 1.0f);
+
+			FVector2D PivotPosition = SourceSprite->GetPivotPosition();
+			if (CurrentPivotMode != ESpritePivotMode::Custom || CurrentPivotPosition != PivotPosition)
+			{
+				TargetSprite->SetPivotMode(ESpritePivotMode::Custom, PivotPosition);
+				TargetSprite->PostEditChange();
+			}
+			SourceTextureViewComponent->SetSpriteColor(DarkTintColor);
+
+			bool bRotated = SourceSprite->IsRotatedInSourceImage();
+			if (bRotated)
+			{
+				FQuat Rotation(PaperAxisZ, FMath::DegreesToRadians(90.0f));
+				SourceTextureViewComponent->SetWorldTransform(FTransform(Rotation, Translation));
+			}
+			else
+			{
+				SourceTextureViewComponent->SetWorldTransform(FTransform(Translation));
+			}
+		}
 	}
 	else
 	{
@@ -723,7 +764,7 @@ void FSpriteEditorViewportClient::Draw(const FSceneView* View, FPrimitiveDrawInt
 
 	// We don't draw the pivot when showing the source region
 	// The pivot may be outside the actual texture bounds there
-	if (bShowPivot && !bShowSourceTexture && CurrentMode != ESpriteEditorMode::EditSourceRegionMode)
+	if (bShowPivot && !bShowSourceTexture && !IsInSourceRegionEditMode())
 	{
 		FUnrealEdUtils::DrawWidget(View, PDI, RenderSpriteComponent->ComponentToWorld.ToMatrixWithScale(), 0, 0, EAxisList::XZ, EWidgetMovementMode::WMM_Translate);
 	}
@@ -753,14 +794,14 @@ void FSpriteEditorViewportClient::Tick(float DeltaSeconds)
 		const FVector PivotInWorldSpace = TextureSpaceToWorldSpace(PivotInTextureSpace);
 		RenderSpriteComponent->SetRelativeLocation(PivotInWorldSpace);
 
-		bool bSourceTextureViewComponentVisibility = bShowSourceTexture || (CurrentMode == ESpriteEditorMode::EditSourceRegionMode);
+		bool bSourceTextureViewComponentVisibility = bShowSourceTexture || IsInSourceRegionEditMode();
 		if (bSourceTextureViewComponentVisibility != SourceTextureViewComponent->IsVisible())
 		{
 			bDeferZoomToSprite = true;
 			SourceTextureViewComponent->SetVisibility(bSourceTextureViewComponentVisibility);
 		}
 
-		bool bRenderTextureViewComponentVisibility = !bSourceTextureViewComponentVisibility;
+		bool bRenderTextureViewComponentVisibility = !IsInSourceRegionEditMode();
 		if (bRenderTextureViewComponentVisibility != RenderSpriteComponent->IsVisible())
 		{
 			bDeferZoomToSprite = true;
@@ -957,7 +998,7 @@ void FSpriteEditorViewportClient::ProcessClick(FSceneView& View, HHitProxy* HitP
 				}
 			}
 		}
-		else if (IsInRenderingEditMode() || IsInCollisionEditMode())
+		else if (IsEditingGeometry())
 		{
 			if (bIsCtrlKeyDown)
 			{
diff --git a/Engine/Plugins/Experimental/Paper2D/Source/Paper2DEditor/Private/SpriteEditor/SpriteEditorViewportClient.h b/Engine/Plugins/Experimental/Paper2D/Source/Paper2DEditor/Private/SpriteEditor/SpriteEditorViewportClient.h
--- a/Engine/Plugins/Experimental/Paper2D/Source/Paper2DEditor/Private/SpriteEditor/SpriteEditorViewportClient.h
+++ b/Engine/Plugins/Experimental/Paper2D/Source/Paper2DEditor/Private/SpriteEditor/SpriteEditorViewportClient.h
@@ -66,9 +66,6 @@ public:
 	virtual ECoordSystem GetWidgetCoordSystemSpace() const;
 	// End of FEditorViewportClient interface
 
-	void ToggleShowSourceTexture();
-	bool IsShowSourceTextureChecked() const { return bShowSourceTexture; }
-
 	void ToggleShowSockets() { bShowSockets = !bShowSockets; Invalidate(); }
 	bool IsShowSocketsChecked() const { return bShowSockets; }
 
@@ -96,6 +93,10 @@ public:
 	//
 	bool IsEditingGeometry() const { return IsInCollisionEditMode() || IsInRenderingEditMode(); }
 
+	void ToggleShowSourceTexture();
+	bool IsShowSourceTextureChecked() const { return bShowSourceTexture; }
+	bool CanShowSourceTexture() const { return !IsInSourceRegionEditMode(); }
+
 	void FocusOnSprite();
 
 	// Geometry editing commands
diff --git a/Engine/Plugins/Experimental/Paper2D/Source/PaperJsonImporter/Private/PaperJsonSpriteSheetImporter.cpp b/Engine/Plugins/Experimental/Paper2D/Source/PaperJsonImporter/Private/PaperJsonSpriteSheetImporter.cpp
--- a/Engine/Plugins/Experimental/Paper2D/Source/PaperJsonImporter/Private/PaperJsonSpriteSheetImporter.cpp
+++ b/Engine/Plugins/Experimental/Paper2D/Source/PaperJsonImporter/Private/PaperJsonSpriteSheetImporter.cpp
@@ -522,7 +522,7 @@ bool FPaperJsonSpriteSheetImporter::PerformImport(const FString& LongPackagePath
 			TextureSpacePivotPoint.X = Frame.SpritePosInSheet.X - Frame.SpriteSourcePos.X + Frame.ImageSourceSize.X * Frame.Pivot.X;
 			TextureSpacePivotPoint.Y = Frame.SpritePosInSheet.Y - Frame.SpriteSourcePos.Y + Frame.ImageSourceSize.Y * Frame.Pivot.Y;
 		}
-		TargetSprite->SetPivot(PivotType, TextureSpacePivotPoint);
+		TargetSprite->SetPivotMode(PivotType, TextureSpacePivotPoint);
 
 		// Create the entry in the animation
 		SpriteSheet->SpriteNames.Add(Frame.FrameName.ToString());
