Paper2D: Add a default material to UPaperFlipbook assets and changed the Material override in UPaperFlipbookComponent to default to null
- This should have no impact on existing content, but allows flipbooks to have a material set once and show up correctly inside of the flipbook editor, rather than requiring it to be set on each instance.

diff --git a/Engine/Plugins/Experimental/Paper2D/Source/Paper2D/Classes/PaperFlipbook.h b/Engine/Plugins/Experimental/Paper2D/Source/Paper2D/Classes/PaperFlipbook.h
--- a/Engine/Plugins/Experimental/Paper2D/Source/Paper2D/Classes/PaperFlipbook.h
+++ b/Engine/Plugins/Experimental/Paper2D/Source/Paper2D/Classes/PaperFlipbook.h
@@ -41,6 +41,10 @@ protected:
 	UPROPERTY(Category=Sprite, EditAnywhere)
 	TArray<FPaperFlipbookKeyFrame> KeyFrames;
 
+	// The material to use on a flipbook player instance if not overridden
+	UPROPERTY(Category=Sprite, EditAnywhere, BlueprintReadOnly)
+	UMaterialInterface* DefaultMaterial;
+
 public:
 	// Returns the nominal frame rate to play this flipbook animation back at
 	int32 GetFramesPerSecond() const;
@@ -92,6 +96,9 @@ public:
 	// Returns a list of all of the sockets
 	void QuerySupportedSockets(TArray<FComponentSocketDescription>& OutSockets) const;
 
+	// Return the default material for this flipbook
+	UMaterialInterface* GetDefaultMaterial() const { return DefaultMaterial; }
+
 	// Rebuilds cached data about the animation (such as total number of frames that the keyframes span, etc...)
 	void InvalidateCachedData();
 
diff --git a/Engine/Plugins/Experimental/Paper2D/Source/Paper2D/Private/PaperFlipbook.cpp b/Engine/Plugins/Experimental/Paper2D/Source/Paper2D/Private/PaperFlipbook.cpp
--- a/Engine/Plugins/Experimental/Paper2D/Source/Paper2D/Private/PaperFlipbook.cpp
+++ b/Engine/Plugins/Experimental/Paper2D/Source/Paper2D/Private/PaperFlipbook.cpp
@@ -15,6 +15,9 @@ UPaperFlipbook::UPaperFlipbook(const FPostConstructInitializeProperties& PCIP)
 	: Super(PCIP)
 {
 	FramesPerSecond = 15;
+
+	static ConstructorHelpers::FObjectFinder<UMaterialInterface> MaskedMaterialRef(TEXT("/Paper2D/MaskedUnlitSpriteMaterial"));
+	DefaultMaterial = MaskedMaterialRef.Object;
 }
 
 int32 UPaperFlipbook::GetFramesPerSecond() const
diff --git a/Engine/Plugins/Experimental/Paper2D/Source/Paper2D/Private/PaperFlipbookComponent.cpp b/Engine/Plugins/Experimental/Paper2D/Source/Paper2D/Private/PaperFlipbookComponent.cpp
--- a/Engine/Plugins/Experimental/Paper2D/Source/Paper2D/Private/PaperFlipbookComponent.cpp
+++ b/Engine/Plugins/Experimental/Paper2D/Source/Paper2D/Private/PaperFlipbookComponent.cpp
@@ -14,8 +14,7 @@ UPaperFlipbookComponent::UPaperFlipbookComponent(const FPostConstructInitializeP
 	BodyInstance.SetCollisionEnabled(ECollisionEnabled::NoCollision);
 
 	SetCollisionProfileName(UCollisionProfile::NoCollision_ProfileName);
-	static ConstructorHelpers::FObjectFinder<UMaterial> DefaultMaterial(TEXT("/Paper2D/DefaultSpriteMaterial"));
-	Material = DefaultMaterial.Object;
+	Material = nullptr;
 
 	SpriteColor = FLinearColor::White;
 
@@ -275,7 +274,18 @@ UPaperFlipbook* UPaperFlipbookComponent::GetFlipbook()
 
 UMaterialInterface* UPaperFlipbookComponent::GetSpriteMaterial() const
 {
-	return Material;
+	if (Material != nullptr)
+	{
+		return Material;
+	}
+	else if (SourceFlipbook != nullptr)
+	{
+		return SourceFlipbook->GetDefaultMaterial();
+	}
+	else
+	{
+		return nullptr;
+	}
 }
 
 void UPaperFlipbookComponent::SetSpriteColor(FLinearColor NewColor)
