Paper2D: Convert UPaperFlipbookComponent to derive from UMeshComponent

diff --git a/Engine/Plugins/Experimental/Paper2D/Source/Paper2D/Classes/PaperFlipbookComponent.h b/Engine/Plugins/Experimental/Paper2D/Source/Paper2D/Classes/PaperFlipbookComponent.h
--- a/Engine/Plugins/Experimental/Paper2D/Source/Paper2D/Classes/PaperFlipbookComponent.h
+++ b/Engine/Plugins/Experimental/Paper2D/Source/Paper2D/Classes/PaperFlipbookComponent.h
@@ -9,7 +9,7 @@
 DECLARE_DYNAMIC_MULTICAST_DELEGATE(FFlipbookFinishedPlaySignature);
 
 UCLASS(ShowCategories=(Mobility, ComponentReplication), ClassGroup=Paper2D, EarlyAccessPreview, meta=(BlueprintSpawnableComponent))
-class PAPER2D_API UPaperFlipbookComponent : public UPrimitiveComponent
+class PAPER2D_API UPaperFlipbookComponent : public UMeshComponent
 {
 	GENERATED_UCLASS_BODY()
 
@@ -18,9 +18,9 @@ protected:
 	UPROPERTY(Category=Sprite, EditAnywhere, meta=(DisplayThumbnail = "true"), ReplicatedUsing=OnRep_SourceFlipbook)
 	UPaperFlipbook* SourceFlipbook;
 
-	/** Material used to display the flipbook */
-	UPROPERTY(Category=Sprite, EditAnywhere)
-	UMaterialInterface* Material;
+	// DEPRECATED in 4.5: The material override for this flipbook component (if any); replaced by the Materials array inherited from UMeshComponent
+	UPROPERTY()
+	UMaterialInterface* Material_DEPRECATED;
 
 	/** Current play rate of the flipbook */
 	UPROPERTY(Category=Sprite, EditAnywhere, Replicated)
@@ -64,8 +64,9 @@ public:
 	UFUNCTION(BlueprintPure, Category="Sprite")
 	virtual UPaperFlipbook* GetFlipbook();
 
-	/** Gets the material used by this instance */
-	UFUNCTION(BlueprintCallable, Category="Sprite")
+	/** DEPRECATED! Use GetMaterial() instead.  Gets the material used by this instance */
+	//@TODO: DEPRECATED(4.5, "GetSpriteMaterial has been replaced by GetMaterial")
+	UFUNCTION(BlueprintCallable, Category="Sprite", meta=(DeprecatedFunction, DeprecationMessage="Use GetMaterial instead"))
 	UMaterialInterface* GetSpriteMaterial() const;
 
 	/** Set color of the sprite */
@@ -150,6 +151,13 @@ protected:
 	void TickFlipbook(float DeltaTime);
 
 public:
+	// UObject interface
+#if WITH_EDITORONLY_DATA
+	virtual void Serialize(FArchive& Ar) override;
+	virtual void PostLoad() override;
+#endif
+	// End of UObject interface
+
 	// UActorComponent interface
 	virtual void TickComponent(float DeltaTime, enum ELevelTick TickType, FActorComponentTickFunction* ThisTickFunction) override;
 	virtual void SendRenderDynamicData_Concurrent() override;
@@ -165,6 +173,11 @@ public:
 	// UPrimitiveComponent interface
 	virtual FPrimitiveSceneProxy* CreateSceneProxy() override;
 	virtual FBoxSphereBounds CalcBounds(const FTransform & LocalToWorld) const override;
+	virtual void GetUsedTextures(TArray<UTexture*>& OutTextures, EMaterialQualityLevel::Type QualityLevel) override;
+	virtual UMaterialInterface* GetMaterial(int32 MaterialIndex) const override;
+	virtual void GetUsedMaterials(TArray<UMaterialInterface*>& OutMaterials) const override;
+	virtual void GetStreamingTextureInfo(TArray<FStreamingTexturePrimitiveInfo>& OutStreamingTextures) const override;
+	virtual int32 GetNumMaterials() const override;
 	// End of UPrimitiveComponent interface
 };
 
diff --git a/Engine/Plugins/Experimental/Paper2D/Source/Paper2D/Private/PaperCustomVersion.h b/Engine/Plugins/Experimental/Paper2D/Source/Paper2D/Private/PaperCustomVersion.h
--- a/Engine/Plugins/Experimental/Paper2D/Source/Paper2D/Private/PaperCustomVersion.h
+++ b/Engine/Plugins/Experimental/Paper2D/Source/Paper2D/Private/PaperCustomVersion.h
@@ -37,6 +37,9 @@ struct FPaperCustomVersion
 		// Added the option to snap pivots to a pixel grid
 		AddPivotSnapToPixelGrid = 9,
 
+		// Converted UPaperFlipbookComponent to be derived from UMeshComponent
+		ConvertPaperFlipbookComponentToBeMeshComponent = 10,
+
 		// -----<new versions can be added above this line>-------------------------------------------------
 		VersionPlusOne,
 		LatestVersion = VersionPlusOne - 1
diff --git a/Engine/Plugins/Experimental/Paper2D/Source/Paper2D/Private/PaperFlipbookComponent.cpp b/Engine/Plugins/Experimental/Paper2D/Source/Paper2D/Private/PaperFlipbookComponent.cpp
--- a/Engine/Plugins/Experimental/Paper2D/Source/Paper2D/Private/PaperFlipbookComponent.cpp
+++ b/Engine/Plugins/Experimental/Paper2D/Source/Paper2D/Private/PaperFlipbookComponent.cpp
@@ -3,7 +3,9 @@
 #include "Paper2DPrivatePCH.h"
 #include "PaperFlipbookSceneProxy.h"
 #include "PaperFlipbookComponent.h"
+#include "PaperCustomVersion.h"
 #include "Runtime/Engine/Public/Net/UnrealNetwork.h"
+#include "Runtime/Engine/Public/ContentStreaming.h"
 
 //////////////////////////////////////////////////////////////////////////
 // UPaperFlipbookComponent
@@ -14,7 +16,11 @@ UPaperFlipbookComponent::UPaperFlipbookComponent(const FPostConstructInitializeP
 	BodyInstance.SetCollisionEnabled(ECollisionEnabled::NoCollision);
 
 	SetCollisionProfileName(UCollisionProfile::NoCollision_ProfileName);
-	Material = nullptr;
+	Material_DEPRECATED = nullptr;
+
+	CastShadow = false;
+	bUseAsOccluder = false;
+	bCanEverAffectNavigation = false;
 
 	SpriteColor = FLinearColor::White;
 
@@ -42,6 +48,31 @@ UPaperSprite* UPaperFlipbookComponent::GetSpriteAtCachedIndex() const
 	return SpriteToSend;
 }
 
+
+#if WITH_EDITORONLY_DATA
+void UPaperFlipbookComponent::Serialize(FArchive& Ar)
+{
+	Super::Serialize(Ar);
+
+	Ar.UsingCustomVersion(FPaperCustomVersion::GUID);
+}
+
+void UPaperFlipbookComponent::PostLoad()
+{
+	Super::PostLoad();
+
+	const int32 PaperVer = GetLinkerCustomVersion(FPaperCustomVersion::GUID);
+
+	if (PaperVer < FPaperCustomVersion::ConvertPaperFlipbookComponentToBeMeshComponent)
+	{
+		if (Material_DEPRECATED != nullptr)
+		{
+			SetMaterial(0, Material_DEPRECATED);
+		}
+	}
+}
+#endif
+
 FPrimitiveSceneProxy* UPaperFlipbookComponent::CreateSceneProxy()
 {
 	FPaperFlipbookSceneProxy* NewProxy = new FPaperFlipbookSceneProxy(this);
@@ -88,6 +119,58 @@ FBoxSphereBounds UPaperFlipbookComponent::CalcBounds(const FTransform & LocalToW
 	}
 }
 
+void UPaperFlipbookComponent::GetUsedTextures(TArray<UTexture*>& OutTextures, EMaterialQualityLevel::Type QualityLevel)
+{
+	// Get any textures referenced by our materials
+	Super::GetUsedTextures(OutTextures, QualityLevel);
+
+	// Get the texture referenced by each keyframe
+	if (SourceFlipbook != nullptr)
+	{
+		for (int32 Index = 0; Index < SourceFlipbook->GetNumKeyFrames(); ++Index)
+		{
+			const FPaperFlipbookKeyFrame& KeyFrame = SourceFlipbook->GetKeyFrameChecked(Index);
+			if (KeyFrame.Sprite != nullptr)
+			{
+				if (UTexture* BakedTexture = KeyFrame.Sprite->GetBakedTexture())
+				{
+					OutTextures.AddUnique(BakedTexture);
+				}
+			}
+		}
+	}
+}
+
+UMaterialInterface* UPaperFlipbookComponent::GetMaterial(int32 MaterialIndex) const
+{
+	if (Materials.IsValidIndex(MaterialIndex) && (Materials[MaterialIndex] != nullptr))
+	{
+		return Materials[MaterialIndex];
+	}
+	else if (SourceFlipbook != nullptr)
+	{
+		return SourceFlipbook->GetDefaultMaterial();
+	}
+
+	return nullptr;
+}
+
+void UPaperFlipbookComponent::GetUsedMaterials(TArray<UMaterialInterface*>& OutMaterials) const
+{
+	return Super::GetUsedMaterials(OutMaterials);
+}
+
+void UPaperFlipbookComponent::GetStreamingTextureInfo(TArray<FStreamingTexturePrimitiveInfo>& OutStreamingTextures) const
+{
+	//@TODO: PAPER2D: Need to support this for proper texture streaming
+	return Super::GetStreamingTextureInfo(OutStreamingTextures);
+}
+
+int32 UPaperFlipbookComponent::GetNumMaterials() const
+{
+	return FMath::Max<int32>(Materials.Num(), 1);
+}
+
 void UPaperFlipbookComponent::CalculateCurrentFrame()
 {
 	const int32 LastCachedFrame = CachedFrameIndex;
@@ -257,6 +340,10 @@ bool UPaperFlipbookComponent::SetFlipbook(class UPaperFlipbook* NewFlipbook)
 			// Update physics representation right away
 			RecreatePhysicsState();
 
+			// Notify the streaming system. Don't use Update(), because this may be the first time the mesh has been set
+			// and the component may have to be added to the streaming system for the first time.
+			IStreamingManager::Get().NotifyPrimitiveAttached(this, DPT_Spawned);
+
 			// Since we have new mesh, we need to update bounds
 			UpdateBounds();
 
@@ -274,18 +361,7 @@ UPaperFlipbook* UPaperFlipbookComponent::GetFlipbook()
 
 UMaterialInterface* UPaperFlipbookComponent::GetSpriteMaterial() const
 {
-	if (Material != nullptr)
-	{
-		return Material;
-	}
-	else if (SourceFlipbook != nullptr)
-	{
-		return SourceFlipbook->GetDefaultMaterial();
-	}
-	else
-	{
-		return nullptr;
-	}
+	return GetMaterial(0);
 }
 
 void UPaperFlipbookComponent::SetSpriteColor(FLinearColor NewColor)
diff --git a/Engine/Plugins/Experimental/Paper2D/Source/Paper2D/Private/PaperFlipbookSceneProxy.cpp b/Engine/Plugins/Experimental/Paper2D/Source/Paper2D/Private/PaperFlipbookSceneProxy.cpp
--- a/Engine/Plugins/Experimental/Paper2D/Source/Paper2D/Private/PaperFlipbookSceneProxy.cpp
+++ b/Engine/Plugins/Experimental/Paper2D/Source/Paper2D/Private/PaperFlipbookSceneProxy.cpp
@@ -11,10 +11,6 @@ FPaperFlipbookSceneProxy::FPaperFlipbookSceneProxy(const UPaperFlipbookComponent
 {
 	//@TODO: PAPER2D: WireframeColor = RenderComp->GetWireframeColor();
 
-	Material = InComponent->GetSpriteMaterial();
-
-	if (Material)
-	{
-		MaterialRelevance = Material->GetRelevance();
-	}
+	Material = InComponent->GetMaterial(0);
+	MaterialRelevance = InComponent->GetMaterialRelevance();
 }
