Paper2D: Add support for a negative play rate on flipbook components

diff --git a/Engine/Plugins/Experimental/Paper2D/Source/Paper2D/Private/PaperFlipbookComponent.cpp b/Engine/Plugins/Experimental/Paper2D/Source/Paper2D/Private/PaperFlipbookComponent.cpp
--- a/Engine/Plugins/Experimental/Paper2D/Source/Paper2D/Private/PaperFlipbookComponent.cpp
+++ b/Engine/Plugins/Experimental/Paper2D/Source/Paper2D/Private/PaperFlipbookComponent.cpp
@@ -192,13 +192,13 @@ void UPaperFlipbookComponent::TickFlipbook(float DeltaTime)
 
 	if (bPlaying)
 	{
-		float NewPosition = AccumulatedTime;
 		const float TimelineLength = GetFlipbookLength();
+		const float EffectiveDeltaTime = DeltaTime * (bReversePlayback ? (-PlayRate) : (PlayRate));
 
-		if (!bReversePlayback)
-		{
-			NewPosition = AccumulatedTime + (DeltaTime * PlayRate);
+		float NewPosition = AccumulatedTime + EffectiveDeltaTime;
 
+		if (EffectiveDeltaTime > 0.0f)
+		{
 			if (NewPosition > TimelineLength)
 			{
 				if (bLooping)
@@ -230,8 +230,6 @@ void UPaperFlipbookComponent::TickFlipbook(float DeltaTime)
 		}
 		else
 		{
-			NewPosition = AccumulatedTime - (DeltaTime * PlayRate);
-
 			if (NewPosition < 0.0f)
 			{
 				if (bLooping)
