Paper2D: Add support for choosing project-wide default materials that will be applied to newly created sprites
- Combine various InitSprite variants into one function that takes an init parameters object

diff --git a/Engine/Plugins/Experimental/Paper2D/Source/Paper2D/Classes/PaperRuntimeSettings.h b/Engine/Plugins/Experimental/Paper2D/Source/Paper2D/Classes/PaperRuntimeSettings.h
--- a/Engine/Plugins/Experimental/Paper2D/Source/Paper2D/Classes/PaperRuntimeSettings.h
+++ b/Engine/Plugins/Experimental/Paper2D/Source/Paper2D/Classes/PaperRuntimeSettings.h
@@ -16,6 +16,16 @@ class PAPER2D_API UPaperRuntimeSettings : public UObject
 	UPROPERTY(GlobalConfig, EditAnywhere, Category=Settings)
 	float DefaultPixelsPerUnrealUnit;
 
+#if WITH_EDITORONLY_DATA
+	// The default translucent or masked material for newly created sprites
+	UPROPERTY(Config, EditAnywhere, Category=Settings, meta=(AllowedClasses="MaterialInterface", DisplayName="Default Translucent Sprite Material"))
+	FStringAssetReference DefaultTranslucentSpriteMaterialName;
+
+	// The default opaque material for newly created sprites
+	UPROPERTY(Config, EditAnywhere, Category=Settings, meta=(AllowedClasses="MaterialInterface", DisplayName="Default Opaque Sprite Material"))
+	FStringAssetReference DefaultOpaqueSpriteMaterialName;
+#endif
+
 	// Enables experimental *incomplete and unsupported* texture atlas groups that sprites can be assigned to
 	UPROPERTY(EditAnywhere, config, Category=Experimental)
 	bool bEnableSpriteAtlasGroups;
diff --git a/Engine/Plugins/Experimental/Paper2D/Source/Paper2D/Classes/PaperSprite.h b/Engine/Plugins/Experimental/Paper2D/Source/Paper2D/Classes/PaperSprite.h
--- a/Engine/Plugins/Experimental/Paper2D/Source/Paper2D/Classes/PaperSprite.h
+++ b/Engine/Plugins/Experimental/Paper2D/Source/Paper2D/Classes/PaperSprite.h
@@ -199,10 +199,7 @@ public:
 	void Triangulate(const FSpritePolygonCollection& Source, TArray<FVector2D>& Target);
 
 	// Reinitializes this sprite (NOTE: Does not register existing components in the world)
-	void InitializeSprite(UTexture2D* Texture, float InPixelsPerUnrealUnit);
-	void InitializeSprite(UTexture2D* Texture);
-	void InitializeSprite(UTexture2D* Texture, const FVector2D& Offset, const FVector2D& Dimension);
-	void InitializeSprite(UTexture2D* Texture, const FVector2D& Offset, const FVector2D& Dimension, float InPixelsPerUnit);
+	void InitializeSprite(const FSpriteAssetInitParameters& InitParams);
 
 	void SetTrim(bool bTrimmed, const FVector2D& OriginInSourceImage, const FVector2D& SourceImageDimension);
 	void SetRotated(bool bRotated);
diff --git a/Engine/Plugins/Experimental/Paper2D/Source/Paper2D/Classes/SpriteEditorOnlyTypes.h b/Engine/Plugins/Experimental/Paper2D/Source/Paper2D/Classes/SpriteEditorOnlyTypes.h
--- a/Engine/Plugins/Experimental/Paper2D/Source/Paper2D/Classes/SpriteEditorOnlyTypes.h
+++ b/Engine/Plugins/Experimental/Paper2D/Source/Paper2D/Classes/SpriteEditorOnlyTypes.h
@@ -139,6 +139,49 @@ struct FSpritePolygonCollection
 	}
 };
 
+USTRUCT()
+struct FSpriteAssetInitParameters
+{
+	GENERATED_USTRUCT_BODY()
+
+	FSpriteAssetInitParameters()
+		: Texture(nullptr)
+		, Offset(FVector2D::ZeroVector)
+		, Dimension(FVector2D::ZeroVector)
+		, bNewlyCreated(false)
+	{
+	}
+
+	// Set the texture and the offset/dimension to fully match the specified texture 
+	void SetTextureAndFill(UTexture2D* InTexture)
+	{
+		Texture = InTexture;
+		if (Texture != nullptr)
+		{
+			Dimension = FVector2D(Texture->GetSizeX(), Texture->GetSizeY());
+			Offset = FVector2D::ZeroVector;
+		}
+		else
+		{
+			Dimension = FVector2D::ZeroVector;
+			Offset = FVector2D::ZeroVector;
+		}
+	}
+
+public:
+	// The texture to use
+	UTexture2D* Texture;
+
+	// The offset within the texture (in pixels)
+	FVector2D Offset;
+
+	// The dimension of the subregion within the texture (in pixels)
+	FVector2D Dimension;
+
+	// Is this sprite newly created (should we pull the default pixels/uu and materials from the project settings)?
+	bool bNewlyCreated;
+};
+
 UENUM()
 namespace ESpritePivotMode
 {
diff --git a/Engine/Plugins/Experimental/Paper2D/Source/Paper2D/Private/PaperRuntimeSettings.cpp b/Engine/Plugins/Experimental/Paper2D/Source/Paper2D/Private/PaperRuntimeSettings.cpp
--- a/Engine/Plugins/Experimental/Paper2D/Source/Paper2D/Private/PaperRuntimeSettings.cpp
+++ b/Engine/Plugins/Experimental/Paper2D/Source/Paper2D/Private/PaperRuntimeSettings.cpp
@@ -13,4 +13,11 @@ UPaperRuntimeSettings::UPaperRuntimeSettings(const FPostConstructInitializePrope
 	, bEnableTileMapEditing(false)
 	, bEnableTerrainSplineEditing(false)
 {
+#if WITH_EDITORONLY_DATA
+	if (!IsRunningCommandlet())
+	{
+		DefaultTranslucentSpriteMaterialName = FStringAssetReference("/Paper2D/MaskedUnlitSpriteMaterial.MaskedUnlitSpriteMaterial");
+		DefaultOpaqueSpriteMaterialName = FStringAssetReference("/Paper2D/OpaqueUnlitSpriteMaterial.OpaqueUnlitSpriteMaterial");
+	}
+#endif
 }
diff --git a/Engine/Plugins/Experimental/Paper2D/Source/Paper2D/Private/PaperSprite.cpp b/Engine/Plugins/Experimental/Paper2D/Source/Paper2D/Private/PaperSprite.cpp
--- a/Engine/Plugins/Experimental/Paper2D/Source/Paper2D/Private/PaperSprite.cpp
+++ b/Engine/Plugins/Experimental/Paper2D/Source/Paper2D/Private/PaperSprite.cpp
@@ -1020,39 +1020,31 @@ void UPaperSprite::ExtractRectsFromTexture(UTexture2D* Texture, TArray<FIntRect>
 	SpriteTextureBitmap.ExtractRects(/*out*/ OutRects);
 }
 
-void UPaperSprite::InitializeSprite(UTexture2D* Texture, float InPixelsPerUnrealUnit)
+void UPaperSprite::InitializeSprite(const FSpriteAssetInitParameters& InitParams)
 {
-	if (Texture != NULL)
+	if (InitParams.bNewlyCreated)
 	{
-		const FVector2D Dimension(Texture->GetSizeX(), Texture->GetSizeY());
-		InitializeSprite(Texture, FVector2D::ZeroVector, Dimension, InPixelsPerUnrealUnit);
-	}
-	else
-	{
-		InitializeSprite(NULL, FVector2D::ZeroVector, FVector2D::ZeroVector, InPixelsPerUnrealUnit);
-	}
-}
+		const UPaperRuntimeSettings* DefaultSettings = GetDefault<UPaperRuntimeSettings>();
+		PixelsPerUnrealUnit = DefaultSettings->DefaultPixelsPerUnrealUnit;
 
-void UPaperSprite::InitializeSprite(UTexture2D* Texture)
-{
-	InitializeSprite(Texture, GetDefault<UPaperRuntimeSettings>()->DefaultPixelsPerUnrealUnit);
-}
+		if (UMaterialInterface* TranslucentMaterial = LoadObject<UMaterialInterface>(nullptr, *DefaultSettings->DefaultTranslucentSpriteMaterialName.ToString(), nullptr, LOAD_None, nullptr))
+		{
+			DefaultMaterial = TranslucentMaterial;
+		}
 
-void UPaperSprite::InitializeSprite(UTexture2D* Texture, const FVector2D& Offset, const FVector2D& Dimension, float InPixelsPerUnrealUnit)
-{
-	PixelsPerUnrealUnit = InPixelsPerUnrealUnit;
-	SourceTexture = Texture;
-	SourceUV = Offset;
-	SourceDimension = Dimension;
+		if (UMaterialInterface* OpaqueMaterial = LoadObject<UMaterialInterface>(nullptr, *DefaultSettings->DefaultOpaqueSpriteMaterialName.ToString(), nullptr, LOAD_None, nullptr))
+		{
+			AlternateMaterial = OpaqueMaterial;
+		}
+	}
+
+	SourceTexture = InitParams.Texture;
+	SourceUV = InitParams.Offset;
+	SourceDimension = InitParams.Dimension;
 	RebuildCollisionData();
 	RebuildRenderData();
 }
 
-void UPaperSprite::InitializeSprite(UTexture2D* Texture, const FVector2D& Offset, const FVector2D& Dimension)
-{
-	InitializeSprite(Texture, Offset, Dimension, GetDefault<UPaperRuntimeSettings>()->DefaultPixelsPerUnrealUnit);
-}
-
 void UPaperSprite::SetTrim(bool bTrimmed, const FVector2D& OriginInSourceImage, const FVector2D& SourceImageDimension)
 {
 	this->bTrimmedInSourceImage = bTrimmed;
diff --git a/Engine/Plugins/Experimental/Paper2D/Source/Paper2DEditor/Private/PaperSpriteFactory.cpp b/Engine/Plugins/Experimental/Paper2D/Source/Paper2DEditor/Private/PaperSpriteFactory.cpp
--- a/Engine/Plugins/Experimental/Paper2D/Source/Paper2DEditor/Private/PaperSpriteFactory.cpp
+++ b/Engine/Plugins/Experimental/Paper2D/Source/Paper2DEditor/Private/PaperSpriteFactory.cpp
@@ -28,14 +28,19 @@ UObject* UPaperSpriteFactory::FactoryCreateNew(UClass* Class, UObject* InParent,
 {
 	UPaperSprite* NewSprite = ConstructObject<UPaperSprite>(Class, InParent, Name, Flags | RF_Transactional);
 
+	FSpriteAssetInitParameters SpriteInitParams;
+	SpriteInitParams.bNewlyCreated = true;
 	if (bUseSourceRegion)
 	{
-		NewSprite->InitializeSprite(InitialTexture, InitialSourceUV, InitialSourceDimension);
+		SpriteInitParams.Texture = InitialTexture;
+		SpriteInitParams.Offset = InitialSourceUV;
+		SpriteInitParams.Dimension = InitialSourceDimension;
 	}
 	else
 	{
-		NewSprite->InitializeSprite(InitialTexture);
+		SpriteInitParams.SetTextureAndFill(InitialTexture);
 	}
+	NewSprite->InitializeSprite(SpriteInitParams);
 
 	return NewSprite;
 }
diff --git a/Engine/Plugins/Experimental/Paper2D/Source/Paper2DEditor/Private/SpriteEditor/SpriteEditorViewportClient.cpp b/Engine/Plugins/Experimental/Paper2D/Source/Paper2DEditor/Private/SpriteEditor/SpriteEditorViewportClient.cpp
--- a/Engine/Plugins/Experimental/Paper2D/Source/Paper2DEditor/Private/SpriteEditor/SpriteEditorViewportClient.cpp
+++ b/Engine/Plugins/Experimental/Paper2D/Source/Paper2DEditor/Private/SpriteEditor/SpriteEditorViewportClient.cpp
@@ -108,7 +108,11 @@ void FSpriteEditorViewportClient::UpdateSourceTextureSpriteFromSprite(UPaperSpri
 		if ((SourceSprite->GetSourceTexture() != TargetSprite->GetSourceTexture()) || (TargetSprite->PixelsPerUnrealUnit != SourceSprite->PixelsPerUnrealUnit))
 		{
 			FComponentReregisterContext ReregisterSprite(SourceTextureViewComponent);
-			TargetSprite->InitializeSprite(SourceSprite->SourceTexture, SourceSprite->PixelsPerUnrealUnit);
+
+			FSpriteAssetInitParameters SpriteReinitParams;
+			SpriteReinitParams.SetTextureAndFill(SourceSprite->SourceTexture);
+			TargetSprite->PixelsPerUnrealUnit = SourceSprite->PixelsPerUnrealUnit;
+			TargetSprite->InitializeSprite(SpriteReinitParams);
 		}
 
 
diff --git a/Engine/Plugins/Experimental/Paper2D/Source/PaperJsonImporter/Private/PaperJsonSpriteSheetImporter.cpp b/Engine/Plugins/Experimental/Paper2D/Source/PaperJsonImporter/Private/PaperJsonSpriteSheetImporter.cpp
--- a/Engine/Plugins/Experimental/Paper2D/Source/PaperJsonImporter/Private/PaperJsonSpriteSheetImporter.cpp
+++ b/Engine/Plugins/Experimental/Paper2D/Source/PaperJsonImporter/Private/PaperJsonSpriteSheetImporter.cpp
@@ -502,7 +502,12 @@ bool FPaperJsonSpriteSheetImporter::PerformImport(const FString& LongPackagePath
 		}
 
 		TargetSprite->Modify();
-		TargetSprite->InitializeSprite(ImageTexture, Frame.SpritePosInSheet, Frame.SpriteSizeInSheet, GetDefault<UPaperRuntimeSettings>()->DefaultPixelsPerUnrealUnit);
+		FSpriteAssetInitParameters SpriteInitParams;
+		SpriteInitParams.Texture = ImageTexture;
+		SpriteInitParams.Offset = Frame.SpritePosInSheet;
+		SpriteInitParams.Dimension = Frame.SpriteSizeInSheet;
+		SpriteInitParams.bNewlyCreated = true;
+		TargetSprite->InitializeSprite(SpriteInitParams);
 
 		if (Frame.bRotated)
 		{
