Paper2D: Add debug rendering for pawn/visibility collision on sprites
- Add support for VMI_CollisionVisibility and VMI_CollisionPawn view modes to FPaperRenderSceneProxy
- Change sprite wireframe colors to match static mesh display (including differentation for mobility)

diff --git a/Engine/Plugins/Experimental/Paper2D/Source/Paper2D/Private/PaperRenderSceneProxy.cpp b/Engine/Plugins/Experimental/Paper2D/Source/Paper2D/Private/PaperRenderSceneProxy.cpp
--- a/Engine/Plugins/Experimental/Paper2D/Source/Paper2D/Private/PaperRenderSceneProxy.cpp
+++ b/Engine/Plugins/Experimental/Paper2D/Source/Paper2D/Private/PaperRenderSceneProxy.cpp
@@ -159,6 +159,7 @@ FPaperRenderSceneProxy::FPaperRenderSceneProxy(const UPrimitiveComponent* InComp
 	: FPrimitiveSceneProxy(InComponent)
 	, Material(NULL)
 	, Owner(InComponent->GetOwner())
+	, BodySetup(const_cast<UPrimitiveComponent*>(InComponent)->GetBodySetup())
 	, bCastShadow(InComponent->CastShadow)
 	, CollisionResponse(InComponent->GetCollisionResponseToChannels())
 {
@@ -191,39 +192,102 @@ void FPaperRenderSceneProxy::GetDynamicMeshElements(const TArray<const FSceneVie
 	QUICK_SCOPE_CYCLE_COUNTER(STAT_PaperRenderSceneProxy_GetDynamicMeshElements);
 	checkSlow(IsInRenderingThread());
 
-	for (int32 ViewIndex = 0; ViewIndex < Views.Num(); ViewIndex++)
+	const FEngineShowFlags& EngineShowFlags = ViewFamily.EngineShowFlags;
+
+	bool bDrawSimpleCollision = false;
+	bool bDrawComplexCollision = false;
+	const bool bInCollisionView = IsCollisionView(EngineShowFlags, bDrawSimpleCollision, bDrawComplexCollision);
+	
+	// Sprites don't distinguish between simple and complex collision; when viewing visibility we should still render simple collision
+	bDrawSimpleCollision |= bDrawComplexCollision;
+
+	// Draw simple collision as wireframe if 'show collision', collision is enabled
+	const bool bDrawWireframeCollision = EngineShowFlags.Collision && IsCollisionEnabled();
+
+	const bool bDrawSprite = !bInCollisionView;
+
+	if (bDrawSprite)
 	{
-		if (VisibilityMap & (1 << ViewIndex))
+		for (int32 ViewIndex = 0; ViewIndex < Views.Num(); ViewIndex++)
 		{
-			const FSceneView* View = Views[ViewIndex];
-	
-#if !(UE_BUILD_SHIPPING || UE_BUILD_TEST)
-			if (View->Family->EngineShowFlags.Paper2DSprites)
+			if (VisibilityMap & (1 << ViewIndex))
 			{
-				RenderBounds(Collector.GetPDI(ViewIndex), ViewFamily.EngineShowFlags, GetBounds(), (Owner == NULL) || IsSelected());
-			}
-#endif
+				const FSceneView* View = Views[ViewIndex];
 
-			// Set the selection/hover color from the current engine setting.
-			FLinearColor OverrideColor;
-			bool bUseOverrideColor = false;
+				// Set the selection/hover color from the current engine setting.
+				FLinearColor OverrideColor;
+				bool bUseOverrideColor = false;
 
-			const bool bShowAsSelected = !(GIsEditor && ViewFamily.EngineShowFlags.Selection) || IsSelected();
-			if (bShowAsSelected || IsHovered())
-			{
-				bUseOverrideColor = true;
-				OverrideColor = GetSelectionColor(FLinearColor::White, bShowAsSelected, IsHovered());
-			}
+				const bool bShowAsSelected = !(GIsEditor && EngineShowFlags.Selection) || IsSelected();
+				if (bShowAsSelected || IsHovered())
+				{
+					bUseOverrideColor = true;
+					OverrideColor = GetSelectionColor(FLinearColor::White, bShowAsSelected, IsHovered());
+				}
 
-			bUseOverrideColor = false;
+				bUseOverrideColor = false;
 
-			// Sprites of locked actors draw in red.
-			//FLinearColor LevelColorToUse = IsSelected() ? ColorToUse : (FLinearColor)LevelColor;
-			//FLinearColor PropertyColorToUse = PropertyColor;
+				// Sprites of locked actors draw in red.
+				//FLinearColor LevelColorToUse = IsSelected() ? ColorToUse : (FLinearColor)LevelColor;
+				//FLinearColor PropertyColorToUse = PropertyColor;
 
 #if TEST_BATCHING
 #else
-			GetDynamicMeshElementsForView(View, ViewIndex, bUseOverrideColor, OverrideColor, Collector);
+				GetDynamicMeshElementsForView(View, ViewIndex, bUseOverrideColor, OverrideColor, Collector);
+#endif
+			}
+		}
+	}
+
+
+	for (int32 ViewIndex = 0; ViewIndex < Views.Num(); ViewIndex++)
+	{
+		if (VisibilityMap & (1 << ViewIndex))
+		{
+			if ((bDrawSimpleCollision || bDrawWireframeCollision) && AllowDebugViewmodes())
+			{
+				if (BodySetup)
+				{
+					if (FMath::Abs(GetLocalToWorld().Determinant()) < SMALL_NUMBER)
+					{
+						// Catch this here or otherwise GeomTransform below will assert
+						// This spams so commented out
+						//UE_LOG(LogStaticMesh, Log, TEXT("Zero scaling not supported (%s)"), *StaticMesh->GetPathName());
+					}
+					else
+					{
+						const bool bDrawSolid = !bDrawWireframeCollision;
+
+						if (bDrawSolid)
+						{
+							// Make a material for drawing solid collision stuff
+							auto SolidMaterialInstance = new FColoredMaterialRenderProxy(
+								GEngine->ShadedLevelColorationUnlitMaterial->GetRenderProxy(IsSelected(), IsHovered()),
+								WireframeColor
+								);
+
+							Collector.RegisterOneFrameMaterialProxy(SolidMaterialInstance);
+
+							FTransform GeomTransform(GetLocalToWorld());
+							BodySetup->AggGeom.GetAggGeom(GeomTransform, WireframeColor, SolidMaterialInstance, false, true, UseEditorDepthTest(), ViewIndex, Collector);
+						}
+						else
+						{
+							// wireframe
+							FColor CollisionColor = FColor(157, 149, 223, 255);
+							FTransform GeomTransform(GetLocalToWorld());
+							BodySetup->AggGeom.GetAggGeom(GeomTransform, GetSelectionColor(CollisionColor, IsSelected(), IsHovered()), NULL, (Owner == NULL), false, UseEditorDepthTest(), ViewIndex, Collector);
+						}
+					}
+				}
+			}
+
+			// Draw bounds
+#if !(UE_BUILD_SHIPPING || UE_BUILD_TEST)
+			if (EngineShowFlags.Paper2DSprites)
+			{
+				RenderBounds(Collector.GetPDI(ViewIndex), EngineShowFlags, GetBounds(), (Owner == NULL) || IsSelected());
+			}
 #endif
 		}
 	}
@@ -235,34 +299,89 @@ void FPaperRenderSceneProxy::DrawDynamicElements(FPrimitiveDrawInterface* PDI, c
 
 	checkSlow(IsInRenderingThread());
 
+
+	const FEngineShowFlags& EngineShowFlags = View->Family->EngineShowFlags;
+
+	bool bDrawSimpleCollision = false;
+	bool bDrawComplexCollision = false;
+	const bool bInCollisionView = IsCollisionView(EngineShowFlags, bDrawSimpleCollision, bDrawComplexCollision);
+
+	// Sprites don't distinguish between simple and complex collision; when viewing visibility we should still render simple collision
+	bDrawSimpleCollision |= bDrawComplexCollision;
+	
+	// Draw simple collision as wireframe if 'show collision', collision is enabled
+	const bool bDrawWireframeCollision = EngineShowFlags.Collision && IsCollisionEnabled();
+
+	if ((bDrawSimpleCollision || bDrawWireframeCollision) && AllowDebugViewmodes())
+	{
+		if (BodySetup)
+		{
+			if (FMath::Abs(GetLocalToWorld().Determinant()) < SMALL_NUMBER)
+			{
+				// Catch this here or otherwise GeomTransform below will assert
+				// This spams so commented out
+				//UE_LOG(LogStaticMesh, Log, TEXT("Zero scaling not supported (%s)"), *StaticMesh->GetPathName());
+			}
+			else
+			{
+				const bool bDrawSolid = !bDrawWireframeCollision;
+
+				if (bDrawSolid)
+				{
+					// Make a material for drawing solid collision stuff
+					const FColoredMaterialRenderProxy SolidMaterialInstance(
+						GEngine->ShadedLevelColorationUnlitMaterial->GetRenderProxy(IsSelected(), IsHovered()),
+						WireframeColor
+						);
+
+					FTransform GeomTransform(GetLocalToWorld());
+					BodySetup->AggGeom.DrawAggGeom(PDI, GeomTransform, WireframeColor, &SolidMaterialInstance, false, true, UseEditorDepthTest());
+				}
+				else
+				{
+					// wireframe
+					FColor CollisionColor = FColor(157, 149, 223, 255);
+					FTransform GeomTransform(GetLocalToWorld());
+					BodySetup->AggGeom.DrawAggGeom(PDI, GeomTransform, GetSelectionColor(CollisionColor, IsSelected(), IsHovered()), NULL, (Owner == NULL), false, UseEditorDepthTest());
+				}
+			}
+		}
+	}
+
+	// Draw bounds
 #if !(UE_BUILD_SHIPPING || UE_BUILD_TEST)
-	if (View->Family->EngineShowFlags.Paper2DSprites)
+	if (EngineShowFlags.Paper2DSprites)
 	{
-		RenderBounds(PDI, View->Family->EngineShowFlags, GetBounds(), (Owner == NULL) || IsSelected());
+		RenderBounds(PDI, EngineShowFlags, GetBounds(), (Owner == NULL) || IsSelected());
 	}
 #endif
 
-	// Set the selection/hover color from the current engine setting.
-	FLinearColor OverrideColor;
-	bool bUseOverrideColor = false;
+	const bool bDrawSprite = !bInCollisionView;
 
-	const bool bShowAsSelected = !(GIsEditor && View->Family->EngineShowFlags.Selection) || IsSelected();
-	if (bShowAsSelected || IsHovered())
+	if (bDrawSprite)
 	{
-		bUseOverrideColor = true;
-		OverrideColor = GetSelectionColor(FLinearColor::White, bShowAsSelected, IsHovered());
-	}
+		// Set the selection/hover color from the current engine setting.
+		FLinearColor OverrideColor;
+		bool bUseOverrideColor = false;
+
+		const bool bShowAsSelected = !(GIsEditor && View->Family->EngineShowFlags.Selection) || IsSelected();
+		if (bShowAsSelected || IsHovered())
+		{
+			bUseOverrideColor = true;
+			OverrideColor = GetSelectionColor(FLinearColor::White, bShowAsSelected, IsHovered());
+		}
 
-	bUseOverrideColor = false;
+		bUseOverrideColor = false;
 
-	// Sprites of locked actors draw in red.
-	//FLinearColor LevelColorToUse = IsSelected() ? ColorToUse : (FLinearColor)LevelColor;
-	//FLinearColor PropertyColorToUse = PropertyColor;
+		// Sprites of locked actors draw in red.
+		//FLinearColor LevelColorToUse = IsSelected() ? ColorToUse : (FLinearColor)LevelColor;
+		//FLinearColor PropertyColorToUse = PropertyColor;
 
 #if TEST_BATCHING
 #else
-	DrawDynamicElements_RichMesh(PDI, View, bUseOverrideColor, OverrideColor);
+		DrawDynamicElements_RichMesh(PDI, View, bUseOverrideColor, OverrideColor);
 #endif
+	}
 }
 
 FVertexFactory* FPaperRenderSceneProxy::GetPaperSpriteVertexFactory() const
@@ -612,10 +731,12 @@ void FPaperRenderSceneProxy::DrawBatch(FPrimitiveDrawInterface* PDI, const FScen
 
 FPrimitiveViewRelevance FPaperRenderSceneProxy::GetViewRelevance(const FSceneView* View)
 {
+	const FEngineShowFlags& EngineShowFlags = View->Family->EngineShowFlags;
+
 	checkSlow(IsInRenderingThread());
 
 	FPrimitiveViewRelevance Result;
-	Result.bDrawRelevance = IsShown(View) && View->Family->EngineShowFlags.Paper2DSprites;
+	Result.bDrawRelevance = IsShown(View) && EngineShowFlags.Paper2DSprites;
 	Result.bRenderCustomDepth = ShouldRenderCustomDepth();
 	Result.bRenderInMainPass = ShouldRenderInMainPass();
 
@@ -630,7 +751,7 @@ FPrimitiveViewRelevance FPaperRenderSceneProxy::GetViewRelevance(const FSceneVie
 #if SUPPORT_EXTRA_RENDERING
 	bool bDrawSimpleCollision = false;
 	bool bDrawComplexCollision = false;
-	const bool bInCollisionView = IsCollisionView(View, bDrawSimpleCollision, bDrawComplexCollision);
+	const bool bInCollisionView = IsCollisionView(EngineShowFlags, bDrawSimpleCollision, bDrawComplexCollision);
 #endif
 
 	Result.bDynamicRelevance = true;
@@ -638,9 +759,9 @@ FPrimitiveViewRelevance FPaperRenderSceneProxy::GetViewRelevance(const FSceneVie
 // 	if (
 // #if SUPPORT_EXTRA_RENDERING
 // 		IsRichView(View) ||
-// 		View->Family->EngineShowFlags.Collision ||
+// 		EngineShowFlags.Collision ||
 // 		bInCollisionView ||
-// 		View->Family->EngineShowFlags.Bounds ||
+// 		EngineShowFlags.Bounds ||
 // #endif
 // 		// Force down dynamic rendering path if invalid lightmap settings, so we can apply an error material in DrawRichMesh
 // 		(HasStaticLighting() && !HasValidSettingsForStaticLighting()) ||
@@ -658,7 +779,7 @@ FPrimitiveViewRelevance FPaperRenderSceneProxy::GetViewRelevance(const FSceneVie
 // 		Result.bStaticRelevance = true;
 // 	}
 
-	if (!View->Family->EngineShowFlags.Materials
+	if (!EngineShowFlags.Materials
 #if SUPPORT_EXTRA_RENDERING
 		|| bInCollisionView
 #endif
@@ -693,23 +814,23 @@ void FPaperRenderSceneProxy::SetDrawCall_RenderThread(const FSpriteDrawCallRecor
 	Record = NewDynamicData;
 }
 
-bool FPaperRenderSceneProxy::IsCollisionView(const FSceneView* View, bool& bDrawSimpleCollision, bool& bDrawComplexCollision) const
+bool FPaperRenderSceneProxy::IsCollisionView(const FEngineShowFlags& EngineShowFlags, bool& bDrawSimpleCollision, bool& bDrawComplexCollision) const
 {
 	bDrawSimpleCollision = false;
 	bDrawComplexCollision = false;
 
 	// If in a 'collision view' and collision is enabled
-	const bool bInCollisionView = View->Family->EngineShowFlags.CollisionVisibility || View->Family->EngineShowFlags.CollisionPawn;
+	const bool bInCollisionView = EngineShowFlags.CollisionVisibility || EngineShowFlags.CollisionPawn;
 	if (bInCollisionView && IsCollisionEnabled())
 	{
 		// See if we have a response to the interested channel
-		bool bHasResponse = View->Family->EngineShowFlags.CollisionPawn && (CollisionResponse.GetResponse(ECC_Pawn) != ECR_Ignore);
-		bHasResponse |= View->Family->EngineShowFlags.CollisionVisibility && (CollisionResponse.GetResponse(ECC_Visibility) != ECR_Ignore);
+		bool bHasResponse = EngineShowFlags.CollisionPawn && (CollisionResponse.GetResponse(ECC_Pawn) != ECR_Ignore);
+		bHasResponse |= EngineShowFlags.CollisionVisibility && (CollisionResponse.GetResponse(ECC_Visibility) != ECR_Ignore);
 
 		if (bHasResponse)
 		{
-			bDrawComplexCollision = View->Family->EngineShowFlags.CollisionVisibility;
-			bDrawSimpleCollision = View->Family->EngineShowFlags.CollisionPawn;
+			bDrawComplexCollision = EngineShowFlags.CollisionVisibility;
+			bDrawSimpleCollision = EngineShowFlags.CollisionPawn;
 		}
 	}
 
diff --git a/Engine/Plugins/Experimental/Paper2D/Source/Paper2D/Private/PaperRenderSceneProxy.h b/Engine/Plugins/Experimental/Paper2D/Source/Paper2D/Private/PaperRenderSceneProxy.h
--- a/Engine/Plugins/Experimental/Paper2D/Source/Paper2D/Private/PaperRenderSceneProxy.h
+++ b/Engine/Plugins/Experimental/Paper2D/Source/Paper2D/Private/PaperRenderSceneProxy.h
@@ -115,7 +115,7 @@ protected:
 	void GetBatchMesh(const FSceneView* View, bool bUseOverrideColor, const FLinearColor& OverrideColor, class UMaterialInterface* BatchMaterial, const TArray<FSpriteDrawCallRecord>& Batch, int32 ViewIndex, FMeshElementCollector& Collector) const;
 	void GetNewBatchMeshes(const FSceneView* View, bool bUseOverrideColor, const FLinearColor& OverrideColor, int32 ViewIndex, FMeshElementCollector& Collector) const;
 
-	bool IsCollisionView(const FSceneView* View, bool& bDrawSimpleCollision, bool& bDrawComplexCollision) const;
+	bool IsCollisionView(const FEngineShowFlags& EngineShowFlags, bool& bDrawSimpleCollision, bool& bDrawComplexCollision) const;
 
 	friend class FPaperBatchSceneProxy;
 
@@ -129,6 +129,7 @@ protected:
 	class UMaterialInterface* Material;
 	FVector Origin;
 	AActor* Owner;
+	UBodySetup* BodySetup;
 
 	bool bCastShadow;
 
diff --git a/Engine/Plugins/Experimental/Paper2D/Source/Paper2D/Private/PaperSpriteComponent.cpp b/Engine/Plugins/Experimental/Paper2D/Source/Paper2D/Private/PaperSpriteComponent.cpp
--- a/Engine/Plugins/Experimental/Paper2D/Source/Paper2D/Private/PaperSpriteComponent.cpp
+++ b/Engine/Plugins/Experimental/Paper2D/Source/Paper2D/Private/PaperSpriteComponent.cpp
@@ -276,7 +276,21 @@ void UPaperSpriteComponent::SetSpriteColor(FLinearColor NewColor)
 
 FLinearColor UPaperSpriteComponent::GetWireframeColor() const
 {
-	return FLinearColor::Yellow;
+	if (Mobility == EComponentMobility::Static)
+	{
+		return FColor(0, 255, 255, 255);
+	}
+	else
+	{
+		if (BodyInstance.bSimulatePhysics)
+		{
+			return FColor(0, 255, 128, 255);
+		}
+		else
+		{
+			return FColor(255, 0, 255, 255);
+		}
+	}
 }
 
 const UObject* UPaperSpriteComponent::AdditionalStatObject() const
