Paper2D: Expose RenderComponent to Blueprints in APaperFlipbookActor and APaperSpriteActor, and add 'pass-thru' metadata so various functions can be called on the component from an actor reference

diff --git a/Engine/Plugins/Experimental/Paper2D/Source/Paper2D/Classes/PaperFlipbookActor.h b/Engine/Plugins/Experimental/Paper2D/Source/Paper2D/Classes/PaperFlipbookActor.h
--- a/Engine/Plugins/Experimental/Paper2D/Source/Paper2D/Classes/PaperFlipbookActor.h
+++ b/Engine/Plugins/Experimental/Paper2D/Source/Paper2D/Classes/PaperFlipbookActor.h
@@ -10,7 +10,7 @@ class APaperFlipbookActor : public AActor
 {
 	GENERATED_UCLASS_BODY()
 
-	UPROPERTY(Category=Sprite, VisibleAnywhere)
+	UPROPERTY(Category=Sprite, VisibleAnywhere, BlueprintReadOnly, meta=(ExposeFunctionCategories="Sprite,Rendering,Physics,Components|Flipbook"))
 	TSubobjectPtr<class UPaperFlipbookComponent> RenderComponent;
 
 	// AActor interface
diff --git a/Engine/Plugins/Experimental/Paper2D/Source/Paper2D/Classes/PaperSpriteActor.h b/Engine/Plugins/Experimental/Paper2D/Source/Paper2D/Classes/PaperSpriteActor.h
--- a/Engine/Plugins/Experimental/Paper2D/Source/Paper2D/Classes/PaperSpriteActor.h
+++ b/Engine/Plugins/Experimental/Paper2D/Source/Paper2D/Classes/PaperSpriteActor.h
@@ -10,7 +10,7 @@ class APaperSpriteActor : public AActor
 {
 	GENERATED_UCLASS_BODY()
 
-	UPROPERTY(Category=Sprite, VisibleAnywhere)
+	UPROPERTY(Category=Sprite, VisibleAnywhere, BlueprintReadOnly, meta=(ExposeFunctionCategories="Sprite,Rendering,Physics,Components|Sprite"))
 	TSubobjectPtr<class UPaperSpriteComponent> RenderComponent;
 
 	// AActor interface
