Paper2D: Flipbook import in SpriteSheet & Sprite asset type actions share code now, behave exactly the same.
- Its now possible to create large numbers of flipbooks almost automatically by multi selecting sprites and "Create flipbooks"

diff --git a/Engine/Plugins/Experimental/Paper2D/Source/Paper2DEditor/Private/PaperFlipbookHelpers.cpp b/Engine/Plugins/Experimental/Paper2D/Source/Paper2DEditor/Private/PaperFlipbookHelpers.cpp
new file mode 100644
--- /dev/null
+++ b/Engine/Plugins/Experimental/Paper2D/Source/Paper2DEditor/Private/PaperFlipbookHelpers.cpp
@@ -0,0 +1,170 @@
+// Copyright 1998-2014 Epic Games, Inc. All Rights Reserved.
+
+#include "Paper2DEditorPrivatePCH.h"
+#include "PaperFlipbookHelpers.h"
+#include "Json.h"
+
+//////////////////////////////////////////////////////////////////////////
+// FPaperFlipbookHelpers
+
+static FString CleanName(const FString& Name)
+{
+	int LastCharacter = Name.Len() - 1;
+
+	// Strip off "_Sprite"
+	if (Name.Right(7) == TEXT("_Sprite"))
+	{
+		LastCharacter = Name.Len() - 7;
+	}
+
+	while (LastCharacter >= 0 && FChar::IsPunct(Name[LastCharacter]))
+	{
+		LastCharacter--;
+	}
+
+	if (LastCharacter > 0)
+	{
+		return Name.Left(LastCharacter + 1);
+	}
+	else
+	{
+		return Name;
+	}
+}
+
+static bool ExtractSpriteNumber(const FString& String, FString& BareString, int& Number)
+{
+	bool bExtracted = false;
+
+	int LastCharacter = String.Len() - 1;
+	if (LastCharacter >= 0)
+	{
+		// Find the last character that isn't a digit (Handle sprite names with numbers inside inverted commas / parentheses)
+		while (LastCharacter >= 0 && !FChar::IsDigit(String[LastCharacter]))
+		{
+			LastCharacter--;
+		}
+
+		// Only proceed if we found a number in the sprite name
+		if (LastCharacter >= 0)
+		{
+			while (LastCharacter > 0 && FChar::IsDigit(String[LastCharacter - 1]))
+			{
+				LastCharacter--;
+			}
+
+			if (LastCharacter >= 0)
+			{
+				int FirstDigit = LastCharacter;
+				int EndCharacter = FirstDigit;
+				while (EndCharacter > 0 && !FChar::IsAlnum(String[EndCharacter - 1]))
+				{
+					EndCharacter--;
+				}
+
+				if (EndCharacter == 0)
+				{
+					// This string consists of non alnum + number, eg. _42
+					// The flipbook / category name in this case will be _
+					// Otherwise, we strip out all trailing non-alnum chars
+					EndCharacter = FirstDigit;
+				}
+
+				FString NumberString = String.Mid(FirstDigit);
+				BareString =  (EndCharacter > 0) ? CleanName(String.Left(EndCharacter)) : CleanName(String);
+				Number = FCString::Atoi(*NumberString);
+				bExtracted = true;
+			}
+		}
+	}
+
+	if (!bExtracted)
+	{
+		BareString = String;
+		Number = -1;
+	}
+	return bExtracted;
+}
+
+void FPaperFlipbookHelpers::ExtractFlipbooksFromSprites(TMap<FString, TArray<UPaperSprite*> >& OutSpriteFlipbookMap, const TArray<UPaperSprite*>& Sprites, const TArray<FString>& InSpriteNames)
+{
+	OutSpriteFlipbookMap.Reset();
+
+	// Local copy
+	check((InSpriteNames.Num() == 0) || (InSpriteNames.Num() == Sprites.Num()));
+	TArray<FString> SpriteNames = InSpriteNames;
+	if (InSpriteNames.Num() == 0)
+	{
+		SpriteNames.Reset();
+		for (int32 SpriteIndex = 0; SpriteIndex < Sprites.Num(); ++SpriteIndex)
+		{
+			check(Sprites[SpriteIndex] != nullptr);
+			SpriteNames.Add(Sprites[SpriteIndex]->GetName());
+		}
+	}
+
+	// Group them
+	TMap<FString, UPaperSprite*> SpriteNameMap;
+	TArray<UPaperSprite*> RemainingSprites;
+
+	for (int32 SpriteIndex = 0; SpriteIndex < Sprites.Num(); ++SpriteIndex)
+	{
+		UPaperSprite* Sprite = Sprites[SpriteIndex];
+		const FString SpriteName = SpriteNames[SpriteIndex];
+
+		SpriteNameMap.Add(SpriteName, Sprite);
+
+		int SpriteNumber = 0;
+		FString SpriteBareString;
+		if (ExtractSpriteNumber(SpriteName, /*out*/SpriteBareString, /*out*/SpriteNumber))
+		{
+			if (!OutSpriteFlipbookMap.Contains(SpriteBareString))
+			{
+				OutSpriteFlipbookMap.Add(SpriteBareString, TArray<UPaperSprite*>());
+			}
+
+			OutSpriteFlipbookMap[SpriteBareString].Add(Sprite);
+		}
+		else
+		{
+			RemainingSprites.Add(Sprite);
+		}
+	}
+
+	// Natural sort using the same method as above
+	struct FSpriteSortPredicate
+	{
+		FSpriteSortPredicate() {}
+
+		// Sort predicate operator
+		bool operator()(UPaperSprite& LHS, UPaperSprite& RHS) const
+		{
+			FString LeftString;
+			int LeftNumber;
+			ExtractSpriteNumber(LHS.GetName(), /*out*/LeftString, /*out*/LeftNumber);
+
+			FString RightString;
+			int RightNumber;
+			ExtractSpriteNumber(RHS.GetName(), /*out*/RightString, /*out*/RightNumber);
+
+			return (LeftString == RightString) ? LeftNumber < RightNumber : LeftString < RightString;
+		}
+	};
+
+	// Sort sprites
+	TArray<FString> Keys;
+	OutSpriteFlipbookMap.GetKeys(Keys);
+	for (auto SpriteName : Keys)
+	{
+		TArray<UPaperSprite*>& Sprites = OutSpriteFlipbookMap[SpriteName];
+		Sprites.Sort(FSpriteSortPredicate());
+	}
+
+	// Create a flipbook from all remaining sprites
+	// Not sure if this is desirable behaviour, might want one flipbook per sprite
+	if (RemainingSprites.Num() > 0)
+	{
+		RemainingSprites.Sort(FSpriteSortPredicate());
+		OutSpriteFlipbookMap.Add(CleanName(RemainingSprites[0]->GetName()) + TEXT("_Flipbook"), RemainingSprites);
+	}
+}
diff --git a/Engine/Plugins/Experimental/Paper2D/Source/Paper2DEditor/Private/SpriteAssetTypeActions.cpp b/Engine/Plugins/Experimental/Paper2D/Source/Paper2DEditor/Private/SpriteAssetTypeActions.cpp
--- a/Engine/Plugins/Experimental/Paper2D/Source/Paper2DEditor/Private/SpriteAssetTypeActions.cpp
+++ b/Engine/Plugins/Experimental/Paper2D/Source/Paper2DEditor/Private/SpriteAssetTypeActions.cpp
@@ -6,6 +6,7 @@
 #include "SpriteEditor/SpriteEditor.h"
 #include "AssetToolsModule.h"
 #include "ContentBrowserModule.h"
+#include "PaperFlipbookHelpers.h"
 
 #define LOCTEXT_NAMESPACE "AssetTypeActions"
 
@@ -51,8 +52,8 @@ void FSpriteAssetTypeActions::GetActions(const TArray<UObject*>& InObjects, FMen
 	auto Sprites = GetTypedWeakObjectPtrs<UPaperSprite>(InObjects);
 
 	MenuBuilder.AddMenuEntry(
-		LOCTEXT("Sprite_CreateFlipbook", "Create Flipbook"),
-		LOCTEXT("Sprite_CreateFlipbookTooltip", "Creates a flipbook from the selected sprites."),
+		LOCTEXT("Sprite_CreateFlipbook", "Create Flipbooks"),
+		LOCTEXT("Sprite_CreateFlipbookTooltip", "Creates flipbooks from the selected sprites."),
 		FSlateIcon(),
 		FUIAction(
 		FExecuteAction::CreateSP(this, &FSpriteAssetTypeActions::ExecuteCreateFlipbook, Sprites),
@@ -68,92 +69,82 @@ void FSpriteAssetTypeActions::ExecuteCreateFlipbook(TArray<TWeakObjectPtr<UPaper
 {
 	FAssetToolsModule& AssetToolsModule = FModuleManager::GetModuleChecked<FAssetToolsModule>("AssetTools");
 	FContentBrowserModule& ContentBrowserModule = FModuleManager::LoadModuleChecked<FContentBrowserModule>("ContentBrowser");
-	TArray<UPaperSprite*> Sprites;
+	TArray<UPaperSprite*> AllSprites;
 
 	for (auto ObjIt = Objects.CreateConstIterator(); ObjIt; ++ObjIt)
 	{
 		UPaperSprite *Object = (*ObjIt).Get();
 		if (Object && Object->IsValidLowLevel())
 		{
-			Sprites.Add(Object);
+			AllSprites.Add(Object);
 		}
 	}
 
-	// Attempt to natural sort if possible
-	// Cache results if it becomes a bottleneck
-	struct FSpriteSortPredicate
+	TMap<FString, TArray<UPaperSprite*> > SpriteFlipbookMap;
+	FPaperFlipbookHelpers::ExtractFlipbooksFromSprites(/*out*/SpriteFlipbookMap, AllSprites, TArray<FString>());
+	TArray<UObject*> ObjectsToSync;
+
+	if (SpriteFlipbookMap.Num() > 0)
 	{
-		FSpriteSortPredicate() {}
+		GWarn->BeginSlowTask(NSLOCTEXT("Paper2D", "Paper2D_CreateFlipbooks", "Creating flipbooks from selection"), true, true);
+
+		int Progress = 0;
+		int TotalProgress = SpriteFlipbookMap.Num();
 
-		void ExtractNumber(const FString& String, FString& BareString, int& Number) const
+		// Create the flipbook
+		bool bOneFlipbookCreated = SpriteFlipbookMap.Num() == 1;
+		for (auto Iter : SpriteFlipbookMap)
 		{
-			bool bExtracted = false;
-			int LastCharacter = String.Len() - 1;
-			if (LastCharacter >= 0 && FChar::IsDigit(String[LastCharacter]))
+			GWarn->UpdateProgress(Progress++, TotalProgress);
+
+			const FString& FlipbookName = Iter.Key;
+			TArray<UPaperSprite*> Sprites = Iter.Value;
+
+			const FString SpritePathName = AllSprites[0]->GetOutermost()->GetPathName();
+			const FString LongPackagePath = FPackageName::GetLongPackagePath(AllSprites[0]->GetOutermost()->GetPathName());
+
+			const FString NewFlipBookDefaultPath = LongPackagePath + TEXT("/") + FlipbookName;
+			FString DefaultSuffix;
+			FString AssetName;
+			FString PackageName;
+
+			UPaperFlipbookFactory* FlipbookFactory = ConstructObject<UPaperFlipbookFactory>(UPaperFlipbookFactory::StaticClass());
+			for (int32 SpriteIndex = 0; SpriteIndex < Sprites.Num(); ++SpriteIndex)
 			{
-				while (LastCharacter > 0 && FChar::IsDigit(String[LastCharacter - 1]))
-				{
-					LastCharacter--;
-				}
+				UPaperSprite* Sprite = Sprites[SpriteIndex];
+				FPaperFlipbookKeyFrame* KeyFrame = new (FlipbookFactory->KeyFrames) FPaperFlipbookKeyFrame();
+				KeyFrame->Sprite = Sprite;
+				KeyFrame->FrameRun = 1;
+			}
 
-				if (LastCharacter >= 0)
+			AssetToolsModule.Get().CreateUniqueAssetName(NewFlipBookDefaultPath, /*out*/ DefaultSuffix, /*out*/ PackageName, /*out*/ AssetName);
+			const FString PackagePath = FPackageName::GetLongPackagePath(PackageName);
+			if (bOneFlipbookCreated)
+			{
+				ContentBrowserModule.Get().CreateNewAsset(AssetName, PackagePath, UPaperFlipbook::StaticClass(), FlipbookFactory);
+			}
+			else
+			{
+				if (UObject* NewAsset = AssetToolsModule.Get().CreateAsset(AssetName, PackagePath, UPaperFlipbook::StaticClass(), FlipbookFactory))
 				{
-					int FirstDigit = LastCharacter;
-					FString NumberString = String.Mid(FirstDigit);
-					BareString = String.Left(FirstDigit);
-					Number = FCString::Atoi(*NumberString);
-					bExtracted = true;
+					ObjectsToSync.Add(NewAsset);
 				}
 			}
-		
-			if (!bExtracted)
+
+			if (GWarn->ReceivedUserCancel())
 			{
-				BareString = String;
-				Number = -1;
+				break;
 			}
 		}
 
-		// Sort predicate operator
-		bool operator()(UPaperSprite& LHS, UPaperSprite& RHS) const
-		{
-			FString LeftString;
-			int LeftNumber;
-			ExtractNumber(LHS.GetName(), /*out*/LeftString, /*out*/LeftNumber);
-
-			FString RightString;
-			int RightNumber;
-			ExtractNumber(RHS.GetName(), /*out*/RightString, /*out*/RightNumber);
+		GWarn->EndSlowTask();
 
-			return (LeftString == RightString) ? LeftNumber < RightNumber : LeftString < RightString;
-		}
-	};
-
-	Sprites.Sort(FSpriteSortPredicate());
-	
-	// Create the flipbook
-	if (Sprites.Num() > 0)
-	{
-		const FString SpritePathName = Sprites[0]->GetOutermost()->GetPathName();
-		const FString LongPackagePath = FPackageName::GetLongPackagePath(Sprites[0]->GetOutermost()->GetPathName());
-
-		const FString NewFlipBookDefaultPath = LongPackagePath + TEXT("/NewFlipbook");
-		FString DefaultSuffix;
-		FString AssetName;
-		FString PackageName;
-		AssetToolsModule.Get().CreateUniqueAssetName(NewFlipBookDefaultPath, /*out*/ DefaultSuffix, /*out*/ PackageName, /*out*/ AssetName);
-		const FString PackagePath = FPackageName::GetLongPackagePath(PackageName);
-
-		UPaperFlipbookFactory* FlipbookFactory = ConstructObject<UPaperFlipbookFactory>(UPaperFlipbookFactory::StaticClass());
-		for (int32 SpriteIndex = 0; SpriteIndex < Sprites.Num(); ++SpriteIndex)
+		if (ObjectsToSync.Num() > 0)
 		{
-			UPaperSprite* Sprite = Sprites[SpriteIndex];
-			FPaperFlipbookKeyFrame* KeyFrame = new (FlipbookFactory->KeyFrames) FPaperFlipbookKeyFrame();
-			KeyFrame->Sprite = Sprite;
-			KeyFrame->FrameRun = 1;
+			ContentBrowserModule.Get().SyncBrowserToAssets(ObjectsToSync);
 		}
-
-		ContentBrowserModule.Get().CreateNewAsset(AssetName, PackagePath, UPaperFlipbook::StaticClass(), FlipbookFactory);
 	}
+
 }
 
 
diff --git a/Engine/Plugins/Experimental/Paper2D/Source/Paper2DEditor/Public/PaperFlipbookHelpers.h b/Engine/Plugins/Experimental/Paper2D/Source/Paper2DEditor/Public/PaperFlipbookHelpers.h
new file mode 100644
--- /dev/null
+++ b/Engine/Plugins/Experimental/Paper2D/Source/Paper2DEditor/Public/PaperFlipbookHelpers.h
@@ -0,0 +1,24 @@
+// Copyright 1998-2014 Epic Games, Inc. All Rights Reserved.
+
+#pragma once
+
+//////////////////////////////////////////////////////////////////////////
+// FPaperFlipbookHelpers
+
+class PAPER2DEDITOR_API FPaperFlipbookHelpers
+{
+public:
+	// Try and best guess split up the sprites into multiple flipbooks based on names
+	// Numbers are naturally sorted
+	// Eg. Walk01, Walk02, Run01, Run02 -> Walk (Walk01, Walk02), Run (Run01, Run02) if autoGroup is true
+	// All ungrouped sprites are stored in one flipbook
+	// Expected input: 
+	// Sprites array contains valid sprites
+	// SpriteNames contains an array of sprite names (parallel to Sprites)
+	// or, an empty array if Sprite->GetName() is to be used.
+	static void ExtractFlipbooksFromSprites(TMap<FString, TArray<UPaperSprite*> >& OutSpriteFlipbookMap, const TArray<UPaperSprite*>& Sprites, const TArray<FString>& SpriteNames);
+
+private:
+	FPaperFlipbookHelpers() {}
+};
+
diff --git a/Engine/Plugins/Experimental/Paper2D/Source/PaperJsonImporter/Private/PaperSpriteSheetAssetTypeActions.cpp b/Engine/Plugins/Experimental/Paper2D/Source/PaperJsonImporter/Private/PaperSpriteSheetAssetTypeActions.cpp
--- a/Engine/Plugins/Experimental/Paper2D/Source/PaperJsonImporter/Private/PaperSpriteSheetAssetTypeActions.cpp
+++ b/Engine/Plugins/Experimental/Paper2D/Source/PaperJsonImporter/Private/PaperSpriteSheetAssetTypeActions.cpp
@@ -9,6 +9,7 @@
 #include "AssetRegistryModule.h"
 #include "PaperFlipbookFactory.h"
 #include "PackageTools.h"
+#include "PaperFlipbookHelpers.h"
 
 #define LOCTEXT_NAMESPACE "AssetTypeActions"
 
@@ -77,35 +78,46 @@ void FPaperSpriteSheetAssetTypeActions::ExecuteReimport(TArray<TWeakObjectPtr<UP
 static bool ExtractSpriteNumber(const FString& String, FString& BareString, int& Number)
 {
 	bool bExtracted = false;
+
 	int LastCharacter = String.Len() - 1;
-	if (LastCharacter >= 0 && FChar::IsDigit(String[LastCharacter]))
+	if (LastCharacter >= 0)
 	{
-		while (LastCharacter > 0 && FChar::IsDigit(String[LastCharacter - 1]))
+		// Find the last character that isn't a digit (Handle sprite names with numbers inside inverted commas / parentheses)
+		while (LastCharacter >= 0 && !FChar::IsDigit(String[LastCharacter]))
 		{
 			LastCharacter--;
 		}
 
+		// Only proceed if we found a number in the sprite name
 		if (LastCharacter >= 0)
 		{
-			int FirstDigit = LastCharacter;
-			int EndCharacter = FirstDigit;
-			while (EndCharacter > 0 && !FChar::IsAlnum(String[EndCharacter - 1]))
+			while (LastCharacter > 0 && FChar::IsDigit(String[LastCharacter - 1]))
 			{
-				EndCharacter--;
+				LastCharacter--;
 			}
 
-			if (EndCharacter == 0)
+			if (LastCharacter >= 0)
 			{
-				// This string consists of non alnum + number, eg. _42
-				// The flipbook / category name in this case will be _
-				// Otherwise, we strip out all trailing non-alnum chars
-				EndCharacter = FirstDigit;
+				int FirstDigit = LastCharacter;
+				int EndCharacter = FirstDigit;
+				while (EndCharacter > 0 && !FChar::IsAlnum(String[EndCharacter - 1]))
+				{
+					EndCharacter--;
+				}
+
+				if (EndCharacter == 0)
+				{
+					// This string consists of non alnum + number, eg. _42
+					// The flipbook / category name in this case will be _
+					// Otherwise, we strip out all trailing non-alnum chars
+					EndCharacter = FirstDigit;
+				}
+
+				FString NumberString = String.Mid(FirstDigit);
+				BareString = String.Left(EndCharacter);
+				Number = FCString::Atoi(*NumberString);
+				bExtracted = true;
 			}
-			
-			FString NumberString = String.Mid(FirstDigit);
-			BareString = String.Left(EndCharacter);
-			Number = FCString::Atoi(*NumberString);
-			bExtracted = true;
 		}
 	}
 
@@ -129,14 +141,13 @@ void FPaperSpriteSheetAssetTypeActions::ExecuteCreateFlipbooks(TArray<TWeakObjec
 			FAssetToolsModule& AssetToolsModule = FModuleManager::GetModuleChecked<FAssetToolsModule>("AssetTools");
 			FContentBrowserModule& ContentBrowserModule = FModuleManager::LoadModuleChecked<FContentBrowserModule>("ContentBrowser");
 			
-			TMap<FString, UPaperSprite*> SpriteNameMap;
-			TSet<FString> FlipbookNames;
-			TMap<FString, TArray<UPaperSprite*> > FlipbookSprites;
-			TArray<UPaperSprite*> AllValidSprites;
-
 			check(SpriteSheet->SpriteNames.Num() == SpriteSheet->Sprites.Num());
 			bool useSpriteNames = (SpriteSheet->SpriteNames.Num() == SpriteSheet->Sprites.Num());
 
+			// Create a list of sprites and sprite names to feed into paper flipbook helpers
+			TArray<UPaperSprite*> Sprites;
+			TArray<FString> SpriteNames;
+
 			for (int SpriteIndex = 0; SpriteIndex < SpriteSheet->Sprites.Num(); ++SpriteIndex)
 			{
 				auto SpriteAssetPtr = SpriteSheet->Sprites[SpriteIndex];
@@ -149,73 +160,30 @@ void FPaperSpriteSheetAssetTypeActions::ExecuteCreateFlipbooks(TArray<TWeakObjec
 				if (Sprite != nullptr)
 				{
 					const FString SpriteName = useSpriteNames ? SpriteSheet->SpriteNames[SpriteIndex] : Sprite->GetName();
-
-					SpriteNameMap.Add(SpriteName, Sprite);
-					
-					int SpriteNumber = 0;
-					FString SpriteBareString;
-					if (ExtractSpriteNumber(SpriteName, /*out*/SpriteBareString, /*out*/SpriteNumber))
-					{
-						FlipbookNames.Add(SpriteBareString);
-						if (FlipbookSprites.Contains(SpriteBareString))
-						{
-							FlipbookSprites[SpriteBareString].Add(Sprite);
-						}
-						else
-						{
-							FlipbookSprites.Add(SpriteBareString, TArray<UPaperSprite*>());
-						}
-					}
-
-					AllValidSprites.Add(Sprite);
+					Sprites.Add(Sprite);
+					SpriteNames.Add(SpriteName);
 				}
 			}
 
-			// Natural sort using the same method as above
-			struct FSpriteSortPredicate
-			{
-				FSpriteSortPredicate() {}
-
-				// Sort predicate operator
-				bool operator()(UPaperSprite& LHS, UPaperSprite& RHS) const
-				{
-					FString LeftString;
-					int LeftNumber;
-					ExtractSpriteNumber(LHS.GetName(), /*out*/LeftString, /*out*/LeftNumber);
-
-					FString RightString;
-					int RightNumber;
-					ExtractSpriteNumber(RHS.GetName(), /*out*/RightString, /*out*/RightNumber);
-
-					return (LeftString == RightString) ? LeftNumber < RightNumber : LeftString < RightString;
-				}
-			};
-
-			// Haven't matched any existing names, so create one flipbook with all sprites in it
-			if (FlipbookNames.Num() == 0 &&
-				EAppReturnType::Yes == FMessageDialog::Open(EAppMsgType::YesNo, NSLOCTEXT("Paper2D", "Paper2D_SpriteSheetCreateOneFlipbook", "Unable to detect flipbooks by name in this sprite sheet. Proceed to create a single flipbook containing all sprites in this sprite sheet?")) )
-			{
-				const FString& DefaultFlipbookName = SpriteSheet->GetName() + TEXT(" Flipbook");
-				FlipbookNames.Add(DefaultFlipbookName);
-				FlipbookSprites.Add(DefaultFlipbookName, AllValidSprites);
-			}
+			TMap<FString, TArray<UPaperSprite*> > SpriteFlipbookMap;
+			FPaperFlipbookHelpers::ExtractFlipbooksFromSprites(/*out*/SpriteFlipbookMap, Sprites, SpriteNames);
 
 			// Create one flipbook for every grouped flipbook name
-			if (FlipbookNames.Num() > 0)
+			if (SpriteFlipbookMap.Num() > 0)
 			{
 				UPaperFlipbookFactory* FlipbookFactory = ConstructObject<UPaperFlipbookFactory>(UPaperFlipbookFactory::StaticClass());
 
-				GWarn->BeginSlowTask(NSLOCTEXT("Paper2D", "Paper2D_SpriteSheetCreateFlipbooks", "Creating flipbooks from Sprite Sheet"), true, true);
+				GWarn->BeginSlowTask(NSLOCTEXT("Paper2D", "Paper2D_CreateFlipbooks", "Creating flipbooks from selection"), true, true);
 
-				TArray<FString> FlipbookNameArray = FlipbookNames.Array();
+				int Progress = 0;
+				int TotalProgress = SpriteFlipbookMap.Num();
 				TArray<UObject*> ObjectsToSync;
-				for (int FlipbookNameIndex = 0; FlipbookNameIndex < FlipbookNames.Num(); ++FlipbookNameIndex)
+				for (auto It : SpriteFlipbookMap)
 				{
-					GWarn->UpdateProgress(FlipbookNameIndex, FlipbookNames.Num());
+					GWarn->UpdateProgress(Progress++, TotalProgress);
 
-					const FString& FlipbookName = FlipbookNameArray[FlipbookNameIndex];
-					TArray<UPaperSprite*>& Sprites = FlipbookSprites[FlipbookName];
-					Sprites.Sort(FSpriteSortPredicate());
+					const FString& FlipbookName = It.Key;
+					TArray<UPaperSprite*>& Sprites = It.Value;
 
 					const FString TentativePackagePath = PackageTools::SanitizePackageName(PackagePath + TEXT("/") + FlipbookName);
 					FString DefaultSuffix;
