Paper2D: Sprite region widget deltas were in wrong coordinate system, was sometimes unresponsive
#TTP 346641

diff --git a/Engine/Plugins/Experimental/Paper2D/Source/Paper2DEditor/Private/SpriteEditor/SpriteEditorSelections.h b/Engine/Plugins/Experimental/Paper2D/Source/Paper2DEditor/Private/SpriteEditor/SpriteEditorSelections.h
--- a/Engine/Plugins/Experimental/Paper2D/Source/Paper2DEditor/Private/SpriteEditor/SpriteEditorSelections.h
+++ b/Engine/Plugins/Experimental/Paper2D/Source/Paper2DEditor/Private/SpriteEditor/SpriteEditorSelections.h
@@ -116,7 +116,7 @@ public:
 		}
 	}
 
-	virtual void ApplyDeltaIndexed(const FVector2D& Delta, int32 TargetVertexIndex)
+	virtual void ApplyDeltaIndexed(const FVector2D& WorldSpaceDelta, int32 TargetVertexIndex)
 	{
 		if (UPaperSprite* Sprite = SpritePtr.Get())
 		{
@@ -158,27 +158,29 @@ public:
 			default:
 				break;
 			}
-            
+
+			const FVector2D TextureSpaceDelta = Sprite->ConvertWorldSpaceDeltaToTextureSpace(PaperAxisX * WorldSpaceDelta.X + PaperAxisY * WorldSpaceDelta.Y);
+
             if (XAxis == -1)
             {
-                float AllowedDelta = FMath::Clamp(Delta.X, -Sprite->SourceUV.X, Sprite->SourceDimension.X - 1);
+                float AllowedDelta = FMath::Clamp(TextureSpaceDelta.X, -Sprite->SourceUV.X, Sprite->SourceDimension.X - 1);
                 Sprite->SourceUV.X += AllowedDelta;
                 Sprite->SourceDimension.X -= AllowedDelta;
             }
             else if (XAxis == 1)
             {
-				Sprite->SourceDimension.X = FMath::Clamp(Sprite->SourceDimension.X + Delta.X, 1.0f, SourceDims.X - Sprite->SourceUV.X);
+				Sprite->SourceDimension.X = FMath::Clamp(Sprite->SourceDimension.X + TextureSpaceDelta.X, 1.0f, SourceDims.X - Sprite->SourceUV.X);
             }
             
             if (YAxis == -1)
             {
-				float AllowedDelta = FMath::Clamp(Delta.Y, -Sprite->SourceUV.Y, Sprite->SourceDimension.Y - 1);
+				float AllowedDelta = FMath::Clamp(TextureSpaceDelta.Y, -Sprite->SourceUV.Y, Sprite->SourceDimension.Y - 1);
 				Sprite->SourceUV.Y += AllowedDelta;
 				Sprite->SourceDimension.Y -= AllowedDelta;
             }
             else if (YAxis == 1)
             {
-				Sprite->SourceDimension.Y = FMath::Clamp(Sprite->SourceDimension.Y + Delta.Y, 1.0f, SourceDims.Y - Sprite->SourceUV.Y);
+				Sprite->SourceDimension.Y = FMath::Clamp(Sprite->SourceDimension.Y + TextureSpaceDelta.Y, 1.0f, SourceDims.Y - Sprite->SourceUV.Y);
             }
 		}
 	}
