Paper2D: Zoom to fit when the sprite texture is changed.
#TTP 341472

diff --git a/Engine/Plugins/Experimental/Paper2D/Source/Paper2DEditor/Private/SpriteEditor/SpriteEditorViewportClient.cpp b/Engine/Plugins/Experimental/Paper2D/Source/Paper2DEditor/Private/SpriteEditor/SpriteEditorViewportClient.cpp
--- a/Engine/Plugins/Experimental/Paper2D/Source/Paper2DEditor/Private/SpriteEditor/SpriteEditorViewportClient.cpp
+++ b/Engine/Plugins/Experimental/Paper2D/Source/Paper2DEditor/Private/SpriteEditor/SpriteEditorViewportClient.cpp
@@ -113,6 +113,8 @@ void FSpriteEditorViewportClient::UpdateSourceTextureSpriteFromSprite(UPaperSpri
 			SpriteReinitParams.SetTextureAndFill(SourceSprite->SourceTexture);
 			TargetSprite->PixelsPerUnrealUnit = SourceSprite->PixelsPerUnrealUnit;
 			TargetSprite->InitializeSprite(SpriteReinitParams);
+
+			bDeferZoomToSprite = true;
 		}
 
 
