Paper2D: Add class comments to APaperSpriteActor and UPaperSpriteComponent, remove EarlyAccessPreview flag from UPaperSpriteComponent

diff --git a/Engine/Plugins/2D/Paper2D/Source/Paper2D/Classes/PaperSpriteActor.h b/Engine/Plugins/2D/Paper2D/Source/Paper2D/Classes/PaperSpriteActor.h
--- a/Engine/Plugins/2D/Paper2D/Source/Paper2D/Classes/PaperSpriteActor.h
+++ b/Engine/Plugins/2D/Paper2D/Source/Paper2D/Classes/PaperSpriteActor.h
@@ -5,7 +5,13 @@
 #include "PaperSpriteComponent.h"
 #include "PaperSpriteActor.generated.h"
 
-UCLASS(MinimalAPI)
+/**
+ * An instance of a UPaperSprite in a level.
+ *
+ * This actor is created when you drag a sprite asset from the content browser into the level, and
+ * it is just a thin wrapper around a UPaperSpriteComponent that actually references the asset.
+ */
+UCLASS(MinimalAPI, meta=(ChildCanTick))
 class APaperSpriteActor : public AActor
 {
 	GENERATED_UCLASS_BODY()
diff --git a/Engine/Plugins/2D/Paper2D/Source/Paper2D/Classes/PaperSpriteComponent.h b/Engine/Plugins/2D/Paper2D/Source/Paper2D/Classes/PaperSpriteComponent.h
--- a/Engine/Plugins/2D/Paper2D/Source/Paper2D/Classes/PaperSpriteComponent.h
+++ b/Engine/Plugins/2D/Paper2D/Source/Paper2D/Classes/PaperSpriteComponent.h
@@ -6,7 +6,16 @@
 
 #include "PaperSpriteComponent.generated.h"
 
-UCLASS(ShowCategories=(Mobility), ClassGroup=Paper2D, EarlyAccessPreview, meta=(BlueprintSpawnableComponent))
+/**
+ * A component that handles rendering and collision for a single instance of a UPaperSprite asset.
+ *
+ * This component is created when you drag a sprite asset from the content browser into a Blueprint, or
+ * contained inside of the actor created when you drag one into the level.
+ *
+ * @see UPrimitiveComponent, UPaperSprite
+ */
+
+UCLASS(ShowCategories=(Mobility), ClassGroup=Paper2D, meta=(BlueprintSpawnableComponent))
 class PAPER2D_API UPaperSpriteComponent : public UMeshComponent
 {
 	GENERATED_UCLASS_BODY()
