Paper2D: Adjust plugin description

diff --git a/Engine/Plugins/2D/Paper2D/Paper2D.uplugin b/Engine/Plugins/2D/Paper2D/Paper2D.uplugin
--- a/Engine/Plugins/2D/Paper2D/Paper2D.uplugin
+++ b/Engine/Plugins/2D/Paper2D/Paper2D.uplugin
@@ -7,8 +7,8 @@
 	"CreatedBy" : "Epic Games, Inc.",
 	"CreatedByURL" : "http://epicgames.com",
 	"EngineVersion" : "4.2.0",
-	"Description" : "[EARLY ACCESS PREVIEW] Paper2D adds tools and assets to help create 2D games including animated sprite assets, tilesets, 2D level editing tools, and more.",
-	"Category" : "2D.Helpers",
+	"Description" : "Paper2D adds tools and assets to help create 2D games including animated sprite assets, tilesets (experimental), 2D level editing tools, and more.",
+	"Category" : "2D.Paper2D",
 	"EnabledByDefault" : true,
 
 	"Modules" :
