Paper2D: Terrain spline collision support WIP

diff --git a/Engine/Plugins/2D/Paper2D/Source/Paper2D/Classes/PaperTerrainComponent.h b/Engine/Plugins/2D/Paper2D/Source/Paper2D/Classes/PaperTerrainComponent.h
--- a/Engine/Plugins/2D/Paper2D/Source/Paper2D/Classes/PaperTerrainComponent.h
+++ b/Engine/Plugins/2D/Paper2D/Source/Paper2D/Classes/PaperTerrainComponent.h
@@ -22,15 +22,13 @@ class PAPER2D_API UPaperTerrainComponent : public UPrimitiveComponent
 	GENERATED_UCLASS_BODY()
 
 public:
+	/** The terrain material to apply to this component (set of rules for which sprites are used on different surfaces or the interior) */
 	UPROPERTY(Category=Sprite, EditAnywhere, BlueprintReadOnly)
 	class UPaperTerrainMaterial* TerrainMaterial;
 
 	UPROPERTY(Category = Sprite, EditAnywhere, BlueprintReadOnly)
 	bool bClosedSpline;
 
-	UPROPERTY(Category = Sprite, EditAnywhere, BlueprintReadOnly)
-	float TestScaleFactor;
-
 	UPROPERTY()
 	class UPaperTerrainSplineComponent* AssociatedSpline;
 
@@ -39,14 +37,22 @@ public:
 	int32 RandomSeed;
 
 protected:
-	// The color of the terrain (passed to the sprite material as a vertex color)
+	/** The color of the terrain (passed to the sprite material as a vertex color) */
 	UPROPERTY(Category=Sprite, BlueprintReadOnly, Interp)
 	FLinearColor TerrainColor;
 
 	/** Number of steps per spline segment to place in the reparameterization table */
-	UPROPERTY(EditAnywhere, Category=Sprite, meta=(ClampMin=4, UIMin=4))
+	UPROPERTY(Category=Sprite, EditAnywhere, meta=(ClampMin=4, UIMin=4), AdvancedDisplay)
 	int32 ReparamStepsPerSegment;
 
+	/** Collision domain (no collision, 2D (experimental), or 3D) */
+	UPROPERTY(Category=Collision, EditAnywhere)
+	TEnumAsByte<ESpriteCollisionMode::Type> SpriteCollisionDomain;
+
+	/** The extrusion thickness of collision geometry when using a 3D collision domain */
+	UPROPERTY(Category=Collision, EditAnywhere)
+	float CollisionThickness;
+
 public:
 	// UObject interface
 	virtual const UObject* AdditionalStatObject() const override;
@@ -63,18 +69,23 @@ public:
 	// UPrimitiveComponent interface
 	virtual FPrimitiveSceneProxy* CreateSceneProxy() override;
 	virtual FBoxSphereBounds CalcBounds(const FTransform& LocalToWorld) const override;
+	virtual UBodySetup* GetBodySetup() override;
 	// End of UPrimitiveComponent interface
 
 	// Set color of the terrain
-	UFUNCTION(BlueprintCallable, Category = "Sprite")
+	UFUNCTION(BlueprintCallable, Category="Sprite")
 	void SetTerrainColor(FLinearColor NewColor);
 
 protected:
-	void SpawnSegment(const class UPaperSprite* NewSprite, float Position, float HorizontalScale);
+	void SpawnSegment(const class UPaperSprite* NewSprite, float Position, float HorizontalScale, float NominalWidth, const struct FPaperTerrainMaterialRule* Rule);
 	void SpawnFromPoly(const class UPaperSprite* NewSprite, FSpriteDrawCallRecord& FillDrawCall, const FVector2D& TextureSize, FPoly& Poly);
 
 	void OnSplineEdited();
 
+	/** Description of collision */
+	UPROPERTY(Transient, DuplicateTransient)
+	class UBodySetup* CachedBodySetup;
+
 	TArray<FPaperTerrainMaterialPair> GeneratedSpriteGeometry;
 	
 	FTransform GetTransformAtDistance(float InDistance) const;
diff --git a/Engine/Plugins/2D/Paper2D/Source/Paper2D/Classes/PaperTerrainMaterial.h b/Engine/Plugins/2D/Paper2D/Source/Paper2D/Classes/PaperTerrainMaterial.h
--- a/Engine/Plugins/2D/Paper2D/Source/Paper2D/Classes/PaperTerrainMaterial.h
+++ b/Engine/Plugins/2D/Paper2D/Source/Paper2D/Classes/PaperTerrainMaterial.h
@@ -11,26 +11,35 @@ struct FPaperTerrainMaterialRule
 	GENERATED_USTRUCT_BODY()
 
 	// The sprite to use at the 'left' (closest to spline start) edge of the terrain segment
-	UPROPERTY(Category=Sprite, EditAnywhere, meta=(DisplayThumbnail = "true"))
+	UPROPERTY(Category=Sprite, EditAnywhere, meta=(DisplayThumbnail="true"))
 	UPaperSprite* StartCap;
 
 	// A set of sprites to randomly choose to fill up the interior space between the caps in a terrain segment
-	UPROPERTY(Category = Sprite, EditAnywhere, meta = (DisplayThumbnail = "true"))
+	UPROPERTY(Category=Sprite, EditAnywhere, meta=(DisplayThumbnail="true"))
 	TArray<UPaperSprite*> Body;
 
 	// The sprite to use at the 'right' (closest to spline end) edge of the terrain segment
-	UPROPERTY(Category = Sprite, EditAnywhere, meta = (DisplayThumbnail = "true"))
+	UPROPERTY(Category=Sprite, EditAnywhere, meta=(DisplayThumbnail="true"))
 	UPaperSprite* EndCap;
 
 	// Minimum slope angle (in degrees) to apply this rule
-	UPROPERTY(Category = Sprite, EditAnywhere)
+	UPROPERTY(Category=Sprite, EditAnywhere)
 	float MinimumAngle;
 
 	// Maximum slope angle (in degrees) to apply this rule
-	UPROPERTY(Category = Sprite, EditAnywhere)
+	UPROPERTY(Category=Sprite, EditAnywhere)
 	float MaximumAngle;
 
+	// If true, collision is enabled for sections matching this rule
+	UPROPERTY(Category=Sprite, EditAnywhere)
+	bool bEnableCollision;
+
+	// How much should the collision be lofted from the spline (positive values go out from the spline, negative values go in to the spline)
+	UPROPERTY(Category=Sprite, EditAnywhere)
+	float CollisionOffset;
+
 #if WITH_EDITORONLY_DATA
+	// Readable description for the rule (unused anywhere, just for clarity when editing the material)
 	UPROPERTY(Category = Sprite, EditAnywhere)
 	FText Description;
 #endif
@@ -40,6 +49,8 @@ struct FPaperTerrainMaterialRule
 		, EndCap(nullptr)
 		, MinimumAngle(0.0f)
 		, MaximumAngle(360.0f)
+		, bEnableCollision(true)
+		, CollisionOffset(0.0f)
 	{
 	}
 };
@@ -55,11 +66,11 @@ class PAPER2D_API UPaperTerrainMaterial : public UDataAsset
 {
 	GENERATED_UCLASS_BODY()
 
-	UPROPERTY(Category=Sprite, EditAnywhere, meta=(DisplayThumbnail = "true"))
+	UPROPERTY(Category=Sprite, EditAnywhere)
 	TArray<FPaperTerrainMaterialRule> Rules;
 
 	// The sprite to use for an interior region fill
-	UPROPERTY(Category = Sprite, EditAnywhere, meta = (DisplayThumbnail = "true"))
+	UPROPERTY(Category = Sprite, EditAnywhere, meta=(DisplayThumbnail="true"))
 	UPaperSprite* InteriorFill;
 
 #if WITH_EDITOR
diff --git a/Engine/Plugins/2D/Paper2D/Source/Paper2D/Private/Terrain/PaperTerrainComponent.cpp b/Engine/Plugins/2D/Paper2D/Source/Paper2D/Private/Terrain/PaperTerrainComponent.cpp
--- a/Engine/Plugins/2D/Paper2D/Source/Paper2D/Private/Terrain/PaperTerrainComponent.cpp
+++ b/Engine/Plugins/2D/Paper2D/Source/Paper2D/Private/Terrain/PaperTerrainComponent.cpp
@@ -83,8 +83,10 @@ UPaperTerrainComponent::UPaperTerrainComponent(const FPostConstructInitializePro
 	: Super(PCIP)
 	, TerrainColor(FLinearColor::White)
 	, ReparamStepsPerSegment(8)
+	, SpriteCollisionDomain(ESpriteCollisionMode::Use3DPhysics)
+	, CollisionThickness(200.0f)
 {
-	TestScaleFactor = 1.0f;
+	bCanEverAffectNavigation = true;
 
 	static ConstructorHelpers::FObjectFinder<UPaperTerrainMaterial> DefaultMaterialRef(TEXT("/Paper2D/DefaultPaperTerrainMaterial"));
 	TerrainMaterial = DefaultMaterialRef.Object;
@@ -173,17 +175,15 @@ FBoxSphereBounds UPaperTerrainComponent::CalcBounds(const FTransform& LocalToWor
 	// Graphics bounds.
 	FBoxSphereBounds NewBounds = LocalRenderBounds.TransformBy(LocalToWorld);
 
-#if 0
 	// Add bounds of collision geometry (if present).
-	if (UBodySetup* BodySetup = GetBodySetup())
+	if (CachedBodySetup != nullptr)
 	{
-		const FBox AggGeomBox = BodySetup->AggGeom.CalcAABB(LocalToWorld);
+		const FBox AggGeomBox = CachedBodySetup->AggGeom.CalcAABB(LocalToWorld);
 		if (AggGeomBox.IsValid)
 		{
 			NewBounds = Union(NewBounds, FBoxSphereBounds(AggGeomBox));
 		}
 	}
-#endif
 
 	// Apply bounds scale
 	NewBounds.BoxExtent *= BoundsScale;
@@ -192,6 +192,11 @@ FBoxSphereBounds UPaperTerrainComponent::CalcBounds(const FTransform& LocalToWor
 	return NewBounds;
 }
 
+UBodySetup* UPaperTerrainComponent::GetBodySetup()
+{
+	return CachedBodySetup;
+}
+
 FTransform UPaperTerrainComponent::GetTransformAtDistance(float InDistance) const
 {
 	const float SplineLength = AssociatedSpline->GetSplineLength();
@@ -231,6 +236,16 @@ FTransform UPaperTerrainComponent::GetTransformAtDistance(float InDistance) cons
 
 void UPaperTerrainComponent::OnSplineEdited()
 {
+	// Ensure we have the data structure for the desired collision method
+	if (SpriteCollisionDomain == ESpriteCollisionMode::Use3DPhysics)
+	{
+		CachedBodySetup = NewObject<UBodySetup>(this);
+	}
+	else
+	{
+		CachedBodySetup = nullptr;
+	}
+
 	const float SlopeAnalysisTimeRate = 10.0f;
 	const float FillRasterizationTimeRate = 100.0f;
 
@@ -461,7 +476,7 @@ void UPaperTerrainComponent::OnSplineEdited()
 			// Convert stamps into geometry
 			for (const FSpriteStamp& Stamp : Segment.Stamps)
 			{
-				SpawnSegment(Stamp.Sprite, Stamp.Time, Stamp.Scale);
+				SpawnSegment(Stamp.Sprite, Stamp.Time, Stamp.Scale, Stamp.NominalWidth, Segment.Rule);
 			}
 		}
 
@@ -550,8 +565,32 @@ void UPaperTerrainComponent::OnSplineEdited()
 	RecreateRenderState_Concurrent();
 }
 
-void UPaperTerrainComponent::SpawnSegment(const UPaperSprite* NewSprite, float Position, float HorizontalScale)
+void UPaperTerrainComponent::SpawnSegment(const UPaperSprite* NewSprite, float Position, float HorizontalScale, float NominalWidth, const FPaperTerrainMaterialRule* Rule)
 {
+	if ((Rule->bEnableCollision) && (CachedBodySetup != nullptr))
+	{
+		const FTransform LocalTransformAtCenter(GetTransformAtDistance(Position));
+
+		FKBoxElem Box;
+
+		if (bClosedSpline)
+		{
+			//@TODO: Add proper support for this!
+			Box.SetTransform(LocalTransformAtCenter);
+			Box.X = NominalWidth * HorizontalScale;
+			Box.Y = CollisionThickness;
+			Box.Z = 30.0f;
+		}
+		else
+		{
+			Box.SetTransform(LocalTransformAtCenter);
+			Box.X = NominalWidth * HorizontalScale;
+			Box.Y = CollisionThickness;
+			Box.Z = FMath::Max<float>(1.0f, FMath::Abs<float>(Rule->CollisionOffset * 0.5f));
+		}
+		CachedBodySetup->AggGeom.BoxElems.Add(Box);
+	}
+
 	if (NewSprite)
 	{
 		FPaperTerrainMaterialPair& MaterialBatch = *new (GeneratedSpriteGeometry) FPaperTerrainMaterialPair(); //@TODO: Look up the existing one instead
