Paper2D: Updated relative path to project in PaperStyle, icons weren't being found after the move from Experimental

diff --git a/Engine/Plugins/2D/Paper2D/Source/Paper2DEditor/Private/PaperStyle.cpp b/Engine/Plugins/2D/Paper2D/Source/Paper2DEditor/Private/PaperStyle.cpp
--- a/Engine/Plugins/2D/Paper2D/Source/Paper2DEditor/Private/PaperStyle.cpp
+++ b/Engine/Plugins/2D/Paper2D/Source/Paper2DEditor/Private/PaperStyle.cpp
@@ -10,7 +10,7 @@
 
 FString FPaperStyle::InContent(const FString& RelativePath, const ANSICHAR* Extension)
 {
-	static FString ContentDir = FPaths::EnginePluginsDir() / TEXT("Experimental/Paper2D/Content");
+	static FString ContentDir = FPaths::EnginePluginsDir() / TEXT("2D/Paper2D/Content");
 	return (ContentDir / RelativePath) + Extension;
 }
 
