Paper2D: More sprite geometry editing improvements
- Add polygon mode (allows adding additive and subtractive polygons)
- Shift click to insert vertex, Alt click to delete vertex
- Help text cleaned up
#TTP 341521
#TTP 341522
#TTP 341646

diff --git a/Engine/Plugins/2D/Paper2D/Source/Paper2DEditor/Private/PaperStyle.cpp b/Engine/Plugins/2D/Paper2D/Source/Paper2DEditor/Private/PaperStyle.cpp
--- a/Engine/Plugins/2D/Paper2D/Source/Paper2DEditor/Private/PaperStyle.cpp
+++ b/Engine/Plugins/2D/Paper2D/Source/Paper2DEditor/Private/PaperStyle.cpp
@@ -60,10 +60,8 @@ void FPaperStyle::Initialize()
 		StyleSet->Set("SpriteEditor.SetShowPivot", new IMAGE_BRUSH(TEXT("Icons/icon_StaticMeshEd_ShowPivot_40x"), Icon40x40));
 		StyleSet->Set("SpriteEditor.SetShowPivot.Small", new IMAGE_BRUSH(TEXT("Icons/icon_StaticMeshEd_ShowPivot_40x"), Icon20x20));
 
-		StyleSet->Set("SpriteEditor.AddPolygon", new IMAGE_PLUGIN_BRUSH(TEXT("Icons/icon_Paper2D_AddPolygon_40x"), Icon40x40));
-		StyleSet->Set("SpriteEditor.AddPolygon.Small", new IMAGE_PLUGIN_BRUSH(TEXT("Icons/icon_Paper2D_AddPolygon_40x"), Icon20x20));
-		StyleSet->Set("SpriteEditor.AddSubtractivePolygon", new IMAGE_PLUGIN_BRUSH(TEXT("Icons/icon_Paper2D_AddSubtractivePolygon_40x"), Icon40x40));
-		StyleSet->Set("SpriteEditor.AddSubtractivePolygon.Small", new IMAGE_PLUGIN_BRUSH(TEXT("Icons/icon_Paper2D_AddSubtractivePolygon_40x"), Icon20x20));
+		StyleSet->Set("SpriteEditor.ToggleAddPolygonMode", new IMAGE_PLUGIN_BRUSH(TEXT("Icons/icon_Paper2D_AddPolygon_40x"), Icon40x40));
+		StyleSet->Set("SpriteEditor.ToggleAddPolygonMode.Small", new IMAGE_PLUGIN_BRUSH(TEXT("Icons/icon_Paper2D_AddPolygon_40x"), Icon20x20));
 
 		StyleSet->Set("SpriteEditor.SnapAllVertices", new IMAGE_PLUGIN_BRUSH(TEXT("Icons/icon_Paper2D_SnapToPixelGrid_40x"), Icon40x40));
 		StyleSet->Set("SpriteEditor.SnapAllVertices.Small", new IMAGE_PLUGIN_BRUSH(TEXT("Icons/icon_Paper2D_SnapToPixelGrid_40x"), Icon20x20));
diff --git a/Engine/Plugins/2D/Paper2D/Source/Paper2DEditor/Private/SpriteEditor/SpriteEditor.cpp b/Engine/Plugins/2D/Paper2D/Source/Paper2DEditor/Private/SpriteEditor/SpriteEditor.cpp
--- a/Engine/Plugins/2D/Paper2D/Source/Paper2DEditor/Private/SpriteEditor/SpriteEditor.cpp
+++ b/Engine/Plugins/2D/Paper2D/Source/Paper2DEditor/Private/SpriteEditor/SpriteEditor.cpp
@@ -185,13 +185,10 @@ void SSpriteEditorViewport::BindCommands()
 		FExecuteAction::CreateSP( EditorViewportClientRef, &FSpriteEditorViewportClient::SplitEdge ),
 		FCanExecuteAction::CreateSP( EditorViewportClientRef, &FSpriteEditorViewportClient::CanSplitEdge ));
 	CommandList->MapAction(
-		Commands.AddPolygon,
-		FExecuteAction::CreateSP(EditorViewportClientRef, &FSpriteEditorViewportClient::AddPolygon),
-		FCanExecuteAction::CreateSP(EditorViewportClientRef, &FSpriteEditorViewportClient::CanAddPolygon));
-	CommandList->MapAction(
-		Commands.AddSubtractivePolygon,
-		FExecuteAction::CreateSP(EditorViewportClientRef, &FSpriteEditorViewportClient::AddSubtractivePolygon),
-		FCanExecuteAction::CreateSP(EditorViewportClientRef, &FSpriteEditorViewportClient::CanAddSubtractivePolygon));
+		Commands.ToggleAddPolygonMode,
+		FExecuteAction::CreateSP(EditorViewportClientRef, &FSpriteEditorViewportClient::ToggleAddPolygonMode),
+		FCanExecuteAction::CreateSP(EditorViewportClientRef, &FSpriteEditorViewportClient::CanAddPolygon),
+		FIsActionChecked::CreateSP(EditorViewportClientRef, &FSpriteEditorViewportClient::IsAddingPolygon));
 	CommandList->MapAction(
 		Commands.SnapAllVertices,
 		FExecuteAction::CreateSP( EditorViewportClientRef, &FSpriteEditorViewportClient::SnapAllVerticesToPixelGrid ),
@@ -516,8 +513,7 @@ void FSpriteEditor::ExtendToolbar()
 
 			ToolbarBuilder.BeginSection("Tools");
 			{
-				ToolbarBuilder.AddToolBarButton(FSpriteEditorCommands::Get().AddPolygon);
-//				ToolbarBuilder.AddToolBarButton(FSpriteEditorCommands::Get().AddSubtractivePolygon); // disabled until polygon editing is improved
+				ToolbarBuilder.AddToolBarButton(FSpriteEditorCommands::Get().ToggleAddPolygonMode);
 				ToolbarBuilder.AddToolBarButton(FSpriteEditorCommands::Get().SnapAllVertices);
 			}
 			ToolbarBuilder.EndSection();
diff --git a/Engine/Plugins/2D/Paper2D/Source/Paper2DEditor/Private/SpriteEditor/SpriteEditorCommands.cpp b/Engine/Plugins/2D/Paper2D/Source/Paper2DEditor/Private/SpriteEditor/SpriteEditorCommands.cpp
--- a/Engine/Plugins/2D/Paper2D/Source/Paper2DEditor/Private/SpriteEditor/SpriteEditorCommands.cpp
+++ b/Engine/Plugins/2D/Paper2D/Source/Paper2DEditor/Private/SpriteEditor/SpriteEditorCommands.cpp
@@ -33,7 +33,6 @@ void FSpriteEditorCommands::RegisterCommands()
 	// Geometry editing commands
 	UI_COMMAND(DeleteSelection, "Delete", "Delete the selection.", EUserInterfaceActionType::Button, FInputGesture(EKeys::Platform_Delete, EModifierKey::None));
 	UI_COMMAND(SplitEdge, "Split", "Split edge.", EUserInterfaceActionType::Button, FInputGesture(EKeys::Insert, EModifierKey::None));
-	UI_COMMAND(AddPolygon, "Add Polygon", "Adds a new polygon.", EUserInterfaceActionType::Button, FInputGesture());
-	UI_COMMAND(AddSubtractivePolygon, "Add Subtractive Poly", "Adds a subtractive polygon.", EUserInterfaceActionType::Button, FInputGesture());
+	UI_COMMAND(ToggleAddPolygonMode, "Add Polygon", "Adds a new polygon.", EUserInterfaceActionType::ToggleButton, FInputGesture());
 	UI_COMMAND(SnapAllVertices, "Snap to pixel grid", "Snaps all vertices to the pixel grid.", EUserInterfaceActionType::Button, FInputGesture());
 }
diff --git a/Engine/Plugins/2D/Paper2D/Source/Paper2DEditor/Private/SpriteEditor/SpriteEditorCommands.h b/Engine/Plugins/2D/Paper2D/Source/Paper2DEditor/Private/SpriteEditor/SpriteEditorCommands.h
--- a/Engine/Plugins/2D/Paper2D/Source/Paper2DEditor/Private/SpriteEditor/SpriteEditorCommands.h
+++ b/Engine/Plugins/2D/Paper2D/Source/Paper2DEditor/Private/SpriteEditor/SpriteEditorCommands.h
@@ -47,7 +47,6 @@ public:
 	// Geometry editing commands
 	TSharedPtr<FUICommandInfo> DeleteSelection;
 	TSharedPtr<FUICommandInfo> SplitEdge;
-	TSharedPtr<FUICommandInfo> AddPolygon;
-	TSharedPtr<FUICommandInfo> AddSubtractivePolygon;
+	TSharedPtr<FUICommandInfo> ToggleAddPolygonMode;
 	TSharedPtr<FUICommandInfo> SnapAllVertices;
 };
\ No newline at end of file
diff --git a/Engine/Plugins/2D/Paper2D/Source/Paper2DEditor/Private/SpriteEditor/SpriteEditorViewportClient.cpp b/Engine/Plugins/2D/Paper2D/Source/Paper2DEditor/Private/SpriteEditor/SpriteEditorViewportClient.cpp
--- a/Engine/Plugins/2D/Paper2D/Source/Paper2DEditor/Private/SpriteEditor/SpriteEditorViewportClient.cpp
+++ b/Engine/Plugins/2D/Paper2D/Source/Paper2DEditor/Private/SpriteEditor/SpriteEditorViewportClient.cpp
@@ -14,6 +14,25 @@
 
 #define LOCTEXT_NAMESPACE "SpriteEditor"
 
+//////////////////////////////////////////////////////////////////////////
+// Polygon Vertex Hashing
+
+static int32 GetPolygonVertexHash(int32 PolygonIndex, int32 VertexIndex)
+{
+	check(PolygonIndex <= 0x7fff && VertexIndex < 0xffff);
+	return (PolygonIndex << 16) + VertexIndex;
+}
+
+static int32 GetPolygonIndexFromHash(int32 PolygonVertexHash)
+{
+	return PolygonVertexHash >> 16;
+}
+
+static int32 GetVertexIndexFromHash(int32 PolygonVertexHash)
+{
+	return PolygonVertexHash & 0xffff;
+}
+
 //////////////////////////////////////////////////////////////////////////
 // FSelectionTypes
 
@@ -347,6 +366,7 @@ void FSpriteEditorViewportClient::DrawGeometry(FViewport& InViewport, FSceneView
 
 	const float NormalLength = 15.0f;
 	const FLinearColor GeometryNormalColor(0.0f, 1.0f, 0.0f, 0.5f);
+	const FLinearColor SelectedColor(1, 1, 1, 1);
 
 	// Run thru the custom collision vertices and draw hit proxies for them
 	for (int32 PolygonIndex = 0; PolygonIndex < Geometry.Polygons.Num(); ++PolygonIndex)
@@ -363,6 +383,10 @@ void FSpriteEditorViewportClient::DrawGeometry(FViewport& InViewport, FSceneView
 			const FVector2D ScreenPos = TextureSpaceToScreenSpace(View, Polygon.Vertices[VertexIndex]);
 			const FVector2D NextScreenPos = TextureSpaceToScreenSpace(View, Polygon.Vertices[NextVertexIndex]);
 
+			bool IsEdgeSelected = SelectedVertexSet.Contains(GetPolygonVertexHash(PolygonIndex, VertexIndex)) &&
+								  SelectedVertexSet.Contains(GetPolygonVertexHash(PolygonIndex, NextVertexIndex));
+
+
 			// Draw the normal tick
 			if (bShowNormals)
 			{
@@ -378,24 +402,9 @@ void FSpriteEditorViewportClient::DrawGeometry(FViewport& InViewport, FSceneView
 
 			// Draw the edge
 			{
-				if (bIsHitTesting)
-				{
-					TSharedPtr<FSpriteSelectedEdge> Data = MakeShareable(new FSpriteSelectedEdge());
-					Data->SpritePtr = Sprite;
-					Data->PolygonIndex = PolygonIndex;
-					Data->VertexIndex = VertexIndex;
-					Data->bRenderData = bIsRenderGeometry;
-					Canvas.SetHitProxy(new HSpriteSelectableObjectHitProxy(Data));
-				}
-
 				FCanvasLineItem LineItem(ScreenPos, NextScreenPos);
-				LineItem.SetColor(LineColor);
+				LineItem.SetColor(IsEdgeSelected ? SelectedColor : LineColor);
 				Canvas.DrawItem(LineItem);
-
-				if (bIsHitTesting)
-				{
-					Canvas.SetHitProxy(NULL);
-				}
 			}
 		}
 
@@ -406,6 +415,10 @@ void FSpriteEditorViewportClient::DrawGeometry(FViewport& InViewport, FSceneView
 			const float X = ScreenPos.X;
 			const float Y = ScreenPos.Y;
 
+			bool bIsVertexSelected = SelectedVertexSet.Contains(GetPolygonVertexHash(PolygonIndex, VertexIndex));
+			bool bIsVertexLastAdded = IsAddingPolygon() && AddingPolygonIndex == PolygonIndex && VertexIndex == Polygon.Vertices.Num() - 1;
+			bool bNeedHighlightVertex = bIsVertexSelected || bIsVertexLastAdded;
+
 			if (bIsHitTesting)
 			{
 				TSharedPtr<FSpriteSelectedVertex> Data = MakeShareable(new FSpriteSelectedVertex());
@@ -416,7 +429,7 @@ void FSpriteEditorViewportClient::DrawGeometry(FViewport& InViewport, FSceneView
 				Canvas.SetHitProxy(new HSpriteSelectableObjectHitProxy(Data));
 			}
 
-			Canvas.DrawTile(ScreenPos.X - CollisionVertexSize*0.5f, ScreenPos.Y - CollisionVertexSize*0.5f, CollisionVertexSize, CollisionVertexSize, 0.f, 0.f, 1.f, 1.f, VertexColor, GWhiteTexture);
+			Canvas.DrawTile(ScreenPos.X - CollisionVertexSize*0.5f, ScreenPos.Y - CollisionVertexSize*0.5f, CollisionVertexSize, CollisionVertexSize, 0.f, 0.f, 1.f, 1.f, bNeedHighlightVertex ? SelectedColor : VertexColor, GWhiteTexture);
 			
 			if (bIsHitTesting)
 			{
@@ -437,6 +450,7 @@ void FSpriteEditorViewportClient::DrawGeometryStats(FViewport& InViewport, FScen
 	TextItem.Draw(&Canvas);
 	TextItem.Position += FVector2D(6.0f, 18.0f);
 
+
 	// Draw the number of polygons
 	TextItem.Text = FText::Format(LOCTEXT("PolygonCount", "Polys: {0}"), FText::AsNumber(Geometry.Polygons.Num()));
 	TextItem.Draw(&Canvas);
@@ -653,7 +667,10 @@ void FSpriteEditorViewportClient::DrawCanvas(FViewport& Viewport, FSceneView& Vi
 
 	int32 YPos = 42;
 
-	static const FText GeomHelpStr = LOCTEXT("GeomEditHelp", "Select an edge and press Insert to add a vertex.\nSelect an edge or vertex and press Delete to remove it.\n");
+	static const FText GeomHelpStr = LOCTEXT("GeomEditHelp", "Shift + click to insert a vertex.\nSelect one or more vertices and press Delete to remove them.\nDouble click a vertex to select a polygon\n");
+	static const FText GeomAddPolygonHelpStr = LOCTEXT("GeomClickAddPolygon", "Click to start creating a polygon\nCtrl + Click to start creating a subtractive polygon\n");
+	static const FText GeomAddCollisionPolygonHelpStr = LOCTEXT("GeomClickAddCollisionPolygon", "Click to start creating a polygon");
+	static const FText GeomAddVerticesHelpStr = LOCTEXT("GeomClickAddVertices", "Click to add points to the polygon\nEnter to finish adding a polygon\nEscape to cancel this polygon\n");
 	static const FText SourceRegionHelpStr = LOCTEXT("SourceRegionHelp", "Drag handles to adjust source region\nDouble-click on an image region to select all connected pixels\nHold down Ctrl and drag a rectangle to create a new sprite at that position");
 
 	switch (CurrentMode)
@@ -689,10 +706,20 @@ void FSpriteEditorViewportClient::DrawCanvas(FViewport& Viewport, FSceneView& Vi
 		{
  			// Display tool help
 			{
-				FCanvasTextItem TextItem(FVector2D(6, YPos), GeomHelpStr, GEngine->GetSmallFont(), FLinearColor::White);
+				const FText* HelpStr;
+				if (IsAddingPolygon())
+				{
+					HelpStr = (AddingPolygonIndex == -1) ? &GeomAddCollisionPolygonHelpStr : &GeomAddVerticesHelpStr;
+				}
+				else
+				{
+					HelpStr = &GeomHelpStr;
+				}
+
+				FCanvasTextItem TextItem(FVector2D(6, YPos), *HelpStr, GEngine->GetSmallFont(), FLinearColor::White);
 				TextItem.EnableShadow(FLinearColor::Black);
 				TextItem.Draw(&Canvas);
-				YPos += 36;
+				YPos += 54;
 			}
 
 			// Draw the custom collision geometry
@@ -709,10 +736,20 @@ void FSpriteEditorViewportClient::DrawCanvas(FViewport& Viewport, FSceneView& Vi
 		{
  			// Display tool help
 			{
-				FCanvasTextItem TextItem(FVector2D(6, YPos), GeomHelpStr, GEngine->GetSmallFont(), FLinearColor::White);
+				const FText* HelpStr;
+				if (IsAddingPolygon())
+				{
+					HelpStr = (AddingPolygonIndex == -1) ? &GeomAddPolygonHelpStr : &GeomAddVerticesHelpStr;
+				}
+				else
+				{
+					HelpStr = &GeomHelpStr;
+				}
+
+				FCanvasTextItem TextItem(FVector2D(6, YPos), *HelpStr, GEngine->GetSmallFont(), FLinearColor::White);
 				TextItem.EnableShadow(FLinearColor::Black);
 				TextItem.Draw(&Canvas);
-				YPos += 36;
+				YPos += 54;
 			}
 
 			// Draw the custom render geometry
@@ -914,16 +951,35 @@ void FSpriteEditorViewportClient::ProcessClick(FSceneView& View, HHitProxy* HitP
 {
 	const FViewportClick Click(&View, this, Key, Event, HitX, HitY);
 	const bool bIsCtrlKeyDown = Viewport->KeyState(EKeys::LeftControl) || Viewport->KeyState(EKeys::RightControl);
+	const bool bIsShiftKeyDown = Viewport->KeyState(EKeys::LeftShift) || Viewport->KeyState(EKeys::RightShift);
+	const bool bIsAltKeyDown = Viewport->KeyState(EKeys::LeftAlt) || Viewport->KeyState(EKeys::RightAlt);
 	bool bHandled = false;
+	bool bAllowSelectVertex = !(IsEditingGeometry() && IsAddingPolygon());
+	HSpriteSelectableObjectHitProxy* SelectedItemProxy = HitProxyCast<HSpriteSelectableObjectHitProxy>(HitProxy);
+
+	bool bClearSelectionModifier = bIsCtrlKeyDown;
+	bool bDeleteClickedVertex = bIsAltKeyDown;
+	bool bInsertVertexModifier = bIsShiftKeyDown;
 
-	if (HSpriteSelectableObjectHitProxy* SelectedItemProxy = HitProxyCast<HSpriteSelectableObjectHitProxy>(HitProxy))
+	if (bAllowSelectVertex && SelectedItemProxy)
 	{
-		if (!bIsCtrlKeyDown)
+		if (!bClearSelectionModifier)
 		{
 			ClearSelectionSet();
 		}
 
-		if (Event == EInputEvent::IE_DoubleClick)
+		if (bDeleteClickedVertex)
+		{
+			// Delete selection
+			if (const FSpriteSelectedVertex* SelectedVertex = SelectedItemProxy->Data->CastSelectedVertex())
+			{
+				ClearSelectionSet();
+				SelectionSet.Add(SelectedItemProxy->Data);
+				SelectedVertexSet.Add(GetPolygonVertexHash(SelectedVertex->PolygonIndex, SelectedVertex->VertexIndex));
+				DeleteSelection();
+			}
+		}
+		else if (Event == EInputEvent::IE_DoubleClick)
 		{
 			// Double click to select a polygon
 			if (const FSpriteSelectedVertex* SelectedVertex = SelectedItemProxy->Data->CastSelectedVertex())
@@ -936,7 +992,19 @@ void FSpriteEditorViewportClient::ProcessClick(FSceneView& View, HHitProxy* HitP
 		}
 		else
 		{
-			SelectionSet.Add(SelectedItemProxy->Data);
+			if (const FSpriteSelectedVertex* SelectedVertex = SelectedItemProxy->Data->CastSelectedVertex())
+			{
+				int32 PolygonVertexHash = GetPolygonVertexHash(SelectedVertex->PolygonIndex, SelectedVertex->VertexIndex);
+				if (!SelectedVertexSet.Contains(PolygonVertexHash))
+				{
+					SelectedVertexSet.Add(PolygonVertexHash);
+					SelectionSet.Add(SelectedItemProxy->Data);
+				}
+			}
+			else
+			{
+				SelectionSet.Add(SelectedItemProxy->Data);
+			}
 		}
 
 		bHandled = true;
@@ -1004,9 +1072,9 @@ void FSpriteEditorViewportClient::ProcessClick(FSceneView& View, HHitProxy* HitP
 				}
 			}
 		}
-		else if (IsEditingGeometry())
+		else if (IsEditingGeometry() && !IsAddingPolygon())
 		{
-			if (bIsCtrlKeyDown)
+			if (bInsertVertexModifier)
 			{
 				FVector4 WorldPoint = View.PixelToWorld(HitX, HitY, 0);
 				UPaperSprite* Sprite = RenderSpriteComponent->GetSprite();
@@ -1127,6 +1195,68 @@ bool FSpriteEditorViewportClient::InputKey(FViewport* Viewport, int32 Controller
 			bHandled = true;
 		}
 	}
+	else if (IsEditingGeometry())
+	{
+		if (IsAddingPolygon())
+		{
+			if (Key == EKeys::LeftMouseButton && Event == IE_Pressed)
+			{
+				int32 HitX = Viewport->GetMouseX();
+				int32 HitY = Viewport->GetMouseY();
+
+				// Calculate world space positions
+				FSceneViewFamilyContext ViewFamily(FSceneViewFamily::ConstructionValues(Viewport, GetScene(), EngineShowFlags));
+				FSceneView* View = CalcSceneView(&ViewFamily);
+				FVector WorldPoint = View->PixelToWorld(HitX, HitY, 0);
+
+				UPaperSprite* Sprite = GetSpriteBeingEdited();
+				FVector2D TexturePoint = Sprite->ConvertWorldSpaceToTextureSpace(WorldPoint);
+
+				BeginTransaction(LOCTEXT("AddPolygonVertexTransaction", "Add Vertex to Polygon"));
+
+				FSpritePolygonCollection* Geometry = GetGeometryBeingEdited();
+				if (AddingPolygonIndex == -1)
+				{
+					FSpritePolygon NewPolygon;
+					
+					const bool bIsCtrlKeyDown = Viewport->KeyState(EKeys::LeftControl) || Viewport->KeyState(EKeys::RightControl);
+					NewPolygon.bNegativeWinding = CanAddSubtractivePolygon() && bIsCtrlKeyDown;
+
+					Geometry->Polygons.Add(NewPolygon);
+					AddingPolygonIndex = Geometry->Polygons.Num() - 1;
+				}
+
+				if (Geometry->Polygons.IsValidIndex(AddingPolygonIndex))
+				{
+					Geometry->Polygons[AddingPolygonIndex].Vertices.Add(TexturePoint);
+				}
+				else
+				{
+					ResetAddPolyonMode();
+				}
+
+				Geometry->GeometryType = ESpritePolygonMode::FullyCustom;
+				bManipulationDirtiedSomething = true;
+				EndTransaction();
+			}
+			else if (Key == EKeys::Enter)
+			{
+				ResetAddPolyonMode();
+			}
+			else if (Key == EKeys::Escape && AddingPolygonIndex != -1)
+			{
+				FSpritePolygonCollection* Polygons = GetGeometryBeingEdited();
+				if (Polygons->Polygons.IsValidIndex(AddingPolygonIndex))
+				{
+					BeginTransaction(LOCTEXT("DeletePolygon", "Delete Polygon"));
+					Polygons->Polygons.RemoveAt(AddingPolygonIndex);
+					bManipulationDirtiedSomething = true;
+					EndTransaction();
+				}
+				ResetAddPolyonMode();
+			}
+		}
+	}
 
 	// Start the drag
 	//@TODO: EKeys::LeftMouseButton
@@ -1269,6 +1399,7 @@ void FSpriteEditorViewportClient::NotifySpriteBeingEditedHasChanged()
 	//@TODO: Ideally we do this before switching
 	EndTransaction();
 	ClearSelectionSet();
+	ResetAddPolyonMode();
 
 	// Update components to know about the new sprite being edited
 	UPaperSprite* Sprite = GetSpriteBeingEdited();
@@ -1313,7 +1444,6 @@ void FSpriteEditorViewportClient::DeleteSelection()
 {
 	// Make a set of polygon-vertices "pairs" to delete
 	TSet<int32> CompositeIndicesSet;
-	const int32 MaxVerticesPerPolygon = 4096;
 
 	UPaperSprite* Sprite = GetSpriteBeingEdited();
 	FSpritePolygonCollection* Geometry = GetGeometryBeingEdited();
@@ -1322,20 +1452,14 @@ void FSpriteEditorViewportClient::DeleteSelection()
 		if (const FSpriteSelectedVertex* SelectedVertex = SelectionIt->CastSelectedVertex())
 		{
 			bool bValidSelectedVertex = SelectedVertex->IsValidInEditor(Sprite, IsInRenderingEditMode());
-			if (bValidSelectedVertex && ensure(SelectedVertex->PolygonIndex < MAX_int32))
+			if (bValidSelectedVertex)
 			{
-				// ensure we're not going to overflow the pair hash above
-				if (ensure(SelectedVertex->VertexIndex < MaxVerticesPerPolygon))
-				{
-					CompositeIndicesSet.Add(SelectedVertex->PolygonIndex * MaxVerticesPerPolygon + SelectedVertex->VertexIndex);
-				}
+				CompositeIndicesSet.Add(GetPolygonVertexHash(SelectedVertex->PolygonIndex, SelectedVertex->VertexIndex));
+	
 				if (SelectedVertex->IsA(FSelectionTypes::Edge)) // add the "next" point for the edge
 				{
 					int32 NextIndex = (SelectedVertex->VertexIndex + 1) % Geometry->Polygons[SelectedVertex->PolygonIndex].Vertices.Num();
-					if (ensure(NextIndex < MaxVerticesPerPolygon))
-					{
-						CompositeIndicesSet.Add(SelectedVertex->PolygonIndex * MaxVerticesPerPolygon + NextIndex);
-					}
+					CompositeIndicesSet.Add(GetPolygonVertexHash(SelectedVertex->PolygonIndex, NextIndex));
 				}
 			}
 		}
@@ -1351,8 +1475,8 @@ void FSpriteEditorViewportClient::DeleteSelection()
 		for (int32 SelectedCompositeIndex = CompositeIndices.Num() - 1; SelectedCompositeIndex >= 0; --SelectedCompositeIndex)
 		{
 			int32 CompositeIndex = CompositeIndices[SelectedCompositeIndex];
-			int32 PolygonIndex = CompositeIndex / MaxVerticesPerPolygon;
-			int32 VertexIndex = CompositeIndex % MaxVerticesPerPolygon;
+			int32 PolygonIndex = GetPolygonIndexFromHash(CompositeIndex);
+			int32 VertexIndex = GetVertexIndexFromHash(CompositeIndex);
 			DeleteVertexInPolygonInternal(Sprite, Geometry, PolygonIndex, VertexIndex);
 		}
 
@@ -1361,6 +1485,7 @@ void FSpriteEditorViewportClient::DeleteSelection()
 
 		ClearSelectionSet();
 		EndTransaction();
+		ResetAddPolyonMode();
 	}
 }
 
@@ -1385,64 +1510,20 @@ void FSpriteEditorViewportClient::SplitEdge()
 	EndTransaction();
 }
 
-void FSpriteEditorViewportClient::AddPolygon()
+void FSpriteEditorViewportClient::ToggleAddPolygonMode()
 {
-	BeginTransaction(LOCTEXT("AddPolygonTransaction", "Add Polygon"));
-
-	if (FSpritePolygonCollection* Geometry = GetGeometryBeingEdited())
+	if (IsAddingPolygon())
 	{
-		FSpritePolygon& NewPoly = *new (Geometry->Polygons) FSpritePolygon;
-		NewPoly.bNegativeWinding = false;
-
-		//@TODO: Should make this more awesome (or even just regular awesome)
-		// Hack adding verts since we don't have 'line drawing mode' where you can just click-click-click yet
-		UPaperSprite* Sprite = GetSpriteBeingEdited();
-		const FVector2D BaseUV = Sprite->GetSourceUV();
-		const float Size = 10.0f;
-
-		new (NewPoly.Vertices) FVector2D(BaseUV.X, BaseUV.Y);
-		new (NewPoly.Vertices) FVector2D(BaseUV.X + Size, BaseUV.Y);
-		new (NewPoly.Vertices) FVector2D(BaseUV.X + Size, BaseUV.Y + Size);
-		new (NewPoly.Vertices) FVector2D(BaseUV.X, BaseUV.Y + Size);
-		bManipulationDirtiedSomething = true;
-
-		// Select the newly added polygon
-		SelectPolygon(Geometry->Polygons.Num() - 1);
-
-		Geometry->GeometryType = ESpritePolygonMode::FullyCustom;
+		ResetAddPolyonMode();
 	}
-
-	EndTransaction();
-}
-
-void FSpriteEditorViewportClient::AddSubtractivePolygon()
-{
-	BeginTransaction(LOCTEXT("AddSubtractivePolygonTransaction", "Add Subtractive Polygon"));
-
-	if (FSpritePolygonCollection* Geometry = GetGeometryBeingEdited())
+	else
 	{
-		FSpritePolygon& NewPoly = *new (Geometry->Polygons) FSpritePolygon;
-		NewPoly.bNegativeWinding = true;
-
-		//@TODO: Should make this more awesome (or even just regular awesome)
-		// Hack adding verts since we don't have 'line drawing mode' where you can just click-click-click yet
-		UPaperSprite* Sprite = GetSpriteBeingEdited();
-		const FVector2D BaseUV = Sprite->GetSourceUV();
-		const float Size = 10.0f;
-
-		new (NewPoly.Vertices) FVector2D(BaseUV.X, BaseUV.Y);
-		new (NewPoly.Vertices) FVector2D(BaseUV.X + Size, BaseUV.Y);
-		new (NewPoly.Vertices) FVector2D(BaseUV.X + Size, BaseUV.Y + Size);
-		new (NewPoly.Vertices) FVector2D(BaseUV.X, BaseUV.Y + Size);
-		bManipulationDirtiedSomething = true;
-
-		// Select the newly added polygon
-		SelectPolygon(Geometry->Polygons.Num() - 1);
-
-		Geometry->GeometryType = ESpritePolygonMode::FullyCustom;
+		ClearSelectionSet();
+		
+		bIsAddingPolygon = true; 
+		AddingPolygonIndex = -1; 
+		Invalidate();
 	}
-
-	EndTransaction();
 }
 
 void FSpriteEditorViewportClient::SnapAllVerticesToPixelGrid()
@@ -1499,7 +1580,13 @@ void FSpriteEditorViewportClient::SelectPolygon(const int PolygonIndex)
 			Vertex->PolygonIndex = PolygonIndex;
 			Vertex->VertexIndex = VertexIndex;
 			Vertex->bRenderData = IsInRenderingEditMode();
-			SelectionSet.Add(Vertex);
+
+			int32 PolygonVertexHash = GetPolygonVertexHash(PolygonIndex, VertexIndex);
+			if (!SelectedVertexSet.Contains(PolygonVertexHash))
+			{
+				SelectionSet.Add(Vertex);
+				SelectedVertexSet.Add(PolygonVertexHash);
+			}
 		}
 	}
 }
@@ -1507,6 +1594,7 @@ void FSpriteEditorViewportClient::SelectPolygon(const int PolygonIndex)
 void FSpriteEditorViewportClient::ClearSelectionSet()
 {
 	SelectionSet.Empty();
+	SelectedVertexSet.Empty();
 }
 
 static bool ClosestPointOnLine(const FVector2D& Point, const FVector2D& LineStart, const FVector2D& LineEnd, FVector2D& OutClosestPoint)
@@ -1603,6 +1691,7 @@ void FSpriteEditorViewportClient::AddPointToGeometry(const FVector2D& TextureSpa
 		SelectedVertex->VertexIndex = ClosestVertexInsertIndex;
 		SelectedVertex->bRenderData = IsInRenderingEditMode();
 		SelectionSet.Add(SelectedVertex);
+		SelectedVertexSet.Add(GetPolygonVertexHash(ClosestPolygonIndex, ClosestVertexInsertIndex));
 
 	}
 }
diff --git a/Engine/Plugins/2D/Paper2D/Source/Paper2DEditor/Private/SpriteEditor/SpriteEditorViewportClient.h b/Engine/Plugins/2D/Paper2D/Source/Paper2DEditor/Private/SpriteEditor/SpriteEditorViewportClient.h
--- a/Engine/Plugins/2D/Paper2D/Source/Paper2DEditor/Private/SpriteEditor/SpriteEditorViewportClient.h
+++ b/Engine/Plugins/2D/Paper2D/Source/Paper2DEditor/Private/SpriteEditor/SpriteEditorViewportClient.h
@@ -78,11 +78,11 @@ public:
 	void ToggleShowMeshEdges();
 	bool IsShowMeshEdgesChecked() const;
 
-	void EnterViewMode() { CurrentMode = ESpriteEditorMode::ViewMode; ResetMarqueeTracking(); }
-	void EnterSourceRegionEditMode() { CurrentMode = ESpriteEditorMode::EditSourceRegionMode; ResetMarqueeTracking(); UpdateRelatedSpritesList(); }
-	void EnterCollisionEditMode() { CurrentMode = ESpriteEditorMode::EditCollisionMode; ResetMarqueeTracking(); }
-	void EnterRenderingEditMode() { CurrentMode = ESpriteEditorMode::EditRenderingGeomMode; ResetMarqueeTracking(); }
-	void EnterAddSpriteMode() { CurrentMode = ESpriteEditorMode::AddSpriteMode; ResetMarqueeTracking(); }
+	void EnterViewMode() { CurrentMode = ESpriteEditorMode::ViewMode; ResetMarqueeTracking(); ResetAddPolyonMode(); }
+	void EnterSourceRegionEditMode() { CurrentMode = ESpriteEditorMode::EditSourceRegionMode; ResetMarqueeTracking(); UpdateRelatedSpritesList(); ResetAddPolyonMode(); }
+	void EnterCollisionEditMode() { CurrentMode = ESpriteEditorMode::EditCollisionMode; ResetMarqueeTracking(); ResetAddPolyonMode(); }
+	void EnterRenderingEditMode() { CurrentMode = ESpriteEditorMode::EditRenderingGeomMode; ResetMarqueeTracking(); ResetAddPolyonMode(); }
+	void EnterAddSpriteMode() { CurrentMode = ESpriteEditorMode::AddSpriteMode; ResetMarqueeTracking(); ResetAddPolyonMode(); }
 
 	bool IsInViewMode() const { return CurrentMode == ESpriteEditorMode::ViewMode; }
 	bool IsInSourceRegionEditMode() const { return CurrentMode == ESpriteEditorMode::EditSourceRegionMode; }
@@ -106,8 +106,9 @@ public:
 	void SplitEdge();
 	bool CanSplitEdge() const { return IsEditingGeometry(); } //@TODO: Need an edge
 
-	void AddPolygon();
-	void AddSubtractivePolygon();
+	void ResetAddPolyonMode() { bIsAddingPolygon = false; }
+	void ToggleAddPolygonMode();
+	bool IsAddingPolygon() const { return bIsAddingPolygon; }
 	bool CanAddPolygon() const { return IsEditingGeometry(); }
 	bool CanAddSubtractivePolygon() const { return CanAddPolygon() && IsInRenderingEditMode(); }
 
@@ -156,6 +157,9 @@ private:
 	// Selection set of vertices
 	TSet< TSharedPtr<FSelectedItem> > SelectionSet;
 
+	// Set of selected vertices
+	TSet<int32> SelectedVertexSet;
+
 	// The current transaction for undo/redo
 	class FScopedTransaction* ScopedTransaction;
 
@@ -177,6 +181,12 @@ private:
 	// Should we show related sprites in the source texture?
 	bool bShowRelatedSprites;
 
+	// Is waiting to add geometry
+	bool bIsAddingPolygon;
+
+	// The polygon index being added to, -1 if we don't have a polygon yet
+	int AddingPolygonIndex;
+
 	// Marquee tracking
 	bool bIsMarqueeTracking;
 	FVector2D MarqueeStartPos, MarqueeEndPos;
