Paper2D: Marquee selecting vertices in geometry edit mode
#TTP 341523

diff --git a/Engine/Plugins/2D/Paper2D/Source/Paper2DEditor/Private/SpriteEditor/SpriteEditorViewportClient.cpp b/Engine/Plugins/2D/Paper2D/Source/Paper2DEditor/Private/SpriteEditor/SpriteEditorViewportClient.cpp
--- a/Engine/Plugins/2D/Paper2D/Source/Paper2DEditor/Private/SpriteEditor/SpriteEditorViewportClient.cpp
+++ b/Engine/Plugins/2D/Paper2D/Source/Paper2DEditor/Private/SpriteEditor/SpriteEditorViewportClient.cpp
@@ -383,8 +383,7 @@ void FSpriteEditorViewportClient::DrawGeometry(FViewport& InViewport, FSceneView
 			const FVector2D ScreenPos = TextureSpaceToScreenSpace(View, Polygon.Vertices[VertexIndex]);
 			const FVector2D NextScreenPos = TextureSpaceToScreenSpace(View, Polygon.Vertices[NextVertexIndex]);
 
-			bool IsEdgeSelected = SelectedVertexSet.Contains(GetPolygonVertexHash(PolygonIndex, VertexIndex)) &&
-								  SelectedVertexSet.Contains(GetPolygonVertexHash(PolygonIndex, NextVertexIndex));
+			bool IsEdgeSelected = IsPolygonVertexSelected(PolygonIndex, VertexIndex) && IsPolygonVertexSelected(PolygonIndex, NextVertexIndex);
 
 
 			// Draw the normal tick
@@ -415,7 +414,7 @@ void FSpriteEditorViewportClient::DrawGeometry(FViewport& InViewport, FSceneView
 			const float X = ScreenPos.X;
 			const float Y = ScreenPos.Y;
 
-			bool bIsVertexSelected = SelectedVertexSet.Contains(GetPolygonVertexHash(PolygonIndex, VertexIndex));
+			bool bIsVertexSelected = IsPolygonVertexSelected(PolygonIndex, VertexIndex);
 			bool bIsVertexLastAdded = IsAddingPolygon() && AddingPolygonIndex == PolygonIndex && VertexIndex == Polygon.Vertices.Num() - 1;
 			bool bNeedHighlightVertex = bIsVertexSelected || bIsVertexLastAdded;
 
@@ -994,12 +993,7 @@ void FSpriteEditorViewportClient::ProcessClick(FSceneView& View, HHitProxy* HitP
 		{
 			if (const FSpriteSelectedVertex* SelectedVertex = SelectedItemProxy->Data->CastSelectedVertex())
 			{
-				int32 PolygonVertexHash = GetPolygonVertexHash(SelectedVertex->PolygonIndex, SelectedVertex->VertexIndex);
-				if (!SelectedVertexSet.Contains(PolygonVertexHash))
-				{
-					SelectedVertexSet.Add(PolygonVertexHash);
-					SelectionSet.Add(SelectedItemProxy->Data);
-				}
+				AddPolygonVertexToSelection(SelectedVertex->PolygonIndex, SelectedVertex->VertexIndex);
 			}
 			else
 			{
@@ -1150,6 +1144,40 @@ UPaperSprite* FSpriteEditorViewportClient::CreateNewSprite(FVector2D TopLeft, FV
 	return CreatedSprite;
 }
 
+bool FSpriteEditorViewportClient::ProcessMarquee(FViewport* Viewport, int32 ControllerId, FKey Key, EInputEvent Event, bool bMarqueeStartModifierPressed)
+{
+	bool bMarqueeReady = false;
+
+	if (Key == EKeys::LeftMouseButton)
+	{
+		int32 HitX = Viewport->GetMouseX();
+		int32 HitY = Viewport->GetMouseY();
+		if (Event == IE_Pressed && bMarqueeStartModifierPressed)
+		{
+			HHitProxy*	HitResult = Viewport->GetHitProxy(HitX, HitY);
+			if (!HitResult)
+			{
+				bIsMarqueeTracking = true;
+				MarqueeStartPos = FVector2D(HitX, HitY);
+				MarqueeEndPos = MarqueeStartPos;
+			}
+		}
+		else if (bIsMarqueeTracking && Event == IE_Released)
+		{
+			MarqueeEndPos = FVector2D(HitX, HitY);
+			bIsMarqueeTracking = false;
+			bMarqueeReady = true;
+		}
+	}
+	else if (bIsMarqueeTracking && Key == EKeys::Escape)
+	{
+		// Cancel marquee selection
+		bIsMarqueeTracking = false;
+	}
+
+	return bMarqueeReady;
+}
+
 bool FSpriteEditorViewportClient::InputKey(FViewport* Viewport, int32 ControllerId, FKey Key, EInputEvent Event, float AmountDepressed, bool bGamepad)
 {
 	bool bHandled = false;
@@ -1158,42 +1186,16 @@ bool FSpriteEditorViewportClient::InputKey(FViewport* Viewport, int32 Controller
 	// Handle marquee tracking in source region edit mode
 	if (IsInSourceRegionEditMode())
 	{
-		if (Key == EKeys::LeftMouseButton)
+		bool bMarqueeStartModifier = InputState.IsCtrlButtonPressed();
+		if (ProcessMarquee(Viewport, ControllerId, Key, Event, bMarqueeStartModifier))
 		{
-			int32 HitX = Viewport->GetMouseX();
-			int32 HitY = Viewport->GetMouseY();
-
-			if (Event == IE_Pressed && InputState.IsCtrlButtonPressed())
+			FVector2D TextureSpaceStartPos, TextureSpaceDimensions;
+			if (ConvertMarqueeToSourceTextureSpace(/*out*/TextureSpaceStartPos, /*out*/TextureSpaceDimensions))
 			{
-				HHitProxy*	HitResult = Viewport->GetHitProxy(HitX, HitY);
-				if (HitResult == nullptr)
-				{
-					bIsMarqueeTracking = true;
-					MarqueeStartPos = FVector2D(HitX, HitY);
-					MarqueeEndPos = MarqueeStartPos;
-					bHandled = true;
-				}
-			}
-			else if (bIsMarqueeTracking && Event == IE_Released)
-			{
-				MarqueeEndPos = FVector2D(HitX, HitY);
-				bIsMarqueeTracking = false;
-				bHandled = true;
-
-				FVector2D TextureSpaceStartPos, TextureSpaceDimensions;
-				if (ConvertMarqueeToSourceTextureSpace(/*out*/TextureSpaceStartPos, /*out*/TextureSpaceDimensions))
-				{
-					//@TODO: Warn if overlapping with another sprite
-					CreateNewSprite(TextureSpaceStartPos, TextureSpaceDimensions);
-				}
+				//@TODO: Warn if overlapping with another sprite
+				CreateNewSprite(TextureSpaceStartPos, TextureSpaceDimensions);
 			}
 		}
-		else if (bIsMarqueeTracking && Key == EKeys::Escape)
-		{
-			// Allow user to cancel operation
-			bIsMarqueeTracking = false;
-			bHandled = true;
-		}
 	}
 	else if (IsEditingGeometry())
 	{
@@ -1256,6 +1258,14 @@ bool FSpriteEditorViewportClient::InputKey(FViewport* Viewport, int32 Controller
 				ResetAddPolyonMode();
 			}
 		}
+		else
+		{
+			if (ProcessMarquee(Viewport, ControllerId, Key, Event, true))
+			{
+				bool bAddingToSelection = InputState.IsShiftButtonPressed(); //@TODO: control button moves widget? Hopefully make this more consistent when that is changed
+				SelectVerticesInMarquee(bAddingToSelection);
+			}
+		}
 	}
 
 	// Start the drag
@@ -1570,7 +1580,7 @@ void FSpriteEditorViewportClient::SelectPolygon(const int PolygonIndex)
 {
 	UPaperSprite* Sprite = GetSpriteBeingEdited();
 	FSpritePolygonCollection* Geometry = GetGeometryBeingEdited();
-	if (ensure(Geometry != nullptr && PolygonIndex >= 0 && PolygonIndex < Geometry->Polygons.Num()))
+	if (ensure(Geometry != nullptr && Geometry->Polygons.IsValidIndex(PolygonIndex)))
 	{
 		FSpritePolygon& Polygon = Geometry->Polygons[PolygonIndex];
 		for (int VertexIndex = 0; VertexIndex < Polygon.Vertices.Num(); ++VertexIndex)
@@ -1581,16 +1591,41 @@ void FSpriteEditorViewportClient::SelectPolygon(const int PolygonIndex)
 			Vertex->VertexIndex = VertexIndex;
 			Vertex->bRenderData = IsInRenderingEditMode();
 
-			int32 PolygonVertexHash = GetPolygonVertexHash(PolygonIndex, VertexIndex);
-			if (!SelectedVertexSet.Contains(PolygonVertexHash))
+			AddPolygonVertexToSelection(PolygonIndex, VertexIndex);
+		}
+	}
+}
+
+bool FSpriteEditorViewportClient::IsPolygonVertexSelected(const int32 PolygonIndex, const int32 VertexIndex) const
+{
+	return SelectedVertexSet.Contains(GetPolygonVertexHash(PolygonIndex, VertexIndex));
+}
+
+void FSpriteEditorViewportClient::AddPolygonVertexToSelection(const int32 PolygonIndex, const int32 VertexIndex)
+{
+	UPaperSprite* Sprite = GetSpriteBeingEdited();
+	FSpritePolygonCollection* Geometry = GetGeometryBeingEdited();
+	if (ensure(Geometry != nullptr &&  PolygonIndex >= 0 && PolygonIndex < Geometry->Polygons.Num()))
+	{
+		FSpritePolygon& Polygon = Geometry->Polygons[PolygonIndex];
+		if (Polygon.Vertices.IsValidIndex(VertexIndex))
+		{
+			if (!IsPolygonVertexSelected(PolygonIndex, VertexIndex))
 			{
+				TSharedPtr<FSpriteSelectedVertex> Vertex = MakeShareable(new FSpriteSelectedVertex());
+				Vertex->SpritePtr = Sprite;
+				Vertex->PolygonIndex = PolygonIndex;
+				Vertex->VertexIndex = VertexIndex;
+				Vertex->bRenderData = IsInRenderingEditMode();
+
 				SelectionSet.Add(Vertex);
-				SelectedVertexSet.Add(PolygonVertexHash);
+				SelectedVertexSet.Add(GetPolygonVertexHash(PolygonIndex, VertexIndex));
 			}
 		}
 	}
 }
 
+
 void FSpriteEditorViewportClient::ClearSelectionSet()
 {
 	SelectionSet.Empty();
@@ -1619,7 +1654,7 @@ static bool ClosestPointOnLine(const FVector2D& Point, const FVector2D& LineStar
 
 // Adds a point to the currently edited index
 // If SelectedPolygonIndex is set, only that polygon is considered
-void FSpriteEditorViewportClient::AddPointToGeometry(const FVector2D& TextureSpacePoint, int SelectedPolygonIndex)
+void FSpriteEditorViewportClient::AddPointToGeometry(const FVector2D& TextureSpacePoint, const int32 SelectedPolygonIndex)
 {
 	FSpritePolygonCollection* Geometry = GetGeometryBeingEdited();
 
@@ -1685,14 +1720,7 @@ void FSpriteEditorViewportClient::AddPointToGeometry(const FVector2D& TextureSpa
 
 		// Select this vertex
 		ClearSelectionSet();
-		TSharedPtr<FSpriteSelectedVertex> SelectedVertex = MakeShareable(new FSpriteSelectedVertex());
-		SelectedVertex->SpritePtr = GetSpriteBeingEdited();
-		SelectedVertex->PolygonIndex = ClosestPolygonIndex;
-		SelectedVertex->VertexIndex = ClosestVertexInsertIndex;
-		SelectedVertex->bRenderData = IsInRenderingEditMode();
-		SelectionSet.Add(SelectedVertex);
-		SelectedVertexSet.Add(GetPolygonVertexHash(ClosestPolygonIndex, ClosestVertexInsertIndex));
-
+		AddPolygonVertexToSelection(ClosestPolygonIndex, ClosestVertexInsertIndex);
 	}
 }
 
@@ -1746,6 +1774,53 @@ bool FSpriteEditorViewportClient::ConvertMarqueeToSourceTextureSpace(/*out*/FVec
 	return bSuccessful;
 }
 
+void FSpriteEditorViewportClient::SelectVerticesInMarquee(bool bAddToSelection)
+{
+	if (!bAddToSelection)
+	{
+		ClearSelectionSet();
+	}
+
+	if (UPaperSprite* Sprite = GetSpriteBeingEdited())
+	{
+		// Calculate world space positions
+		FSceneViewFamilyContext ViewFamily(FSceneViewFamily::ConstructionValues(Viewport, GetScene(), EngineShowFlags));
+		FSceneView* View = CalcSceneView(&ViewFamily);
+		FVector StartPos = View->PixelToWorld(MarqueeStartPos.X, MarqueeStartPos.Y, 0);
+		FVector EndPos = View->PixelToWorld(MarqueeEndPos.X, MarqueeEndPos.Y, 0);
+
+		// Convert to source texture space to work out the pixels dragged
+		FVector2D TextureSpaceStartPos = Sprite->ConvertWorldSpaceToTextureSpace(StartPos);
+		FVector2D TextureSpaceEndPos = Sprite->ConvertWorldSpaceToTextureSpace(EndPos);
+
+		if (TextureSpaceStartPos.X > TextureSpaceEndPos.X)
+		{
+			Swap(TextureSpaceStartPos.X, TextureSpaceEndPos.X);
+		}
+		if (TextureSpaceStartPos.Y > TextureSpaceEndPos.Y)
+		{
+			Swap(TextureSpaceStartPos.Y, TextureSpaceEndPos.Y);
+		}
+
+		if (FSpritePolygonCollection* Geometry = GetGeometryBeingEdited())
+		{
+			for (int32 PolygonIndex = 0; PolygonIndex < Geometry->Polygons.Num(); ++PolygonIndex)
+			{
+				auto Polygon = Geometry->Polygons[PolygonIndex];
+				for (int32 VertexIndex = 0; VertexIndex < Polygon.Vertices.Num(); ++VertexIndex)
+				{
+					FVector2D& Vertex = Polygon.Vertices[VertexIndex];
+					if (Vertex.X >= TextureSpaceStartPos.X && Vertex.X <= TextureSpaceEndPos.X &&
+						Vertex.Y >= TextureSpaceStartPos.Y && Vertex.Y <= TextureSpaceEndPos.Y)
+					{
+						AddPolygonVertexToSelection(PolygonIndex, VertexIndex);
+					}
+				}
+			}
+		}
+	}
+}
+
 //////////////////////////////////////////////////////////////////////////
 
 #undef LOCTEXT_NAMESPACE
\ No newline at end of file
diff --git a/Engine/Plugins/2D/Paper2D/Source/Paper2DEditor/Private/SpriteEditor/SpriteEditorViewportClient.h b/Engine/Plugins/2D/Paper2D/Source/Paper2DEditor/Private/SpriteEditor/SpriteEditorViewportClient.h
--- a/Engine/Plugins/2D/Paper2D/Source/Paper2DEditor/Private/SpriteEditor/SpriteEditorViewportClient.h
+++ b/Engine/Plugins/2D/Paper2D/Source/Paper2DEditor/Private/SpriteEditor/SpriteEditorViewportClient.h
@@ -66,6 +66,9 @@ public:
 	virtual ECoordSystem GetWidgetCoordSystemSpace() const;
 	// End of FEditorViewportClient interface
 
+	// Process marquee selection, return true of a selection has been performed
+	bool ProcessMarquee(FViewport* Viewport, int32 ControllerId, FKey Key, EInputEvent Event, bool bMarqueeStartModifierPressed);
+
 	void ToggleShowSockets() { bShowSockets = !bShowSockets; Invalidate(); }
 	bool IsShowSocketsChecked() const { return bShowSockets; }
 
@@ -224,13 +227,16 @@ private:
 	void UpdateSourceTextureSpriteFromSprite(UPaperSprite* SourceSprite);
 	
 	// Selection handling
-	void SelectPolygon(const int PolygonIndex);
-	void AddPointToGeometry(const FVector2D& TextureSpacePoint, const int SelectedPolygonIndex = -1);
+	void SelectPolygon(const int32 PolygonIndex);
+	void AddPolygonVertexToSelection(const int32 PolygonIndex, const int32 VertexIndex);
+	bool IsPolygonVertexSelected(const int32 PolygonIndex, const int32 VertexIndex) const;
+	void AddPointToGeometry(const FVector2D& TextureSpacePoint, const int32 SelectedPolygonIndex = -1);
 	void ClearSelectionSet();
 
 
 	void ResetMarqueeTracking();
 	bool ConvertMarqueeToSourceTextureSpace(/*out*/FVector2D& OutStartPos, /*out*/FVector2D& OutDimension);
+	void SelectVerticesInMarquee(bool bAddToSelection);
 	
 	// Can return null
 	FSpritePolygonCollection* GetGeometryBeingEdited() const;
