Paper2D: Expose various flipbook asset methods to Blueprints

diff --git a/Engine/Plugins/2D/Paper2D/Source/Paper2D/Classes/PaperFlipbook.h b/Engine/Plugins/2D/Paper2D/Source/Paper2D/Classes/PaperFlipbook.h
--- a/Engine/Plugins/2D/Paper2D/Source/Paper2D/Classes/PaperFlipbook.h
+++ b/Engine/Plugins/2D/Paper2D/Source/Paper2D/Classes/PaperFlipbook.h
@@ -47,39 +47,50 @@ protected:
 
 public:
 	// Returns the nominal frame rate to play this flipbook animation back at
-	int32 GetFramesPerSecond() const;
+	float GetFramesPerSecond() const;
 
 	// Returns the total number of frames
+	UFUNCTION(BlueprintCallable, Category="Sprite")
 	int32 GetNumFrames() const;
 
 	// Returns the total duration in seconds
+	UFUNCTION(BlueprintCallable, Category="Sprite")
 	float GetTotalDuration() const;
 
 	// Returns the keyframe index that covers the specified time (in seconds), or INDEX_NONE if none exists.
 	// When bClampToEnds is true, it will choose the first or last keyframe if the time is out of range.
+	UFUNCTION(BlueprintCallable, Category="Sprite")
 	int32 GetKeyFrameIndexAtTime(float Time, bool bClampToEnds = false) const;
 
 	// Returns the sprite at the specified time (in seconds), or nullptr if none exists.
 	// When bClampToEnds is true, it will choose the first or last sprite if the time is out of range.
+	UFUNCTION(BlueprintCallable, Category="Sprite")
 	UPaperSprite* GetSpriteAtTime(float Time, bool bClampToEnds = false) const;
 
-	// Returns the sprite at the specified frame index, or nullptr if none exists
+	// Returns the sprite at the specified keyframe index, or nullptr if none exists
+	UFUNCTION(BlueprintCallable, Category="Sprite")
 	UPaperSprite* GetSpriteAtFrame(int32 FrameIndex) const;
 
 	// Returns the render bounds of this sprite
 	FBoxSphereBounds GetRenderBounds() const;
 
 	// Returns the number of key frames
+	UFUNCTION(BlueprintCallable, Category="Sprite")
 	int32 GetNumKeyFrames() const
+#if CPP
 	{
 		return KeyFrames.Num();
 	}
+#endif
 
 	// Is the specified Index within the valid range of key frames?
+	UFUNCTION(BlueprintCallable, Category="Sprite")
 	bool IsValidKeyFrameIndex(int32 Index) const
+#if CPP
 	{
 		return KeyFrames.IsValidIndex(Index);
 	}
+#endif
 
 	// Returns the key frame at the specified index, make sure the index is valid before use
 	const FPaperFlipbookKeyFrame& GetKeyFrameChecked(int32 Index) const
diff --git a/Engine/Plugins/2D/Paper2D/Source/Paper2D/Private/PaperFlipbook.cpp b/Engine/Plugins/2D/Paper2D/Source/Paper2D/Private/PaperFlipbook.cpp
--- a/Engine/Plugins/2D/Paper2D/Source/Paper2D/Private/PaperFlipbook.cpp
+++ b/Engine/Plugins/2D/Paper2D/Source/Paper2D/Private/PaperFlipbook.cpp
@@ -20,7 +20,7 @@ UPaperFlipbook::UPaperFlipbook(const FPostConstructInitializeProperties& PCIP)
 	DefaultMaterial = MaskedMaterialRef.Object;
 }
 
-int32 UPaperFlipbook::GetFramesPerSecond() const
+float UPaperFlipbook::GetFramesPerSecond() const
 {
 	return FramesPerSecond;
 }
