Paper2D: Add map check warning for sprite and flipbook components that use non two sided materials

diff --git a/Engine/Plugins/2D/Paper2D/Source/Paper2D/Classes/PaperFlipbookComponent.h b/Engine/Plugins/2D/Paper2D/Source/Paper2D/Classes/PaperFlipbookComponent.h
--- a/Engine/Plugins/2D/Paper2D/Source/Paper2D/Classes/PaperFlipbookComponent.h
+++ b/Engine/Plugins/2D/Paper2D/Source/Paper2D/Classes/PaperFlipbookComponent.h
@@ -162,6 +162,9 @@ public:
 	virtual void TickComponent(float DeltaTime, enum ELevelTick TickType, FActorComponentTickFunction* ThisTickFunction) override;
 	virtual void SendRenderDynamicData_Concurrent() override;
 	virtual const UObject* AdditionalStatObject() const override;
+#if WITH_EDITOR
+	virtual void CheckForErrors() override;
+#endif
 	// End of UActorComponent interface
 
 	// USceneComponent interface
diff --git a/Engine/Plugins/2D/Paper2D/Source/Paper2D/Classes/PaperSpriteComponent.h b/Engine/Plugins/2D/Paper2D/Source/Paper2D/Classes/PaperSpriteComponent.h
--- a/Engine/Plugins/2D/Paper2D/Source/Paper2D/Classes/PaperSpriteComponent.h
+++ b/Engine/Plugins/2D/Paper2D/Source/Paper2D/Classes/PaperSpriteComponent.h
@@ -63,6 +63,9 @@ public:
 	// UActorComponent interface
 	virtual void SendRenderDynamicData_Concurrent() override;
 	virtual const UObject* AdditionalStatObject() const override;
+#if WITH_EDITOR
+	virtual void CheckForErrors() override;
+#endif
 	// End of UActorComponent interface
 
 	// USceneComponent interface
diff --git a/Engine/Plugins/2D/Paper2D/Source/Paper2D/Private/PaperFlipbookComponent.cpp b/Engine/Plugins/2D/Paper2D/Source/Paper2D/Private/PaperFlipbookComponent.cpp
--- a/Engine/Plugins/2D/Paper2D/Source/Paper2D/Private/PaperFlipbookComponent.cpp
+++ b/Engine/Plugins/2D/Paper2D/Source/Paper2D/Private/PaperFlipbookComponent.cpp
@@ -6,6 +6,11 @@
 #include "PaperCustomVersion.h"
 #include "Runtime/Engine/Public/Net/UnrealNetwork.h"
 #include "Runtime/Engine/Public/ContentStreaming.h"
+#include "Runtime/Core/Public/Logging/MessageLog.h"
+#include "Runtime/Core/Public/Misc/MapErrors.h"
+#include "Runtime/CoreUObject/Public/Misc/UObjectToken.h"
+
+#define LOCTEXT_NAMESPACE "Paper2D"
 
 //////////////////////////////////////////////////////////////////////////
 // UPaperFlipbookComponent
@@ -595,3 +600,31 @@ void UPaperFlipbookComponent::QuerySupportedSockets(TArray<FComponentSocketDescr
 		return SourceFlipbook->QuerySupportedSockets(OutSockets);
 	}
 }
+
+#if WITH_EDITOR
+void UPaperFlipbookComponent::CheckForErrors()
+{
+	Super::CheckForErrors();
+
+	AActor* Owner = GetOwner();
+
+	for (int32 MaterialIndex = 0; MaterialIndex < GetNumMaterials(); ++MaterialIndex)
+	{
+		if (UMaterialInterface* Material = GetMaterial(MaterialIndex))
+		{
+			if (!Material->IsTwoSided())
+			{
+				FMessageLog("MapCheck").Warning()
+					->AddToken(FUObjectToken::Create(Owner))
+					->AddToken(FTextToken::Create(LOCTEXT("MapCheck_Message_PaperFlipbookMaterialNotTwoSided", "The material applied to the flipbook component is not marked as two-sided, which may cause lighting artifacts.")))
+					->AddToken(FUObjectToken::Create(Material))
+					->AddToken(FMapErrorToken::Create(FName(TEXT("PaperFlipbookMaterialNotTwoSided"))));
+			}
+		}
+	}
+}
+#endif
+
+//////////////////////////////////////////////////////////////////////////
+
+#undef LOCTEXT_NAMESPACE
diff --git a/Engine/Plugins/2D/Paper2D/Source/Paper2D/Private/PaperSpriteComponent.cpp b/Engine/Plugins/2D/Paper2D/Source/Paper2D/Private/PaperSpriteComponent.cpp
--- a/Engine/Plugins/2D/Paper2D/Source/Paper2D/Private/PaperSpriteComponent.cpp
+++ b/Engine/Plugins/2D/Paper2D/Source/Paper2D/Private/PaperSpriteComponent.cpp
@@ -6,6 +6,11 @@
 #include "PaperCustomVersion.h"
 #include "Runtime/Engine/Classes/PhysicsEngine/BodySetup2D.h"
 #include "Runtime/Engine/Public/ContentStreaming.h"
+#include "Runtime/Core/Public/Logging/MessageLog.h"
+#include "Runtime/Core/Public/Misc/MapErrors.h"
+#include "Runtime/CoreUObject/Public/Misc/UObjectToken.h"
+
+#define LOCTEXT_NAMESPACE "Paper2D"
 
 //////////////////////////////////////////////////////////////////////////
 // UPaperSpriteComponent
@@ -297,3 +302,31 @@ const UObject* UPaperSpriteComponent::AdditionalStatObject() const
 {
 	return SourceSprite;
 }
+
+#if WITH_EDITOR
+void UPaperSpriteComponent::CheckForErrors()
+{
+	Super::CheckForErrors();
+
+	AActor* Owner = GetOwner();
+
+	for (int32 MaterialIndex = 0; MaterialIndex < GetNumMaterials(); ++MaterialIndex)
+	{
+		if (UMaterialInterface* Material = GetMaterial(MaterialIndex))
+		{
+			if (!Material->IsTwoSided())
+			{
+				FMessageLog("MapCheck").Warning()
+					->AddToken(FUObjectToken::Create(Owner))
+					->AddToken(FTextToken::Create(LOCTEXT("MapCheck_Message_PaperSpriteMaterialNotTwoSided", "The material applied to the sprite component is not marked as two-sided, which may cause lighting artifacts.")))
+					->AddToken(FUObjectToken::Create(Material))
+					->AddToken(FMapErrorToken::Create(FName(TEXT("PaperSpriteMaterialNotTwoSided"))));
+			}
+		}
+	}
+}
+#endif
+
+//////////////////////////////////////////////////////////////////////////
+
+#undef LOCTEXT_NAMESPACE
\ No newline at end of file
diff --git a/Engine/Source/Editor/UnrealEd/Private/EditorServer.cpp b/Engine/Source/Editor/UnrealEd/Private/EditorServer.cpp
--- a/Engine/Source/Editor/UnrealEd/Private/EditorServer.cpp
+++ b/Engine/Source/Editor/UnrealEd/Private/EditorServer.cpp
@@ -3477,6 +3477,10 @@ bool UEditorEngine::Map_Check( UWorld* InWorld, const TCHAR* Str, FOutputDevice&
 			{
 				MapCheckLog.Notify(LOCTEXT("MapCheckGenErrors", "Map check generated errors!"));
 			}
+			else if (WarningCount > 0)
+			{
+				MapCheckLog.Notify(LOCTEXT("MapCheckGenWarnings", "Map check generated warnings!"));
+			}
 		}
 		else
 		{
@@ -3486,10 +3490,14 @@ bool UEditorEngine::Map_Check( UWorld* InWorld, const TCHAR* Str, FOutputDevice&
 			}
 			else if(Notification == EMapCheckNotification::NotifyOfResults)
 			{				
-				if(WarningCount > 0)
+				if (ErrorCount > 0)
 				{
 					MapCheckLog.Notify(LOCTEXT("MapCheckFoundErrors", "Map Check found some errors!"));
 				}
+				else if (WarningCount > 0)
+				{
+					MapCheckLog.Notify(LOCTEXT("MapCheckFoundWarnings", "Map Check found some issues!"));
+				}
 				else
 				{
 					// Nothing to notify about.  Everything went fine!
