Paper2D: Add the ability to edit the default collision profile for new actor instances in the sprite asset editor

diff --git a/Engine/Plugins/2D/Paper2D/Source/Paper2D/Classes/PaperSprite.h b/Engine/Plugins/2D/Paper2D/Source/Paper2D/Classes/PaperSprite.h
--- a/Engine/Plugins/2D/Paper2D/Source/Paper2D/Classes/PaperSprite.h
+++ b/Engine/Plugins/2D/Paper2D/Source/Paper2D/Classes/PaperSprite.h
@@ -105,7 +105,7 @@ protected:
 
 public:
 	// Baked physics data.
-	UPROPERTY() //@TODO: Is anything in this worth exposing to edit? Category=Physics, EditAnywhere, EditInline
+	UPROPERTY(EditAnywhere, Category=NeverShown)
 	class UBodySetup* BodySetup;
 
 #if WITH_EDITORONLY_DATA
diff --git a/Engine/Plugins/2D/Paper2D/Source/Paper2D/Private/PaperCustomVersion.h b/Engine/Plugins/2D/Paper2D/Source/Paper2D/Private/PaperCustomVersion.h
--- a/Engine/Plugins/2D/Paper2D/Source/Paper2D/Private/PaperCustomVersion.h
+++ b/Engine/Plugins/2D/Paper2D/Source/Paper2D/Private/PaperCustomVersion.h
@@ -40,6 +40,9 @@ struct FPaperCustomVersion
 		// Converted UPaperFlipbookComponent to be derived from UMeshComponent
 		ConvertPaperFlipbookComponentToBeMeshComponent = 10,
 
+		// Add default BodyInstance setup in the sprite asset
+		AddDefaultCollisionProfileInSpriteAsset = 11,
+
 		// -----<new versions can be added above this line>-------------------------------------------------
 		VersionPlusOne,
 		LatestVersion = VersionPlusOne - 1
diff --git a/Engine/Plugins/2D/Paper2D/Source/Paper2D/Private/PaperSprite.cpp b/Engine/Plugins/2D/Paper2D/Source/Paper2D/Private/PaperSprite.cpp
--- a/Engine/Plugins/2D/Paper2D/Source/Paper2D/Private/PaperSprite.cpp
+++ b/Engine/Plugins/2D/Paper2D/Source/Paper2D/Private/PaperSprite.cpp
@@ -374,6 +374,8 @@ void UPaperSprite::PostEditChangeProperty(FPropertyChangedEvent& PropertyChanged
 
 void UPaperSprite::RebuildCollisionData()
 {
+	UBodySetup* OldBodySetup = BodySetup;
+
 	// Ensure we have the data structure for the desired collision method
 	switch (SpriteCollisionDomain)
 	{
@@ -422,6 +424,16 @@ void UPaperSprite::RebuildCollisionData()
 		default:
 			check(false); // unknown mode
 		};
+
+		// Copy across or initialize the only editable property we expose on the body setup
+		if (OldBodySetup != nullptr)
+		{
+			BodySetup->DefaultInstance.CopyBodyInstancePropertiesFrom(&(OldBodySetup->DefaultInstance));
+		}
+		else
+		{
+			BodySetup->DefaultInstance.SetCollisionProfileName(UCollisionProfile::BlockAllDynamic_ProfileName);
+		}
 	}
 }
 
@@ -1346,6 +1358,11 @@ void UPaperSprite::PostLoad()
 	{
 		bRebuildCollision = true;
 	}
+	
+	if ((PaperVer < FPaperCustomVersion::AddDefaultCollisionProfileInSpriteAsset) && (BodySetup != nullptr))
+	{
+		BodySetup->DefaultInstance.SetCollisionProfileName(UCollisionProfile::BlockAllDynamic_ProfileName);
+	}
 
 	if (bRebuildCollision)
 	{
diff --git a/Engine/Plugins/2D/Paper2D/Source/Paper2DEditor/Private/PaperSpriteActorFactory.cpp b/Engine/Plugins/2D/Paper2D/Source/Paper2DEditor/Private/PaperSpriteActorFactory.cpp
--- a/Engine/Plugins/2D/Paper2D/Source/Paper2DEditor/Private/PaperSpriteActorFactory.cpp
+++ b/Engine/Plugins/2D/Paper2D/Source/Paper2DEditor/Private/PaperSpriteActorFactory.cpp
@@ -25,6 +25,12 @@ void UPaperSpriteActorFactory::PostSpawnActor(UObject* Asset, AActor* NewActor)
 
 		RenderComponent->UnregisterComponent();
 		RenderComponent->SetSprite(Sprite);
+
+		if (Sprite->BodySetup != nullptr)
+		{
+			RenderComponent->BodyInstance.CopyBodyInstancePropertiesFrom(&(Sprite->BodySetup->DefaultInstance));
+		}
+
 		RenderComponent->RegisterComponent();
 	}
 }
@@ -39,6 +45,11 @@ void UPaperSpriteActorFactory::PostCreateBlueprint(UObject* Asset, AActor* CDO)
 			check(RenderComponent);
 
 			RenderComponent->SetSprite(Sprite);
+
+			if (Sprite->BodySetup != nullptr)
+			{
+				RenderComponent->BodyInstance.CopyBodyInstancePropertiesFrom(&(Sprite->BodySetup->DefaultInstance));
+			}
 		}
 	}
 }
diff --git a/Engine/Plugins/2D/Paper2D/Source/Paper2DEditor/Private/PaperSpriteAssetBroker.h b/Engine/Plugins/2D/Paper2D/Source/Paper2DEditor/Private/PaperSpriteAssetBroker.h
--- a/Engine/Plugins/2D/Paper2D/Source/Paper2DEditor/Private/PaperSpriteAssetBroker.h
+++ b/Engine/Plugins/2D/Paper2D/Source/Paper2DEditor/Private/PaperSpriteAssetBroker.h
@@ -17,13 +17,19 @@ public:
 
 	virtual bool AssignAssetToComponent(UActorComponent* InComponent, UObject* InAsset) override
 	{
-		if (UPaperSpriteComponent* RenderComp = Cast<UPaperSpriteComponent>(InComponent))
+		if (UPaperSpriteComponent* RenderComponent = Cast<UPaperSpriteComponent>(InComponent))
 		{
 			UPaperSprite* Sprite = Cast<UPaperSprite>(InAsset);
 
-			if ((Sprite != NULL) || (InAsset == NULL))
+			if ((Sprite != nullptr) || (InAsset == nullptr))
 			{
-				RenderComp->SetSprite(Sprite);
+				RenderComponent->SetSprite(Sprite);
+
+				if (Sprite->BodySetup != nullptr)
+				{
+					RenderComponent->BodyInstance.CopyBodyInstancePropertiesFrom(&(Sprite->BodySetup->DefaultInstance));
+				}
+
 				return true;
 			}
 		}
@@ -33,11 +39,11 @@ public:
 
 	virtual UObject* GetAssetFromComponent(UActorComponent* InComponent) override
 	{
-		if (UPaperSpriteComponent* RenderComp = Cast<UPaperSpriteComponent>(InComponent))
+		if (UPaperSpriteComponent* RenderComponent = Cast<UPaperSpriteComponent>(InComponent))
 		{
-			return RenderComp->GetSprite();
+			return RenderComponent->GetSprite();
 		}
-		return NULL;
+		return nullptr;
 	}
 };
 
diff --git a/Engine/Plugins/2D/Paper2D/Source/Paper2DEditor/Private/SpriteEditor/SpriteDetailsCustomization.cpp b/Engine/Plugins/2D/Paper2D/Source/Paper2DEditor/Private/SpriteEditor/SpriteDetailsCustomization.cpp
--- a/Engine/Plugins/2D/Paper2D/Source/Paper2DEditor/Private/SpriteEditor/SpriteDetailsCustomization.cpp
+++ b/Engine/Plugins/2D/Paper2D/Source/Paper2DEditor/Private/SpriteEditor/SpriteDetailsCustomization.cpp
@@ -146,6 +146,34 @@ void FSpriteDetailsCustomization::BuildCollisionSection(IDetailCategoryBuilder&
 	CollisionCategory.AddProperty(DetailLayout.GetProperty(*CollisionGeometryPolygonsPropertyPath), EPropertyLocation::Advanced)
 		.DisplayName(LOCTEXT("CollisionPolygons", "Collision Polygons").ToString())
 		.Visibility(ParticipatesInPhysics);
+
+	// Show the default body instance (and only it) from the body setup (if it exists)
+	DetailLayout.HideProperty("BodySetup");
+	IDetailPropertyRow& BodySetupDefaultInstance = CollisionCategory.AddProperty("BodySetup.DefaultInstance");
+	
+	TArray<TWeakObjectPtr<UObject>> SpritesBeingEdited;
+	DetailLayout.GetObjectsBeingCustomized(/*out*/ SpritesBeingEdited);
+
+	TArray<UObject*> BodySetupList;
+	for (auto WeakSpritePtr : SpritesBeingEdited)
+	{
+		if (UPaperSprite* Sprite = Cast<UPaperSprite>(WeakSpritePtr.Get()))
+		{
+			if (UBodySetup* BodySetup = Sprite->BodySetup)
+			{
+				BodySetupList.Add(BodySetup);
+			}
+		}
+	}
+	
+	if (BodySetupList.Num() > 0)
+	{
+		IDetailPropertyRow* DefaultInstanceRow = CollisionCategory.AddExternalProperty(BodySetupList, GET_MEMBER_NAME_CHECKED(UBodySetup, DefaultInstance));
+		if (DefaultInstanceRow != nullptr)
+		{
+			DefaultInstanceRow->Visibility(ParticipatesInPhysics);
+		}
+	}
 }
 
 EVisibility FSpriteDetailsCustomization::PhysicsModeMatches(TSharedPtr<IPropertyHandle> Property, ESpriteCollisionMode::Type DesiredMode) const
