Paper2D: Fixed edge selection when editing geometry

diff --git a/Engine/Plugins/2D/Paper2D/Source/Paper2DEditor/Private/SpriteEditor/SpriteEditorViewportClient.cpp b/Engine/Plugins/2D/Paper2D/Source/Paper2DEditor/Private/SpriteEditor/SpriteEditorViewportClient.cpp
--- a/Engine/Plugins/2D/Paper2D/Source/Paper2DEditor/Private/SpriteEditor/SpriteEditorViewportClient.cpp
+++ b/Engine/Plugins/2D/Paper2D/Source/Paper2DEditor/Private/SpriteEditor/SpriteEditorViewportClient.cpp
@@ -401,9 +401,24 @@ void FSpriteEditorViewportClient::DrawGeometry(FViewport& InViewport, FSceneView
 
 			// Draw the edge
 			{
+				if (bIsHitTesting)
+				{
+					TSharedPtr<FSpriteSelectedEdge> Data = MakeShareable(new FSpriteSelectedEdge());
+					Data->SpritePtr = Sprite;
+					Data->PolygonIndex = PolygonIndex;
+					Data->VertexIndex = VertexIndex;
+					Data->bRenderData = bIsRenderGeometry;
+					Canvas.SetHitProxy(new HSpriteSelectableObjectHitProxy(Data));
+				}
+
 				FCanvasLineItem LineItem(ScreenPos, NextScreenPos);
 				LineItem.SetColor(IsEdgeSelected ? SelectedColor : LineColor);
 				Canvas.DrawItem(LineItem);
+
+				if (bIsHitTesting)
+				{
+					Canvas.SetHitProxy(NULL);
+				}
 			}
 		}
 
@@ -993,7 +1008,15 @@ void FSpriteEditorViewportClient::ProcessClick(FSceneView& View, HHitProxy* HitP
 		{
 			if (const FSpriteSelectedVertex* SelectedVertex = SelectedItemProxy->Data->CastSelectedVertex())
 			{
-				AddPolygonVertexToSelection(SelectedVertex->PolygonIndex, SelectedVertex->VertexIndex);
+				if (SelectedVertex->IsA(FSelectionTypes::Edge))
+				{
+					// Add the next vertex defined by this edge
+					AddPolygonEdgeToSelection(SelectedVertex->PolygonIndex, SelectedVertex->VertexIndex);
+				}
+				else
+				{
+					AddPolygonVertexToSelection(SelectedVertex->PolygonIndex, SelectedVertex->VertexIndex);
+				}
 			}
 			else
 			{
@@ -1605,7 +1628,7 @@ void FSpriteEditorViewportClient::AddPolygonVertexToSelection(const int32 Polygo
 {
 	UPaperSprite* Sprite = GetSpriteBeingEdited();
 	FSpritePolygonCollection* Geometry = GetGeometryBeingEdited();
-	if (ensure(Geometry != nullptr &&  PolygonIndex >= 0 && PolygonIndex < Geometry->Polygons.Num()))
+	if (ensure(Geometry != nullptr && Geometry->Polygons.IsValidIndex(PolygonIndex)))
 	{
 		FSpritePolygon& Polygon = Geometry->Polygons[PolygonIndex];
 		if (Polygon.Vertices.IsValidIndex(VertexIndex))
@@ -1625,6 +1648,19 @@ void FSpriteEditorViewportClient::AddPolygonVertexToSelection(const int32 Polygo
 	}
 }
 
+void FSpriteEditorViewportClient::AddPolygonEdgeToSelection(const int32 PolygonIndex, const int32 FirstVertexIndex)
+{
+	FSpritePolygonCollection* Geometry = GetGeometryBeingEdited();
+	if (ensure(Geometry != nullptr && Geometry->Polygons.IsValidIndex(PolygonIndex)))
+	{
+		FSpritePolygon& Polygon = Geometry->Polygons[PolygonIndex];
+		int NextVertexIndex = (FirstVertexIndex + 1) % Polygon.Vertices.Num();
+
+		AddPolygonVertexToSelection(PolygonIndex, FirstVertexIndex);
+		AddPolygonVertexToSelection(PolygonIndex, NextVertexIndex);
+	}
+}
+
 
 void FSpriteEditorViewportClient::ClearSelectionSet()
 {
diff --git a/Engine/Plugins/2D/Paper2D/Source/Paper2DEditor/Private/SpriteEditor/SpriteEditorViewportClient.h b/Engine/Plugins/2D/Paper2D/Source/Paper2DEditor/Private/SpriteEditor/SpriteEditorViewportClient.h
--- a/Engine/Plugins/2D/Paper2D/Source/Paper2DEditor/Private/SpriteEditor/SpriteEditorViewportClient.h
+++ b/Engine/Plugins/2D/Paper2D/Source/Paper2DEditor/Private/SpriteEditor/SpriteEditorViewportClient.h
@@ -229,6 +229,7 @@ private:
 	// Selection handling
 	void SelectPolygon(const int32 PolygonIndex);
 	void AddPolygonVertexToSelection(const int32 PolygonIndex, const int32 VertexIndex);
+	void AddPolygonEdgeToSelection(const int32 PolygonIndex, const int32 FirstVertexIndex);
 	bool IsPolygonVertexSelected(const int32 PolygonIndex, const int32 VertexIndex) const;
 	void AddPointToGeometry(const FVector2D& TextureSpacePoint, const int32 SelectedPolygonIndex = -1);
 	void ClearSelectionSet();
