Paper2D: Remove incorrect bWillEverBeLit=false in 2D sprite render scene proxies

diff --git a/Engine/Plugins/2D/Paper2D/Source/Paper2D/Private/PaperRenderSceneProxy.cpp b/Engine/Plugins/2D/Paper2D/Source/Paper2D/Private/PaperRenderSceneProxy.cpp
--- a/Engine/Plugins/2D/Paper2D/Source/Paper2D/Private/PaperRenderSceneProxy.cpp
+++ b/Engine/Plugins/2D/Paper2D/Source/Paper2D/Private/PaperRenderSceneProxy.cpp
@@ -164,7 +164,6 @@ FPaperRenderSceneProxy::FPaperRenderSceneProxy(const UPrimitiveComponent* InComp
 	, CollisionResponse(InComponent->GetCollisionResponseToChannels())
 {
 	WireframeColor = FLinearColor::White;
-	bWillEverBeLit = false;
 }
 
 FPaperRenderSceneProxy::~FPaperRenderSceneProxy()
diff --git a/Engine/Plugins/2D/Paper2D/Source/Paper2D/Private/Rendering/PaperBatchSceneProxy.cpp b/Engine/Plugins/2D/Paper2D/Source/Paper2D/Private/Rendering/PaperBatchSceneProxy.cpp
--- a/Engine/Plugins/2D/Paper2D/Source/Paper2D/Private/Rendering/PaperBatchSceneProxy.cpp
+++ b/Engine/Plugins/2D/Paper2D/Source/Paper2D/Private/Rendering/PaperBatchSceneProxy.cpp
@@ -10,7 +10,6 @@
 FPaperBatchSceneProxy::FPaperBatchSceneProxy(const UPrimitiveComponent* InComponent)
 	: FPrimitiveSceneProxy(InComponent)
 {
-	bWillEverBeLit = false;
 }
 
 uint32 FPaperBatchSceneProxy::GetMemoryFootprint() const
