Removed a bunch of stuff from Slate standard include, created SlateBasics.h
* Moved Slate.h into SlateBasics.h and began shifting less commonly used headers into SlateExtras.h.
* Slate.h now simply includes SlateBasics.h and SlateExtras.h.
* Slate.h includes a deprecated warning now to indicate that SlateBasics.h + specific includes should be used instead.
* Moved dozens of inlined functions using Slate widgets into .cpp files to avoid header dependencies.
* All code samples now include SlateBasics.h and SlateExtras.h so future shifts will not break most those projects, but not trigger the deprecation warning of including Slate.h.
#BUN

diff --git a/Engine/Plugins/2D/Paper2D/Source/Paper2DEditor/Private/FlipbookEditor/FlipbookEditor.cpp b/Engine/Plugins/2D/Paper2D/Source/Paper2DEditor/Private/FlipbookEditor/FlipbookEditor.cpp
--- a/Engine/Plugins/2D/Paper2D/Source/Paper2DEditor/Private/FlipbookEditor/FlipbookEditor.cpp
+++ b/Engine/Plugins/2D/Paper2D/Source/Paper2DEditor/Private/FlipbookEditor/FlipbookEditor.cpp
@@ -15,6 +15,8 @@
 #include "SFlipbookEditorViewportToolbar.h"
 #include "Editor/KismetWidgets/Public/SScrubControlPanel.h"
 #include "SFlipbookTimeline.h"
+#include "SDockTab.h"
+#include "GenericCommands.h"
 
 #define LOCTEXT_NAMESPACE "FlipbookEditor"
 
diff --git a/Engine/Plugins/2D/Paper2D/Source/Paper2DEditor/Private/FlipbookEditor/SFlipbookTrackHandle.h b/Engine/Plugins/2D/Paper2D/Source/Paper2DEditor/Private/FlipbookEditor/SFlipbookTrackHandle.h
--- a/Engine/Plugins/2D/Paper2D/Source/Paper2DEditor/Private/FlipbookEditor/SFlipbookTrackHandle.h
+++ b/Engine/Plugins/2D/Paper2D/Source/Paper2DEditor/Private/FlipbookEditor/SFlipbookTrackHandle.h
@@ -1,5 +1,7 @@
 // Copyright 1998-2014 Epic Games, Inc. All Rights Reserved.
 
+#pragma once
+
 //////////////////////////////////////////////////////////////////////////
 // SFlipbookTrackHandle
 
diff --git a/Engine/Plugins/2D/Paper2D/Source/Paper2DEditor/Private/FlipbookEditor/STimelineTrack.h b/Engine/Plugins/2D/Paper2D/Source/Paper2DEditor/Private/FlipbookEditor/STimelineTrack.h
--- a/Engine/Plugins/2D/Paper2D/Source/Paper2DEditor/Private/FlipbookEditor/STimelineTrack.h
+++ b/Engine/Plugins/2D/Paper2D/Source/Paper2DEditor/Private/FlipbookEditor/STimelineTrack.h
@@ -1,5 +1,10 @@
 // Copyright 1998-2014 Epic Games, Inc. All Rights Reserved.
 
+#pragma once 
+
+#include "ScopedTransaction.h"
+#include "SFlipbookTimeline.h"
+
 namespace FFlipbookUIConstants
 {
 	const float HandleWidth = 12.0f;
@@ -21,91 +26,27 @@ public:
 	TWeakObjectPtr<UPaperFlipbook> SourceFlipbook;
 	FScopedTransaction Transaction;
 
-	virtual TSharedPtr<SWidget> GetDefaultDecorator() const override
-	{
-		const FLinearColor BorderColor = (KeyFrameData.Sprite != nullptr) ? FLinearColor::White : FLinearColor::Black;
-
-		return SNew(SBox)
-			.WidthOverride(WidgetWidth - FFlipbookUIConstants::FramePadding.GetTotalSpaceAlong<Orient_Horizontal>())
-			.HeightOverride(FFlipbookUIConstants::FrameHeight - FFlipbookUIConstants::FramePadding.GetTotalSpaceAlong<Orient_Vertical>())
-			[
-				SNew(SBorder)
-				.BorderImage(FPaperStyle::Get()->GetBrush("FlipbookEditor.RegionBody"))
-				.BorderBackgroundColor(BorderColor)
-				[
-					SNullWidget::NullWidget
-				]
-			];
-	}
+	virtual TSharedPtr<SWidget> GetDefaultDecorator() const override;
 
-	virtual void OnDragged(const class FDragDropEvent& DragDropEvent) override
-	{
-		if (CursorDecoratorWindow.IsValid())
-		{
-			CursorDecoratorWindow->MoveWindowTo(DragDropEvent.GetScreenSpacePosition());
-		}
-	}
+	virtual void OnDragged(const class FDragDropEvent& DragDropEvent) override;
 
-	virtual void Construct() override
-	{
-		MouseCursor = EMouseCursor::GrabHandClosed;
-		FDragDropOperation::Construct();
-	}
+	virtual void Construct() override;
 
-	virtual void OnDrop(bool bDropWasHandled, const FPointerEvent& MouseEvent) override
-	{
-		if (!bDropWasHandled)
-		{
-			// Add us back to our source, the drop fizzled
-			InsertInFlipbook(SourceFlipbook.Get(), SourceFrameIndex);
-			Transaction.Cancel();
-		}
-	}
+	virtual void OnDrop(bool bDropWasHandled, const FPointerEvent& MouseEvent) override;
 
-	void AppendToFlipbook(UPaperFlipbook* DestinationFlipbook)
-	{
-		DestinationFlipbook->Modify();
-		FScopedFlipbookMutator EditLock(DestinationFlipbook);
-		EditLock.KeyFrames.Add(KeyFrameData);
-	}
+	void AppendToFlipbook(UPaperFlipbook* DestinationFlipbook);
 
-	void InsertInFlipbook(UPaperFlipbook* DestinationFlipbook, int32 Index)
-	{
-		DestinationFlipbook->Modify();
-		FScopedFlipbookMutator EditLock(DestinationFlipbook);
-		EditLock.KeyFrames.Insert(KeyFrameData, Index);
-	}
+	void InsertInFlipbook(UPaperFlipbook* DestinationFlipbook, int32 Index);
 
 	void SetCanDropHere(bool bCanDropHere)
 	{
 		MouseCursor = bCanDropHere ? EMouseCursor::TextEditBeam : EMouseCursor::SlashedCircle;
 	}
 
-	static TSharedRef<FFlipbookKeyFrameDragDropOp> New(int32 InWidth, UPaperFlipbook* InFlipbook, int32 InFrameIndex)
-	{
-		// Create the drag-drop op containing the key
-		TSharedRef<FFlipbookKeyFrameDragDropOp> Operation = MakeShareable(new FFlipbookKeyFrameDragDropOp);
-		Operation->KeyFrameData = InFlipbook->GetKeyFrameChecked(InFrameIndex);
-		Operation->SourceFrameIndex = InFrameIndex;
-		Operation->SourceFlipbook = InFlipbook;
-		Operation->WidgetWidth = InWidth;
-		Operation->Construct();
-
-		// Remove the key from the flipbook
-		{
-			InFlipbook->Modify();
-			FScopedFlipbookMutator EditLock(InFlipbook);
-			EditLock.KeyFrames.RemoveAt(InFrameIndex);
-		}
-
-		return Operation;
-	}
+	static TSharedRef<FFlipbookKeyFrameDragDropOp> New(int32 InWidth, UPaperFlipbook* InFlipbook, int32 InFrameIndex);
 
 protected:
-	FFlipbookKeyFrameDragDropOp()
-		: Transaction(LOCTEXT("MovedFramesInTimeline", "Reorder key frames"))
-	{
-	}
+	FFlipbookKeyFrameDragDropOp();
 };
 
 //////////////////////////////////////////////////////////////////////////
@@ -125,184 +66,21 @@ public:
 
 	SLATE_END_ARGS()
 
-	void Construct(const FArguments& InArgs, int32 InFrameIndex, TSharedPtr<const FUICommandList> InCommandList)
-	{
-		FrameIndex = InFrameIndex;
-		CommandList = InCommandList;
-		SlateUnitsPerFrame = InArgs._SlateUnitsPerFrame;
-		FlipbookBeingEdited = InArgs._FlipbookBeingEdited;
-		OnSelectionChanged = InArgs._OnSelectionChanged;
-
-		// Color each region based on whether a sprite has been set or not for it
-		const auto BorderColorDelegate = [](TAttribute<UPaperFlipbook*> ThisFlipbookPtr, int32 TestIndex) -> FSlateColor
-		{
-			UPaperFlipbook* FlipbookPtr = ThisFlipbookPtr.Get();
-			const bool bFrameValid = (FlipbookPtr != nullptr) && (FlipbookPtr->GetSpriteAtFrame(TestIndex) != nullptr);
-			return bFrameValid ? FLinearColor::White : FLinearColor::Black;
-		};
-
-		ChildSlot
-		[
-			SNew(SHorizontalBox)
-			+SHorizontalBox::Slot()
-			.AutoWidth()
-			[
-				SNew(SBox)
-				.Padding(FFlipbookUIConstants::FramePadding)
-				.WidthOverride(this, &SFlipbookKeyframeWidget::GetFrameWidth)
-				[
-					SNew(SBorder)
-					.BorderImage(FPaperStyle::Get()->GetBrush("FlipbookEditor.RegionBody"))
-					.BorderBackgroundColor_Static(BorderColorDelegate, FlipbookBeingEdited, FrameIndex)
-					.OnMouseButtonUp(this, &SFlipbookKeyframeWidget::KeyframeOnMouseButtonUp)
-					.ToolTipText(this, &SFlipbookKeyframeWidget::GetKeyframeTooltip)
-					[
-						SNullWidget::NullWidget
-					]
-				]
-			]
-			+SHorizontalBox::Slot()
-			.AutoWidth()
-			[
-				SNew(SBox)
-				.WidthOverride(FFlipbookUIConstants::HandleWidth)
-				[
-					SNew(SFlipbookTrackHandle)
-					.SlateUnitsPerFrame(SlateUnitsPerFrame)
-					.FlipbookBeingEdited(FlipbookBeingEdited)
-					.KeyFrameIdx(FrameIndex)
-				]
-			]
-		];
-	}
-
-	virtual FReply OnMouseButtonDown(const FGeometry& MyGeometry, const FPointerEvent& MouseEvent) override
-	{
-		if (MouseEvent.GetEffectingButton() == EKeys::LeftMouseButton)
-		{
-			return FReply::Handled().DetectDrag(SharedThis(this), EKeys::LeftMouseButton);
-		}
-
-		return FReply::Unhandled();
-	}
-
-	virtual FReply OnDragDetected(const FGeometry& MyGeometry, const FPointerEvent& MouseEvent) override
-	{
-		if (MouseEvent.IsMouseButtonDown(EKeys::LeftMouseButton))
-		{
-			UPaperFlipbook* Flipbook = FlipbookBeingEdited.Get();
-			if ((Flipbook != nullptr) && Flipbook->IsValidKeyFrameIndex(FrameIndex))
-			{
-				TSharedRef<FFlipbookKeyFrameDragDropOp> Operation = FFlipbookKeyFrameDragDropOp::New(
-					GetFrameWidth().Get(), Flipbook, FrameIndex);
-
-				return FReply::Handled().BeginDragDrop(Operation);
-			}
-		}
-
-		return FReply::Unhandled();
-	}
-
-	FReply OnDrop(const FGeometry& MyGeometry, const FDragDropEvent& DragDropEvent) override
-	{
-		bool bWasDropHandled = false;
-
-		UPaperFlipbook* Flipbook = FlipbookBeingEdited.Get();
-		if ((Flipbook != nullptr) && Flipbook->IsValidKeyFrameIndex(FrameIndex))
-		{
-
-			TSharedPtr<FDragDropOperation> Operation = DragDropEvent.GetOperation();
-			if (!Operation.IsValid())
-			{
-			}
-			else if (Operation->IsOfType<FAssetDragDropOp>())
-			{
-				const auto& AssetDragDropOp = StaticCastSharedPtr<FAssetDragDropOp>(Operation);
-				//@TODO: Handle asset inserts
-
-			// 			OnAssetsDropped(*AssetDragDropOp);
-			// 			bWasDropHandled = true;
-			}
-			else if (Operation->IsOfType<FFlipbookKeyFrameDragDropOp>())
-			{
-				const auto& FrameDragDropOp = StaticCastSharedPtr<FFlipbookKeyFrameDragDropOp>(Operation);
-				FrameDragDropOp->InsertInFlipbook(Flipbook, FrameIndex);
-				bWasDropHandled = true;
-			}
-		}
-
-		return bWasDropHandled ? FReply::Handled() : FReply::Unhandled();
-	}
-
-	FReply KeyframeOnMouseButtonUp(const FGeometry& MyGeometry, const FPointerEvent& MouseEvent)
-	{
-		if (MouseEvent.GetEffectingButton() == EKeys::RightMouseButton)
-		{
-			TSharedRef<SWidget> MenuContents = GenerateContextMenu();
-			FSlateApplication::Get().PushMenu(AsShared(), MenuContents, MouseEvent.GetScreenSpacePosition(), FPopupTransitionEffect(FPopupTransitionEffect::ContextMenu));
-
-			return FReply::Handled();
-		}
-
-		return FReply::Unhandled();
-	}
-
-	FText GetKeyframeTooltip() const
-	{
-		UPaperFlipbook* Flipbook = FlipbookBeingEdited.Get();
-		if ((Flipbook != nullptr) && Flipbook->IsValidKeyFrameIndex(FrameIndex))
-		{
-			const FPaperFlipbookKeyFrame& KeyFrame = Flipbook->GetKeyFrameChecked(FrameIndex);
-
-			FText SpriteLine = (KeyFrame.Sprite != nullptr) ? FText::FromString(KeyFrame.Sprite->GetName()) : LOCTEXT("NoSprite", "(none)");
-
-			return FText::Format(LOCTEXT("KeyFrameTooltip", "Sprite: {0}\nIndex: {1}\nDuration: {2} frame(s)"),
-				SpriteLine,
-				FText::AsNumber(FrameIndex),
-				FText::AsNumber(KeyFrame.FrameRun));
-		}
-		else
-		{
-			return LOCTEXT("KeyFrameTooltip_Invalid", "Invalid key frame index");
-		}
-	}
-
-	TSharedRef<SWidget> GenerateContextMenu()
-	{
-		OnSelectionChanged.ExecuteIfBound(FrameIndex);
+	void Construct(const FArguments& InArgs, int32 InFrameIndex, TSharedPtr<const FUICommandList> InCommandList);
 
-		FMenuBuilder MenuBuilder(true, CommandList);
-		MenuBuilder.BeginSection("KeyframeActions", LOCTEXT("KeyframeActionsSectionHeader", "Keyframe Actions"));
+	virtual FReply OnMouseButtonDown(const FGeometry& MyGeometry, const FPointerEvent& MouseEvent) override;
 
-		// 		MenuBuilder.AddMenuEntry(FGenericCommands::Get().Cut);
-		// 		MenuBuilder.AddMenuEntry(FGenericCommands::Get().Copy);
-		// 		MenuBuilder.AddMenuEntry(FGenericCommands::Get().Paste);
-		MenuBuilder.AddMenuEntry(FGenericCommands::Get().Duplicate);
-		MenuBuilder.AddMenuEntry(FGenericCommands::Get().Delete);
+	virtual FReply OnDragDetected(const FGeometry& MyGeometry, const FPointerEvent& MouseEvent) override;
 
-		MenuBuilder.AddMenuSeparator();
+	FReply OnDrop(const FGeometry& MyGeometry, const FDragDropEvent& DragDropEvent) override;
 
-		MenuBuilder.AddMenuEntry(FFlipbookEditorCommands::Get().AddNewFrameBefore);
-		MenuBuilder.AddMenuEntry(FFlipbookEditorCommands::Get().AddNewFrameAfter);
+	FReply KeyframeOnMouseButtonUp(const FGeometry& MyGeometry, const FPointerEvent& MouseEvent);
 
-		MenuBuilder.EndSection();
+	FText GetKeyframeTooltip() const;
 
-		return MenuBuilder.MakeWidget();
-	}
+	TSharedRef<SWidget> GenerateContextMenu();
 
-	FOptionalSize GetFrameWidth() const
-	{
-		UPaperFlipbook* Flipbook = FlipbookBeingEdited.Get();
-		if (Flipbook && Flipbook->IsValidKeyFrameIndex(FrameIndex))
-		{
-			const FPaperFlipbookKeyFrame& KeyFrame = Flipbook->GetKeyFrameChecked(FrameIndex);
-			return FMath::Max<float>(0, KeyFrame.FrameRun * SlateUnitsPerFrame.Get() - FFlipbookUIConstants::HandleWidth);
-		}
-		else
-		{
-			return 1;
-		}
-	}
+	FOptionalSize GetFrameWidth() const;
 
 protected:
 	int32 FrameIndex;
@@ -329,22 +107,7 @@ public:
 
 	SLATE_END_ARGS()
 
-	void Construct(const FArguments& InArgs, TSharedPtr<const FUICommandList> InCommandList)
-	{
-		CommandList = InCommandList;
-		SlateUnitsPerFrame = InArgs._SlateUnitsPerFrame;
-		FlipbookBeingEdited = InArgs._FlipbookBeingEdited;
-		OnSelectionChanged = InArgs._OnSelectionChanged;
-
-		NumKeyframesFromLastRebuild = 0;
-
-		ChildSlot
-		[
-			SAssignNew(MainBoxPtr, SHorizontalBox)
-		];
-
-		Rebuild();
-	}
+	void Construct(const FArguments& InArgs, TSharedPtr<const FUICommandList> InCommandList);
 
 	void Tick(const FGeometry& AllottedGeometry, const double InCurrentTime, const float InDeltaTime) override
 	{
@@ -359,33 +122,7 @@ public:
 	}
 
 private:
-	void Rebuild()
-	{
-		MainBoxPtr->ClearChildren();
-
-		// Create the sections for each keyframe
-		if (UPaperFlipbook* Flipbook = FlipbookBeingEdited.Get())
-		{
-			for (int32 KeyFrameIdx = 0; KeyFrameIdx < Flipbook->GetNumKeyFrames(); ++KeyFrameIdx)
-			{
-				//@TODO: Draggy bits go here
-				MainBoxPtr->AddSlot()
-				.AutoWidth()
-				[
-					SNew(SFlipbookKeyframeWidget, KeyFrameIdx, CommandList)
-					.SlateUnitsPerFrame(this->SlateUnitsPerFrame)
-					.FlipbookBeingEdited(this->FlipbookBeingEdited)
-					.OnSelectionChanged(this->OnSelectionChanged)
-				];
-			}
-
-			NumKeyframesFromLastRebuild = Flipbook->GetNumKeyFrames();
-		}
-		else
-		{
-			NumKeyframesFromLastRebuild = 0;
-		}
-	}
+	void Rebuild();
 
 private:
 	TAttribute<float> SlateUnitsPerFrame;
diff --git a/Engine/Plugins/2D/Paper2D/Source/Paper2DEditor/Private/SPaperEditorViewport.h b/Engine/Plugins/2D/Paper2D/Source/Paper2DEditor/Private/SPaperEditorViewport.h
--- a/Engine/Plugins/2D/Paper2D/Source/Paper2DEditor/Private/SPaperEditorViewport.h
+++ b/Engine/Plugins/2D/Paper2D/Source/Paper2DEditor/Private/SPaperEditorViewport.h
@@ -3,6 +3,7 @@
 #pragma once
 
 #include "GraphEditor.h"
+#include "MarqueeRect.h"
 
 /** Helper for managing marquee operations */
 struct FMarqueeOperation
diff --git a/Engine/Plugins/2D/Paper2D/Source/Paper2DEditor/Private/SpriteEditor/SpriteEditor.cpp b/Engine/Plugins/2D/Paper2D/Source/Paper2DEditor/Private/SpriteEditor/SpriteEditor.cpp
--- a/Engine/Plugins/2D/Paper2D/Source/Paper2DEditor/Private/SpriteEditor/SpriteEditor.cpp
+++ b/Engine/Plugins/2D/Paper2D/Source/Paper2DEditor/Private/SpriteEditor/SpriteEditor.cpp
@@ -15,6 +15,7 @@
 #include "SSpriteEditorViewportToolbar.h"
 
 #include "SSpriteList.h"
+#include "SDockTab.h"
 
 #define LOCTEXT_NAMESPACE "SpriteEditor"
 
diff --git a/Engine/Plugins/2D/Paper2D/Source/Paper2DEditor/Private/TileMapEditing/TileMapEditor.cpp b/Engine/Plugins/2D/Paper2D/Source/Paper2DEditor/Private/TileMapEditing/TileMapEditor.cpp
--- a/Engine/Plugins/2D/Paper2D/Source/Paper2DEditor/Private/TileMapEditing/TileMapEditor.cpp
+++ b/Engine/Plugins/2D/Paper2D/Source/Paper2DEditor/Private/TileMapEditing/TileMapEditor.cpp
@@ -13,6 +13,7 @@
 #include "WorkspaceMenuStructureModule.h"
 #include "Paper2DEditorModule.h"
 #include "STileMapEditorViewportToolbar.h"
+#include "SDockTab.h"
 
 #define LOCTEXT_NAMESPACE "TileMapEditor"
 
diff --git a/Engine/Plugins/2D/Paper2D/Source/Paper2DEditor/Private/TileSetEditor.cpp b/Engine/Plugins/2D/Paper2D/Source/Paper2DEditor/Private/TileSetEditor.cpp
--- a/Engine/Plugins/2D/Paper2D/Source/Paper2DEditor/Private/TileSetEditor.cpp
+++ b/Engine/Plugins/2D/Paper2D/Source/Paper2DEditor/Private/TileSetEditor.cpp
@@ -7,6 +7,7 @@
 #include "PaperEditorViewportClient.h"
 #include "TileMapEditing/EdModeTileMap.h"
 #include "GraphEditor.h"
+#include "SPaperEditorViewport.h"
 
 #define LOCTEXT_NAMESPACE "TileSetEditor"
 
diff --git a/Engine/Plugins/Developer/PerforceSourceControl/Source/PerforceSourceControl/Private/PerforceSourceControlPrivatePCH.h b/Engine/Plugins/Developer/PerforceSourceControl/Source/PerforceSourceControl/Private/PerforceSourceControlPrivatePCH.h
--- a/Engine/Plugins/Developer/PerforceSourceControl/Source/PerforceSourceControl/Private/PerforceSourceControlPrivatePCH.h
+++ b/Engine/Plugins/Developer/PerforceSourceControl/Source/PerforceSourceControl/Private/PerforceSourceControlPrivatePCH.h
@@ -3,7 +3,7 @@
 #pragma once
 
 #include "Core.h"
-#include "Slate.h"
+#include "SlateBasics.h"
 #include "EditorStyle.h"
 #include <p4/clientapi.h>
 #include <p4/i18napi.h>
diff --git a/Engine/Plugins/Developer/PerforceSourceControl/Source/PerforceSourceControl/Private/SPerforceSourceControlSettings.cpp b/Engine/Plugins/Developer/PerforceSourceControl/Source/PerforceSourceControl/Private/SPerforceSourceControlSettings.cpp
--- a/Engine/Plugins/Developer/PerforceSourceControl/Source/PerforceSourceControl/Private/SPerforceSourceControlSettings.cpp
+++ b/Engine/Plugins/Developer/PerforceSourceControl/Source/PerforceSourceControl/Private/SPerforceSourceControlSettings.cpp
@@ -3,6 +3,7 @@
 #include "PerforceSourceControlPrivatePCH.h"
 #include "SPerforceSourceControlSettings.h"
 #include "PerforceSourceControlModule.h"
+#include "SThrobber.h"
 
 TWeakPtr<SEditableTextBox> SPerforceSourceControlSettings::PasswordTextBox;
 
diff --git a/Engine/Plugins/Editor/EpicSurvey/Source/EpicSurvey/Private/EpicSurvey.cpp b/Engine/Plugins/Editor/EpicSurvey/Source/EpicSurvey/Private/EpicSurvey.cpp
--- a/Engine/Plugins/Editor/EpicSurvey/Source/EpicSurvey/Private/EpicSurvey.cpp
+++ b/Engine/Plugins/Editor/EpicSurvey/Source/EpicSurvey/Private/EpicSurvey.cpp
@@ -15,6 +15,8 @@
 #include "SSurveyNotification.h"
 #include "SurveyTitleLocalStorage.h"
 #include "SurveyPage.h"
+#include "SNotificationList.h"
+#include "NotificationManager.h"
 
 #define LOCTEXT_NAMESPACE "EpicSurvey"
 
diff --git a/Engine/Plugins/Editor/PluginsEditor/Source/PluginsEditor/Private/SPluginListTile.cpp b/Engine/Plugins/Editor/PluginsEditor/Source/PluginsEditor/Private/SPluginListTile.cpp
--- a/Engine/Plugins/Editor/PluginsEditor/Source/PluginsEditor/Private/SPluginListTile.cpp
+++ b/Engine/Plugins/Editor/PluginsEditor/Source/PluginsEditor/Private/SPluginListTile.cpp
@@ -7,6 +7,7 @@
 #include "SPluginList.h"
 #include "PluginStyle.h"
 #include "GameProjectGenerationModule.h"
+#include "SHyperlink.h"
 
 #define LOCTEXT_NAMESPACE "PluginListTile"
 
diff --git a/Engine/Plugins/Editor/PluginsEditor/Source/PluginsEditor/Private/SPluginsEditor.cpp b/Engine/Plugins/Editor/PluginsEditor/Source/PluginsEditor/Private/SPluginsEditor.cpp
--- a/Engine/Plugins/Editor/PluginsEditor/Source/PluginsEditor/Private/SPluginsEditor.cpp
+++ b/Engine/Plugins/Editor/PluginsEditor/Source/PluginsEditor/Private/SPluginsEditor.cpp
@@ -5,6 +5,7 @@
 #include "SPluginsEditor.h"
 #include "SPluginList.h"
 #include "SPluginCategories.h"
+#include "SSearchBox.h"
 
 #define LOCTEXT_NAMESPACE "PluginsEditor"
 
diff --git a/Engine/Plugins/Experimental/CodeView/Source/CodeView/Private/SCodeView.cpp b/Engine/Plugins/Experimental/CodeView/Source/CodeView/Private/SCodeView.cpp
--- a/Engine/Plugins/Experimental/CodeView/Source/CodeView/Private/SCodeView.cpp
+++ b/Engine/Plugins/Experimental/CodeView/Source/CodeView/Private/SCodeView.cpp
@@ -2,6 +2,8 @@
 
 #include "CodeViewPrivatePCH.h"
 #include "SCodeView.h"
+#include "SSearchBox.h"
+#include "SThrobber.h"
 
 #define LOCTEXT_NAMESPACE "SCodeView"
 
diff --git a/Engine/Plugins/Experimental/LiveEditor/Source/LiveEditor/Private/LiveEditorConfigWindow.cpp b/Engine/Plugins/Experimental/LiveEditor/Source/LiveEditor/Private/LiveEditorConfigWindow.cpp
--- a/Engine/Plugins/Experimental/LiveEditor/Source/LiveEditor/Private/LiveEditorConfigWindow.cpp
+++ b/Engine/Plugins/Experimental/LiveEditor/Source/LiveEditor/Private/LiveEditorConfigWindow.cpp
@@ -6,6 +6,7 @@
 #include "LiveEditorConnectionWindow.h"
 #include "LiveEditorDeviceSetupWindow.h"
 #include "LiveEditorWizardBase.h"
+#include "SExpandableArea.h"
 
 DEFINE_LOG_CATEGORY_STATIC(LiveEditorConfigWindow, Log, All);
 
diff --git a/Engine/Plugins/Media/MediaPlayerEditor/Source/MediaPlayerEditor/Private/MediaPlayerEditorToolkit.cpp b/Engine/Plugins/Media/MediaPlayerEditor/Source/MediaPlayerEditor/Private/MediaPlayerEditorToolkit.cpp
--- a/Engine/Plugins/Media/MediaPlayerEditor/Source/MediaPlayerEditor/Private/MediaPlayerEditorToolkit.cpp
+++ b/Engine/Plugins/Media/MediaPlayerEditor/Source/MediaPlayerEditor/Private/MediaPlayerEditorToolkit.cpp
@@ -2,7 +2,7 @@
 
 #include "MediaPlayerEditorPrivatePCH.h"
 #include "Factories.h"
-
+#include "SDockTab.h"
 
 #define LOCTEXT_NAMESPACE "FMediaPlayerEditorToolkit"
 
diff --git a/Engine/Plugins/Media/MediaPlayerEditor/Source/MediaPlayerEditor/Private/Widgets/SMediaPlayerEditorViewer.cpp b/Engine/Plugins/Media/MediaPlayerEditor/Source/MediaPlayerEditor/Private/Widgets/SMediaPlayerEditorViewer.cpp
--- a/Engine/Plugins/Media/MediaPlayerEditor/Source/MediaPlayerEditor/Private/Widgets/SMediaPlayerEditorViewer.cpp
+++ b/Engine/Plugins/Media/MediaPlayerEditor/Source/MediaPlayerEditor/Private/Widgets/SMediaPlayerEditorViewer.cpp
@@ -1,6 +1,7 @@
 // Copyright 1998-2014 Epic Games, Inc. All Rights Reserved.
 
 #include "MediaPlayerEditorPrivatePCH.h"
+#include "SNumericEntryBox.h"
 
 
 #define LOCTEXT_NAMESPACE "SMediaPlayerEditorPlayer"
diff --git a/Engine/Plugins/Messaging/MessagingDebugger/Source/MessagingDebugger/Private/Widgets/Endpoints/SMessagingEndpoints.cpp b/Engine/Plugins/Messaging/MessagingDebugger/Source/MessagingDebugger/Private/Widgets/Endpoints/SMessagingEndpoints.cpp
--- a/Engine/Plugins/Messaging/MessagingDebugger/Source/MessagingDebugger/Private/Widgets/Endpoints/SMessagingEndpoints.cpp
+++ b/Engine/Plugins/Messaging/MessagingDebugger/Source/MessagingDebugger/Private/Widgets/Endpoints/SMessagingEndpoints.cpp
@@ -1,6 +1,7 @@
 // Copyright 1998-2014 Epic Games, Inc. All Rights Reserved.
 
 #include "MessagingDebuggerPrivatePCH.h"
+#include "SExpandableArea.h"
 
 
 #define LOCTEXT_NAMESPACE "SMessagingEndpoints"
diff --git a/Engine/Plugins/Messaging/MessagingDebugger/Source/MessagingDebugger/Private/Widgets/Endpoints/SMessagingEndpointsFilterBar.cpp b/Engine/Plugins/Messaging/MessagingDebugger/Source/MessagingDebugger/Private/Widgets/Endpoints/SMessagingEndpointsFilterBar.cpp
new file mode 100644
--- /dev/null
+++ b/Engine/Plugins/Messaging/MessagingDebugger/Source/MessagingDebugger/Private/Widgets/Endpoints/SMessagingEndpointsFilterBar.cpp
@@ -0,0 +1,29 @@
+// Copyright 1998-2014 Epic Games, Inc. All Rights Reserved.
+#include "MessagingDebuggerPrivatePCH.h"
+
+#include "SMessagingEndpointsFilterBar.h"
+#include "SSearchBox.h"
+
+#define LOCTEXT_NAMESPACE "SMessagingEndpointsFilterBar"
+
+void SMessagingEndpointsFilterBar::Construct(const FArguments& InArgs, FMessagingDebuggerEndpointFilterRef InFilter)
+{
+	Filter = InFilter;
+
+	ChildSlot
+		[
+			SNew(SHorizontalBox)
+
+			+ SHorizontalBox::Slot()
+			.FillWidth(1.0f)
+			.VAlign(VAlign_Top)
+			[
+				// search box
+				SNew(SSearchBox)
+				.HintText(LOCTEXT("SearchBoxHint", "Search endpoints"))
+				.OnTextChanged(this, &SMessagingEndpointsFilterBar::HandleFilterStringTextChanged)
+			]
+		];
+}
+
+#undef LOCTEXT_NAMESPACE
diff --git a/Engine/Plugins/Messaging/MessagingDebugger/Source/MessagingDebugger/Private/Widgets/Endpoints/SMessagingEndpointsFilterBar.h b/Engine/Plugins/Messaging/MessagingDebugger/Source/MessagingDebugger/Private/Widgets/Endpoints/SMessagingEndpointsFilterBar.h
--- a/Engine/Plugins/Messaging/MessagingDebugger/Source/MessagingDebugger/Private/Widgets/Endpoints/SMessagingEndpointsFilterBar.h
+++ b/Engine/Plugins/Messaging/MessagingDebugger/Source/MessagingDebugger/Private/Widgets/Endpoints/SMessagingEndpointsFilterBar.h
@@ -3,9 +3,6 @@
 #pragma once
 
 
-#define LOCTEXT_NAMESPACE "SMessagingEndpointsFilterBar"
-
-
 /**
  * Implements the endpoints list filter bar widget.
  */
@@ -25,25 +22,7 @@ public:
 	 * @param InArgs The declaration data for this widget.
 	 * @param InFilter The filter model.
 	 */
-	void Construct( const FArguments& InArgs, FMessagingDebuggerEndpointFilterRef InFilter )
-	{
-		Filter = InFilter;
-
-		ChildSlot
-		[
-			SNew(SHorizontalBox)
-
-			+ SHorizontalBox::Slot()
-				.FillWidth(1.0f)
-				.VAlign(VAlign_Top)
-				[
-					// search box
-					SNew(SSearchBox)
-						.HintText(LOCTEXT("SearchBoxHint", "Search endpoints"))
-						.OnTextChanged(this, &SMessagingEndpointsFilterBar::HandleFilterStringTextChanged)
-				]
-		];
-	}
+	void Construct( const FArguments& InArgs, FMessagingDebuggerEndpointFilterRef InFilter );
 
 private:
 
@@ -58,6 +37,3 @@ private:
 	// Holds a pointer to the filter model.
 	FMessagingDebuggerEndpointFilterPtr Filter;
 };
-
-
-#undef LOCTEXT_NAMESPACE
diff --git a/Engine/Plugins/Messaging/MessagingDebugger/Source/MessagingDebugger/Private/Widgets/History/SMessagingHistory.cpp b/Engine/Plugins/Messaging/MessagingDebugger/Source/MessagingDebugger/Private/Widgets/History/SMessagingHistory.cpp
--- a/Engine/Plugins/Messaging/MessagingDebugger/Source/MessagingDebugger/Private/Widgets/History/SMessagingHistory.cpp
+++ b/Engine/Plugins/Messaging/MessagingDebugger/Source/MessagingDebugger/Private/Widgets/History/SMessagingHistory.cpp
@@ -1,6 +1,7 @@
 // Copyright 1998-2014 Epic Games, Inc. All Rights Reserved.
 
 #include "MessagingDebuggerPrivatePCH.h"
+#include "SHyperlink.h"
 
 
 #define LOCTEXT_NAMESPACE "SMessagingHistory"
diff --git a/Engine/Plugins/Messaging/MessagingDebugger/Source/MessagingDebugger/Private/Widgets/Types/SMessagingTypesFilterBar.cpp b/Engine/Plugins/Messaging/MessagingDebugger/Source/MessagingDebugger/Private/Widgets/Types/SMessagingTypesFilterBar.cpp
new file mode 100644
--- /dev/null
+++ b/Engine/Plugins/Messaging/MessagingDebugger/Source/MessagingDebugger/Private/Widgets/Types/SMessagingTypesFilterBar.cpp
@@ -0,0 +1,30 @@
+// Copyright 1998-2014 Epic Games, Inc. All Rights Reserved.
+
+#include "MessagingDebuggerPrivatePCH.h"
+
+#include "SMessagingTypesFilterBar.h"
+#include "SSearchBox.h"
+
+#define LOCTEXT_NAMESPACE "SMessagingTypesFilterBar"
+
+void SMessagingTypesFilterBar::Construct(const FArguments& InArgs, FMessagingDebuggerTypeFilterRef InFilter)
+{
+	Filter = InFilter;
+
+	ChildSlot
+		[
+			SNew(SHorizontalBox)
+
+			+ SHorizontalBox::Slot()
+			.FillWidth(1.0f)
+			.VAlign(VAlign_Top)
+			[
+				// search box
+				SNew(SSearchBox)
+				.HintText(LOCTEXT("SearchBoxHint", "Search message types"))
+				.OnTextChanged(this, &SMessagingTypesFilterBar::HandleFilterStringTextChanged)
+			]
+		];
+}
+
+#undef LOCTEXT_NAMESPACE
diff --git a/Engine/Plugins/Messaging/MessagingDebugger/Source/MessagingDebugger/Private/Widgets/Types/SMessagingTypesFilterBar.h b/Engine/Plugins/Messaging/MessagingDebugger/Source/MessagingDebugger/Private/Widgets/Types/SMessagingTypesFilterBar.h
--- a/Engine/Plugins/Messaging/MessagingDebugger/Source/MessagingDebugger/Private/Widgets/Types/SMessagingTypesFilterBar.h
+++ b/Engine/Plugins/Messaging/MessagingDebugger/Source/MessagingDebugger/Private/Widgets/Types/SMessagingTypesFilterBar.h
@@ -3,9 +3,6 @@
 #pragma once
 
 
-#define LOCTEXT_NAMESPACE "SMessagingTypesFilterBar"
-
-
 /**
  * Implements the message type list filter bar widget.
  */
@@ -38,25 +35,7 @@ public:
 	 * @param InArgs The declaration data for this widget.
 	 * @param InFilter The filter model.
 	 */
-	void Construct( const FArguments& InArgs, FMessagingDebuggerTypeFilterRef InFilter )
-	{
-		Filter = InFilter;
-
-		ChildSlot
-		[
-			SNew(SHorizontalBox)
-
-			+ SHorizontalBox::Slot()
-				.FillWidth(1.0f)
-				.VAlign(VAlign_Top)
-				[
-					// search box
-					SNew(SSearchBox)
-						.HintText(LOCTEXT("SearchBoxHint", "Search message types"))
-						.OnTextChanged(this, &SMessagingTypesFilterBar::HandleFilterStringTextChanged)
-				]
-		];
-	}
+	void Construct( const FArguments& InArgs, FMessagingDebuggerTypeFilterRef InFilter );
 
 private:
 
@@ -71,6 +50,3 @@ private:
 	// Holds the filter model.
 	FMessagingDebuggerTypeFilterPtr Filter;
 };
-
-
-#undef LOCTEXT_NAMESPACE
diff --git a/Engine/Source/Developer/Android/AndroidPlatformEditor/Private/AndroidTargetSettingsCustomization.cpp b/Engine/Source/Developer/Android/AndroidPlatformEditor/Private/AndroidTargetSettingsCustomization.cpp
--- a/Engine/Source/Developer/Android/AndroidPlatformEditor/Private/AndroidTargetSettingsCustomization.cpp
+++ b/Engine/Source/Developer/Android/AndroidPlatformEditor/Private/AndroidTargetSettingsCustomization.cpp
@@ -13,6 +13,8 @@
 #include "PlatformIconInfo.h"
 #include "SourceControlHelpers.h"
 #include "ManifestUpdateHelper.h"
+#include "SNotificationList.h"
+#include "NotificationManager.h"
 
 #define LOCTEXT_NAMESPACE "AndroidRuntimeSettings"
 
diff --git a/Engine/Source/Developer/AutomationWindow/Private/SAutomationExportMenu.cpp b/Engine/Source/Developer/AutomationWindow/Private/SAutomationExportMenu.cpp
--- a/Engine/Source/Developer/AutomationWindow/Private/SAutomationExportMenu.cpp
+++ b/Engine/Source/Developer/AutomationWindow/Private/SAutomationExportMenu.cpp
@@ -1,6 +1,7 @@
 // Copyright 1998-2014 Epic Games, Inc. All Rights Reserved.
 
 #include "AutomationWindowPrivatePCH.h"
+#include "SNotificationList.h"
 
 
 #define LOCTEXT_NAMESPACE "AutomationExportMenu"
diff --git a/Engine/Source/Developer/AutomationWindow/Private/SAutomationWindow.cpp b/Engine/Source/Developer/AutomationWindow/Private/SAutomationWindow.cpp
--- a/Engine/Source/Developer/AutomationWindow/Private/SAutomationWindow.cpp
+++ b/Engine/Source/Developer/AutomationWindow/Private/SAutomationWindow.cpp
@@ -5,6 +5,9 @@
 #include "AutomationController.h"
 #include "AutomationPresetManager.h"
 #include "SAutomationWindow.h"
+#include "SSearchBox.h"
+#include "SNotificationList.h"
+#include "SThrobber.h"
 
 
 #define LOCTEXT_NAMESPACE "AutomationTest"
diff --git a/Engine/Source/Developer/CollisionAnalyzer/Private/CollisionAnalyzerModule.cpp b/Engine/Source/Developer/CollisionAnalyzer/Private/CollisionAnalyzerModule.cpp
--- a/Engine/Source/Developer/CollisionAnalyzer/Private/CollisionAnalyzerModule.cpp
+++ b/Engine/Source/Developer/CollisionAnalyzer/Private/CollisionAnalyzerModule.cpp
@@ -1,8 +1,8 @@
 // Copyright 1998-2014 Epic Games, Inc. All Rights Reserved.
 
 #include "CollisionAnalyzerPCH.h"
-
 #include "WorkspaceMenuStructureModule.h"
+#include "SDockTab.h"
 
 namespace CollisionAnalyzerModule
 {
diff --git a/Engine/Source/Developer/DeviceManager/Private/DeviceManagerModule.cpp b/Engine/Source/Developer/DeviceManager/Private/DeviceManagerModule.cpp
--- a/Engine/Source/Developer/DeviceManager/Private/DeviceManagerModule.cpp
+++ b/Engine/Source/Developer/DeviceManager/Private/DeviceManagerModule.cpp
@@ -1,7 +1,10 @@
 // Copyright 1998-2014 Epic Games, Inc. All Rights Reserved.
 
 #include "DeviceManagerPrivatePCH.h"
+
 #include "WorkspaceMenuStructureModule.h"
+#include "SDockTab.h"
+
 
 static const FName DeviceManagerTabName("DeviceManager");
 
diff --git a/Engine/Source/Developer/DeviceManager/Private/Widgets/Browser/SDeviceBrowserFilterBar.cpp b/Engine/Source/Developer/DeviceManager/Private/Widgets/Browser/SDeviceBrowserFilterBar.cpp
--- a/Engine/Source/Developer/DeviceManager/Private/Widgets/Browser/SDeviceBrowserFilterBar.cpp
+++ b/Engine/Source/Developer/DeviceManager/Private/Widgets/Browser/SDeviceBrowserFilterBar.cpp
@@ -1,6 +1,7 @@
 // Copyright 1998-2014 Epic Games, Inc. All Rights Reserved.
 
 #include "DeviceManagerPrivatePCH.h"
+#include "SSearchBox.h"
 
 
 #define LOCTEXT_NAMESPACE "SDeviceBrowserFilterBar"
diff --git a/Engine/Source/Developer/Merge/Private/Merge.cpp b/Engine/Source/Developer/Merge/Private/Merge.cpp
--- a/Engine/Source/Developer/Merge/Private/Merge.cpp
+++ b/Engine/Source/Developer/Merge/Private/Merge.cpp
@@ -7,6 +7,9 @@
 #include "SBlueprintMerge.h"
 #include "Toolkits/ToolkitManager.h"
 #include "Toolkits/IToolkit.h"
+#include "SDockTab.h"
+#include "SNotificationList.h"
+#include "NotificationManager.h"
 
 #define LOCTEXT_NAMESPACE "Merge"
 
diff --git a/Engine/Source/Developer/MessageLog/Private/UserInterface/SMessageLogListing.cpp b/Engine/Source/Developer/MessageLog/Private/UserInterface/SMessageLogListing.cpp
--- a/Engine/Source/Developer/MessageLog/Private/UserInterface/SMessageLogListing.cpp
+++ b/Engine/Source/Developer/MessageLog/Private/UserInterface/SMessageLogListing.cpp
@@ -4,6 +4,7 @@
 #include "SMessageLogListing.h"
 #include "IMessageLogListing.h"
 #include "SMessageLogMessageListRow.h"
+#include "GenericCommands.h"
 
 
 #define LOCTEXT_NAMESPACE "Developer.MessageLog"
diff --git a/Engine/Source/Developer/MessageLog/Private/UserInterface/SMessageLogMessageListRow.h b/Engine/Source/Developer/MessageLog/Private/UserInterface/SMessageLogMessageListRow.h
--- a/Engine/Source/Developer/MessageLog/Private/UserInterface/SMessageLogMessageListRow.h
+++ b/Engine/Source/Developer/MessageLog/Private/UserInterface/SMessageLogMessageListRow.h
@@ -12,7 +12,7 @@
 #define LOCTEXT_NAMESPACE "SMessageLogMessageListRow"
 
 
-class SMessageLogMessageListRow
+class MESSAGELOG_API SMessageLogMessageListRow
 	: public STableRow<TSharedPtr<FTokenizedMessage>>
 {
 public:
@@ -34,266 +34,18 @@ public:
 	 *
 	 * @param InArgs  Declaration from which to construct this widget
 	 */
-	void Construct( const FArguments& InArgs, const TSharedRef< STableViewBase >& InOwnerTableView )
-	{
-		this->OnTokenClicked = InArgs._OnTokenClicked;
-
-		Message = InArgs._Message.Get();
-
-		STableRow<TSharedPtr<FTokenizedMessage>>::Construct(
-			STableRow<TSharedPtr<FTokenizedMessage>>::FArguments()
-			.Content()
-			[
-				GenerateWidget()
-			],
-			InOwnerTableView
-		);
-	}
+	void Construct( const FArguments& InArgs, const TSharedRef< STableViewBase >& InOwnerTableView );
 
 public:
 
 	/** @return Widget for this log listing item*/
-	virtual TSharedRef<SWidget> GenerateWidget()
-	{
-		// See if we have any valid tokens which match the column name
-		const TArray<TSharedRef<IMessageToken>>& MessageTokens = Message->GetMessageTokens();
-
-		// Create the horizontal box and add the icon
-		TSharedRef<SHorizontalBox> MessageBox = SNew(SHorizontalBox);
-		TSharedRef<SHorizontalBox> LinkBox = SNew(SHorizontalBox);
-		FName SeverityImageName = NAME_None;
-		bool HasLinks = false;
-
-		// Iterate over parts of the message and create widgets for them
-		for (auto TokenIt = MessageTokens.CreateConstIterator(); TokenIt; ++TokenIt)
-		{
-			const TSharedRef<IMessageToken>& Token = *TokenIt;
-
-			switch (Token->GetType())
-			{
-			case EMessageToken::Severity:
-				{
-					if (SeverityImageName == NAME_None)
-					{
-						const TSharedRef<FSeverityToken> SeverityToken = StaticCastSharedRef<FSeverityToken>(Token);
-						SeverityImageName = FTokenizedMessage::GetSeverityIconName(SeverityToken->GetSeverity());
-					}
-				}
-				break;
-
-			case EMessageToken::Documentation:
-			case EMessageToken::Tutorial:
-				{
-					CreateMessage(LinkBox, Token, 10.0f);
-					HasLinks = true;
-				}
-				break;
-
-			default:
-				CreateMessage(MessageBox, Token, 2.0f);
-				break;
-			}
-		}
-
-		return SNew(SHorizontalBox)
-			.ToolTipText(Message->ToText())
-
-		+ SHorizontalBox::Slot()
-			.AutoWidth()
-			.HAlign(HAlign_Center)
-			.VAlign(VAlign_Center)
-			[
-				SNew(SBox)
-					.Padding(2.0f)
-					[
-						(SeverityImageName == NAME_None)
-							? SNullWidget::NullWidget
-							: static_cast<TSharedRef<SWidget>>(SNew(SImage)
-								.Image(FEditorStyle::GetBrush(SeverityImageName)))
-					]
-			]
-
-		+ SHorizontalBox::Slot()
-			.FillWidth(1.0f)
-			.VAlign(VAlign_Center)
-			[
-				MessageBox
-			]
-
-		+ SHorizontalBox::Slot()
-			.AutoWidth()
-			.HAlign(HAlign_Right)
-			.VAlign(VAlign_Center)
-			.Padding(1.0f)
-			[
-				!HasLinks
-					? SNullWidget::NullWidget
-					: static_cast<TSharedRef<SWidget>>(SNew(SBorder)
-						.BorderImage(FEditorStyle::GetBrush("ToolPanel.GroupBorder"))
-						.Padding(FMargin(0.0f, 1.0f, 10.0f, 1.0f))
-						[
-							LinkBox
-						])
-			];
-	}
+	virtual TSharedRef<SWidget> GenerateWidget();
 
 protected:
 
-	TSharedRef<SWidget> CreateHyperlink( const TSharedRef<IMessageToken>& InMessageToken, const FText& InToolTip = FText() )
-	{
-		return SNew(SHyperlink)
-			.Text(InMessageToken->ToText())
-			.ToolTipText(InToolTip)
-			.TextStyle(FEditorStyle::Get(), "MessageLog")
-			.OnNavigate(this, &SMessageLogMessageListRow::HandleHyperlinkNavigate, InMessageToken);
-	}
-
-	void CreateMessage( const TSharedRef<SHorizontalBox>& InHorzBox, const TSharedRef<IMessageToken>& InMessageToken, float Padding )
-	{
-		TSharedPtr<SWidget> Content;
-		FName IconBrushName;
+	TSharedRef<SWidget> CreateHyperlink( const TSharedRef<IMessageToken>& InMessageToken, const FText& InToolTip = FText() );
 
-		switch(InMessageToken->GetType())
-		{
-		case EMessageToken::Image:
-			{
-				const TSharedRef<FImageToken> ImageToken = StaticCastSharedRef<FImageToken>(InMessageToken);
-				
-				if (ImageToken->GetImageName() != NAME_None)
-				{
-					if (InMessageToken->GetOnMessageTokenActivated().IsBound())
-					{
-						Content = SNew(SButton)
-							.OnClicked(this, &SMessageLogMessageListRow::HandleTokenButtonClicked, InMessageToken)
-							.Content()
-							[
-								SNew(SImage)
-									.Image(FEditorStyle::GetBrush(ImageToken->GetImageName()))
-							];
-					}
-					else
-					{
-						Content = SNew(SImage)
-							.Image(FEditorStyle::GetBrush(ImageToken->GetImageName()));
-					}
-				}
-			}
-			break;
-
-		case EMessageToken::Object:
-			{
-				const TSharedRef<FUObjectToken> UObjectToken = StaticCastSharedRef<FUObjectToken>(InMessageToken);
-				
-				IconBrushName = FName("PropertyWindow.Button_Browse");
-				Content = CreateHyperlink(InMessageToken, FUObjectToken::DefaultOnGetObjectDisplayName().IsBound()
-					? FUObjectToken::DefaultOnGetObjectDisplayName().Execute(UObjectToken->GetObject().Get(), true)
-					: UObjectToken->ToText());
-			}
-			break;
-
-		case EMessageToken::URL:
-			{
-				const TSharedRef<FURLToken> URLToken = StaticCastSharedRef<FURLToken>(InMessageToken);
-
-				IconBrushName = FName("MessageLog.Url");
-				Content = CreateHyperlink(InMessageToken, FText::FromString(URLToken->GetURL()));
-			}
-			break;
-
-		case EMessageToken::Action:
-			{
-				const TSharedRef<FActionToken> ActionToken = StaticCastSharedRef<FActionToken>(InMessageToken);
-
-				IconBrushName = FName("MessageLog.Action");
-				Content = SNew(SHyperlink)
-					.Text(InMessageToken->ToText())
-					.ToolTipText(ActionToken->GetActionDescription())
-					.TextStyle(FEditorStyle::Get(), "MessageLog")
-					.OnNavigate(this, &SMessageLogMessageListRow::HandleActionHyperlinkNavigate, ActionToken);
-			}
-			break;
-
-		case EMessageToken::AssetName:
-			{
-				const TSharedRef<FAssetNameToken> AssetNameToken = StaticCastSharedRef<FAssetNameToken>(InMessageToken);
-
-				IconBrushName = FName("PropertyWindow.Button_Browse");
-				Content = CreateHyperlink(InMessageToken, AssetNameToken->ToText());
-			}
-			break;
-
-#if WITH_EDITOR
-		case EMessageToken::Documentation:
-			{
-				const TSharedRef<FDocumentationToken> DocumentationToken = StaticCastSharedRef<FDocumentationToken>(InMessageToken);
-
-				IconBrushName = FName("MessageLog.Docs");
-				Content = SNew(SHyperlink)
-					.Text(LOCTEXT("DocsLabel", "Docs"))
-					.ToolTip(IDocumentation::Get()->CreateToolTip(
-						LOCTEXT("DocumentationTokenToolTip", "Click to open documentation"),
-						NULL,
-						DocumentationToken->GetPreviewExcerptLink(),
-						DocumentationToken->GetPreviewExcerptName())
-					)
-					.TextStyle(FEditorStyle::Get(), "MessageLog")
-					.OnNavigate(this, &SMessageLogMessageListRow::HandleDocsHyperlinkNavigate, DocumentationToken->GetDocumentationLink());
-			}
-			break;
-
-		case EMessageToken::Text:
-			{
-				Content = SNew(STextBlock)
-					.ColorAndOpacity(FSlateColor::UseSubduedForeground())
-					.Text(InMessageToken->ToText());
-			}
-			break;
-
-		case EMessageToken::Tutorial:
-			{
-				const TSharedRef<FTutorialToken> TutorialToken = StaticCastSharedRef<FTutorialToken>(InMessageToken);
-
-				IconBrushName = FName("MessageLog.Tutorial");
-				Content = SNew(SHyperlink)
-					.Text(LOCTEXT("TutorialLabel", "Tutorial"))
-					.ToolTipText(LOCTEXT("TutorialTokenToolTip", "Click to open tutorial"))
-					.TextStyle(FEditorStyle::Get(), "MessageLog")
-					.OnNavigate(this, &SMessageLogMessageListRow::HandleTutorialHyperlinkNavigate, TutorialToken->GetTutorialAssetName());
-			}
-			break;
-#endif
-		}
-
-		if (Content.IsValid())
-		{
-			InHorzBox->AddSlot()
-				.AutoWidth()
-				.VAlign(VAlign_Center)
-				.Padding(Padding, 0.0f, 0.0f, 0.0f)
-				[
-					SNew(SHorizontalBox)
-
-					+ SHorizontalBox::Slot()
-						.AutoWidth()
-						.VAlign(VAlign_Center)
-						[
-							(IconBrushName == NAME_None)
-								? SNullWidget::NullWidget
-								: static_cast<TSharedRef<SWidget>>(SNew(SImage)
-									.ColorAndOpacity(FSlateColor::UseForeground())
-									.Image(FEditorStyle::GetBrush(IconBrushName)))
-						]
-
-					+ SHorizontalBox::Slot()
-						.AutoWidth()
-						.VAlign(VAlign_Center)
-						.Padding(2.0f, 0.0f, 0.0f, 0.0f)
-						[
-							Content.ToSharedRef()
-						]
-				];
-		}
-	}
+	void CreateMessage( const TSharedRef<SHorizontalBox>& InHorzBox, const TSharedRef<IMessageToken>& InMessageToken, float Padding );
 
 private:
 
diff --git a/Engine/Source/Developer/Profiler/Private/Widgets/SProfilerWindow.cpp b/Engine/Source/Developer/Profiler/Private/Widgets/SProfilerWindow.cpp
--- a/Engine/Source/Developer/Profiler/Private/Widgets/SProfilerWindow.cpp
+++ b/Engine/Source/Developer/Profiler/Private/Widgets/SProfilerWindow.cpp
@@ -12,6 +12,7 @@
 	#include "Runtime/Analytics/Analytics/Public/Interfaces/IAnalyticsProvider.h"
 	#include "Runtime/Engine/Public/EngineAnalytics.h"
 #endif // WITH_EDITOR
+#include "SNotificationList.h"
 
 #define LOCTEXT_NAMESPACE "SProfilerWindow"
 
diff --git a/Engine/Source/Developer/ProjectLauncher/Private/ProjectLauncherModule.cpp b/Engine/Source/Developer/ProjectLauncher/Private/ProjectLauncherModule.cpp
--- a/Engine/Source/Developer/ProjectLauncher/Private/ProjectLauncherModule.cpp
+++ b/Engine/Source/Developer/ProjectLauncher/Private/ProjectLauncherModule.cpp
@@ -1,7 +1,10 @@
 // Copyright 1998-2014 Epic Games, Inc. All Rights Reserved.
 
 #include "ProjectLauncherPrivatePCH.h"
+
 #include "WorkspaceMenuStructureModule.h"
+#include "SDockTab.h"
+
 
 static const FName ProjectLauncherTabName("ProjectLauncher");
 
diff --git a/Engine/Source/Developer/ProjectLauncher/Private/Widgets/Cook/SProjectLauncherCookByTheBookSettings.cpp b/Engine/Source/Developer/ProjectLauncher/Private/Widgets/Cook/SProjectLauncherCookByTheBookSettings.cpp
--- a/Engine/Source/Developer/ProjectLauncher/Private/Widgets/Cook/SProjectLauncherCookByTheBookSettings.cpp
+++ b/Engine/Source/Developer/ProjectLauncher/Private/Widgets/Cook/SProjectLauncherCookByTheBookSettings.cpp
@@ -1,6 +1,8 @@
 // Copyright 1998-2014 Epic Games, Inc. All Rights Reserved.
 
 #include "ProjectLauncherPrivatePCH.h"
+#include "SHyperlink.h"
+#include "SExpandableArea.h"
 
 
 #define LOCTEXT_NAMESPACE "SProjectLauncherCookByTheBookSettings"
diff --git a/Engine/Source/Developer/ProjectLauncher/Private/Widgets/Cook/SProjectLauncherCookedPlatforms.h b/Engine/Source/Developer/ProjectLauncher/Private/Widgets/Cook/SProjectLauncherCookedPlatforms.h
--- a/Engine/Source/Developer/ProjectLauncher/Private/Widgets/Cook/SProjectLauncherCookedPlatforms.h
+++ b/Engine/Source/Developer/ProjectLauncher/Private/Widgets/Cook/SProjectLauncherCookedPlatforms.h
@@ -3,9 +3,6 @@
 #pragma once
 
 
-#define LOCTEXT_NAMESPACE "SProjectLauncherCookedPlatforms"
-
-
 /**
  * Implements the cooked platforms panel.
  */
@@ -25,81 +22,7 @@ public:
 	 * @param InArgs The Slate argument list.
 	 * @param InModel The data model.
 	 */
-	void Construct(	const FArguments& InArgs, const FProjectLauncherModelRef& InModel )
-	{
-		Model = InModel;
-
-		MakePlatformMenu();
-
-		ChildSlot
-		[
-			SNew(SVerticalBox)
-
-			+ SVerticalBox::Slot()
-				.FillHeight(1.0f)
-				[
-					// platform menu
-					SAssignNew(PlatformListView, SListView<TSharedPtr<FString> >)
-					.HeaderRow(
-					SNew(SHeaderRow)
-					.Visibility(EVisibility::Collapsed)
-
-					+ SHeaderRow::Column("PlatformName")
-					.DefaultLabel(LOCTEXT("PlatformListPlatformNameColumnHeader", "Platform"))
-					.FillWidth(1.0f)
-					)
-					.ItemHeight(16.0f)
-					.ListItemsSource(&PlatformList)
-					.OnGenerateRow(this, &SProjectLauncherCookedPlatforms::HandlePlatformListViewGenerateRow)
-					.SelectionMode(ESelectionMode::None)
-				]
-
-			+ SVerticalBox::Slot()
-				.AutoHeight()
-				.Padding(0.0f, 6.0f, 0.0f, 4.0f)
-				[
-					SNew(SSeparator)
-						.Orientation(Orient_Horizontal)
-				]
-
-			+ SVerticalBox::Slot()
-				.AutoHeight()
-				[
-					SNew(SHorizontalBox)
-
-					+ SHorizontalBox::Slot()
-						.FillWidth(1.0f)
-						.HAlign(HAlign_Right)
-						[
-							SNew(STextBlock)
-								.Text(LOCTEXT("SelectLabel", "Select:"))
-						]
-
-					+ SHorizontalBox::Slot()
-						.AutoWidth()
-						.Padding(8.0f, 0.0f)
-						[
-							// all platforms hyper-link
-							SNew(SHyperlink)
-								.OnNavigate(this, &SProjectLauncherCookedPlatforms::HandleAllPlatformsHyperlinkNavigate, true)
-								.Text(LOCTEXT("AllPlatformsHyperlinkLabel", "All"))
-								.ToolTipText(LOCTEXT("AllPlatformsButtonTooltip", "Select all available platforms."))
-								.Visibility(this, &SProjectLauncherCookedPlatforms::HandleAllPlatformsHyperlinkVisibility)									
-						]
-
-					+ SHorizontalBox::Slot()
-						.AutoWidth()
-						[
-							// no platforms hyper-link
-							SNew(SHyperlink)
-								.OnNavigate(this, &SProjectLauncherCookedPlatforms::HandleAllPlatformsHyperlinkNavigate, false)
-								.Text(LOCTEXT("NoPlatformsHyperlinkLabel", "None"))
-								.ToolTipText(LOCTEXT("NoPlatformsHyperlinkTooltip", "Deselect all platforms."))
-								.Visibility(this, &SProjectLauncherCookedPlatforms::HandleAllPlatformsHyperlinkVisibility)									
-						]
-				]
-		];
-	}
+	void Construct(	const FArguments& InArgs, const FProjectLauncherModelRef& InModel );
 
 protected:
 
@@ -182,12 +105,7 @@ private:
 	}
 
 	// Handles generating a row widget in the map list view.
-	TSharedRef<ITableRow> HandlePlatformListViewGenerateRow( TSharedPtr<FString> InItem, const TSharedRef<STableViewBase>& OwnerTable )
-	{
-		return SNew(SProjectLauncherPlatformListRow, Model.ToSharedRef())
-			.PlatformName(InItem)
-			.OwnerTableView(OwnerTable);
-	}
+	TSharedRef<ITableRow> HandlePlatformListViewGenerateRow( TSharedPtr<FString> InItem, const TSharedRef<STableViewBase>& OwnerTable );
 
 private:
 
@@ -202,5 +120,3 @@ private:
 
 };
 
-
-#undef LOCTEXT_NAMESPACE
diff --git a/Engine/Source/Developer/ProjectLauncher/Private/Widgets/Deploy/SProjectLauncherDeviceGroupSelector.h b/Engine/Source/Developer/ProjectLauncher/Private/Widgets/Deploy/SProjectLauncherDeviceGroupSelector.h
--- a/Engine/Source/Developer/ProjectLauncher/Private/Widgets/Deploy/SProjectLauncherDeviceGroupSelector.h
+++ b/Engine/Source/Developer/ProjectLauncher/Private/Widgets/Deploy/SProjectLauncherDeviceGroupSelector.h
@@ -3,9 +3,6 @@
 #pragma once
 
 
-#define LOCTEXT_NAMESPACE "SProjectLauncherDeviceGroupSelector"
-
-
 /**
  * Delegate type for device group selection changes.
  *
@@ -58,50 +55,7 @@ public:
 	 * @param InArgs - The Slate argument list.
 	 * @param InProfileManager - The profile manager to use.
 	 */
-	void Construct( const FArguments& InArgs, const ILauncherProfileManagerRef& InProfileManager )
-	{
-		OnGroupSelected = InArgs._OnGroupSelected;
-
-		ProfileManager = InProfileManager;
-
-		ChildSlot
-		[
-			SNew(SVerticalBox)
-
-			+ SVerticalBox::Slot()
-				.AutoHeight()
-				[
-					SNew(SProjectLauncherFormLabel)
-						.LabelText(LOCTEXT("DeviceGroupComboBoxLabel", "Device group to deploy to:"))
-				]
-
-			+ SVerticalBox::Slot()
-				.AutoHeight()
-				.Padding(0.0, 4.0, 0.0, 0.0)
-				[
-					SAssignNew(DeviceGroupComboBox, SEditableComboBox<ILauncherDeviceGroupPtr>)
-						.InitiallySelectedItem(InArgs._InitiallySelectedGroup)
-						.OptionsSource(&ProfileManager->GetAllDeviceGroups())
-						.AddButtonToolTip(LOCTEXT("AddProfileButtonToolTip", "Add a new device group"))
-						.RemoveButtonToolTip(LOCTEXT("DeleteProfileButtonToolTip", "Delete the selected device group"))
-						.RenameButtonToolTip(LOCTEXT("RenameProfileButtonToolTip", "Rename the selected device group"))
-						.OnAddClicked(this, &SProjectLauncherDeviceGroupSelector::HandleDeviceGroupComboBoxAddClicked)
-						.OnGenerateWidget(this, &SProjectLauncherDeviceGroupSelector::HandleDeviceGroupComboBoxGenerateWidget)
-						.OnGetEditableText(this, &SProjectLauncherDeviceGroupSelector::HandleDeviceGroupComboBoxGetEditableText)
-						.OnRemoveClicked(this, &SProjectLauncherDeviceGroupSelector::HandleDeviceGroupComboBoxRemoveClicked)
-						.OnSelectionChanged(this, &SProjectLauncherDeviceGroupSelector::HandleDeviceGroupComboBoxSelectionChanged)
-						.OnSelectionRenamed(this, &SProjectLauncherDeviceGroupSelector::HandleDeviceGroupComboBoxSelectionRenamed)
-						.Content()
-						[
-							SNew(STextBlock)
-								.Text(this, &SProjectLauncherDeviceGroupSelector::HandleDeviceGroupComboBoxContent)
-						]
-				]
-		];
-
-		ProfileManager->OnDeviceGroupAdded().AddSP(this, &SProjectLauncherDeviceGroupSelector::HandleProfileManagerDeviceGroupsChanged);
-		ProfileManager->OnDeviceGroupRemoved().AddSP(this, &SProjectLauncherDeviceGroupSelector::HandleProfileManagerDeviceGroupsChanged);
-	}
+	void Construct( const FArguments& InArgs, const ILauncherProfileManagerRef& InProfileManager );
 
 	/**
 	 * Gets the currently selected device group.
@@ -110,10 +64,7 @@ public:
 	 *
 	 * @see SetSelectedGroup
 	 */
-	ILauncherDeviceGroupPtr GetSelectedGroup( ) const
-	{
-		return DeviceGroupComboBox->GetSelectedItem();
-	}
+	ILauncherDeviceGroupPtr GetSelectedGroup( ) const;
 
 	/**
 	 * Sets the selected device group.
@@ -122,80 +73,24 @@ public:
 	 *
 	 * @see GetSelectedGroup
 	 */
-	void SetSelectedGroup( const ILauncherDeviceGroupPtr& DeviceGroup )
-	{
-		if (!DeviceGroup.IsValid() || ProfileManager->GetAllDeviceGroups().Contains(DeviceGroup))
-		{
-			DeviceGroupComboBox->SetSelectedItem(DeviceGroup);
-		}		
-	}
+	void SetSelectedGroup( const ILauncherDeviceGroupPtr& DeviceGroup );
 
 private:
 
 	// Callback for clicking the 'Add device group' button
-	FReply HandleDeviceGroupComboBoxAddClicked( )
-	{
-		ILauncherDeviceGroupPtr NewGroup = ProfileManager->AddNewDeviceGroup();
-
-		DeviceGroupComboBox->SetSelectedItem(NewGroup);
-
-		return FReply::Handled();
-	}
+	FReply HandleDeviceGroupComboBoxAddClicked( );
 
 	// Callback for getting the device group combo box content.
-	FString HandleDeviceGroupComboBoxContent() const
-	{
-		ILauncherDeviceGroupPtr SelectedGroup = DeviceGroupComboBox->GetSelectedItem();
-
-		if (SelectedGroup.IsValid())
-		{
-			return SelectedGroup->GetName();
-		}
-
-		return LOCTEXT("CreateOrSelectGroupText", "Create or select a device group...").ToString();
-	}
+	FString HandleDeviceGroupComboBoxContent() const;
 
 	// Callback for generating a row for the device group combo box.
-	TSharedRef<SWidget> HandleDeviceGroupComboBoxGenerateWidget( ILauncherDeviceGroupPtr InItem )
-	{
-		return SNew(STextBlock)
-			.Text(this, &SProjectLauncherDeviceGroupSelector::HandleDeviceGroupComboBoxWidgetText, InItem);
-	}
+	TSharedRef<SWidget> HandleDeviceGroupComboBoxGenerateWidget( ILauncherDeviceGroupPtr InItem );
 
 	// Callback for getting the editable text of the currently selected device group.
-	FString HandleDeviceGroupComboBoxGetEditableText( )
-	{
-		ILauncherDeviceGroupPtr SelectedGroup = DeviceGroupComboBox->GetSelectedItem();
-
-		if (SelectedGroup.IsValid())
-		{
-			return SelectedGroup->GetName();
-		}
-
-		return FString();
-	}
+	FString HandleDeviceGroupComboBoxGetEditableText( );
 
 	// Callback for clicking the 'Delete device group' button
-	FReply HandleDeviceGroupComboBoxRemoveClicked( )
-	{
-		ILauncherDeviceGroupPtr SelectedGroup = DeviceGroupComboBox->GetSelectedItem();
-
-		if (SelectedGroup.IsValid())
-		{
-			ProfileManager->RemoveDeviceGroup(SelectedGroup.ToSharedRef());
-		}
-
-		if (ProfileManager->GetAllDeviceGroups().Num() > 0)
-		{
-			DeviceGroupComboBox->SetSelectedItem(ProfileManager->GetAllDeviceGroups()[0]);
-		}
-		else
-		{
-			DeviceGroupComboBox->SetSelectedItem(NULL);
-		}
-
-		return FReply::Handled();
-	}
+	FReply HandleDeviceGroupComboBoxRemoveClicked( );
 
 	// Callback for changing the selected device group in the device group combo box.
 	void HandleDeviceGroupComboBoxSelectionChanged( ILauncherDeviceGroupPtr Selection, ESelectInfo::Type SelectInfo )
@@ -204,27 +99,13 @@ private:
 	}
 
 	// Callback for when the selected item in the device group combo box has been renamed.
-	void HandleDeviceGroupComboBoxSelectionRenamed( const FText& CommittedText, ETextCommit::Type )
-	{
-		DeviceGroupComboBox->GetSelectedItem()->SetName(CommittedText.ToString());
-	}
+	void HandleDeviceGroupComboBoxSelectionRenamed( const FText& CommittedText, ETextCommit::Type );
 
 	// Callback for getting the text of a device group combo box widget.
-	FString HandleDeviceGroupComboBoxWidgetText( ILauncherDeviceGroupPtr Group ) const
-	{
-		if (Group.IsValid())
-		{
-			return Group->GetName();
-		}
-
-		return FString();
-	}
+	FString HandleDeviceGroupComboBoxWidgetText( ILauncherDeviceGroupPtr Group ) const;
 
 	// Callback for changing the list of groups in the profile manager.
-	void HandleProfileManagerDeviceGroupsChanged( const ILauncherDeviceGroupRef& /*ChangedProfile*/ )
-	{
-		DeviceGroupComboBox->RefreshOptions();
-	}
+	void HandleProfileManagerDeviceGroupsChanged( const ILauncherDeviceGroupRef& /*ChangedProfile*/ );
 
 private:
 
@@ -240,5 +121,3 @@ private:
 	FOnProjectLauncherDeviceGroupSelected OnGroupSelected;
 };
 
-
-#undef LOCTEXT_NAMESPACE
diff --git a/Engine/Source/Developer/ProjectLauncher/Private/Widgets/Launch/SProjectLauncherLaunchRoleEditor.h b/Engine/Source/Developer/ProjectLauncher/Private/Widgets/Launch/SProjectLauncherLaunchRoleEditor.h
--- a/Engine/Source/Developer/ProjectLauncher/Private/Widgets/Launch/SProjectLauncherLaunchRoleEditor.h
+++ b/Engine/Source/Developer/ProjectLauncher/Private/Widgets/Launch/SProjectLauncherLaunchRoleEditor.h
@@ -3,8 +3,6 @@
 #pragma once
 
 
-#define LOCTEXT_NAMESPACE "SProjectLauncherLaunchRoleEditor"
-
 
 /**
  * Implements the settings panel for a single launch role.
@@ -38,212 +36,14 @@ public:
 	 * @param InArgs The Slate argument list.
 	 * @param InRole The launch role to edit.
 	 */
-	BEGIN_SLATE_FUNCTION_BUILD_OPTIMIZATION
-	void Construct(	const FArguments& InArgs )
-	{
-		AvailableCultures = InArgs._AvailableCultures;
-		AvailableMaps = InArgs._AvailableMaps;
-		Role = InArgs._InitialRole;
-
-		// create instance types menu
-		FMenuBuilder InstanceTypeMenuBuilder(true, NULL);
-		{
-			FUIAction StandaloneClientAction(FExecuteAction::CreateSP(this, &SProjectLauncherLaunchRoleEditor::HandleInstanceTypeMenuEntryClicked, ELauncherProfileRoleInstanceTypes::StandaloneClient));
-			InstanceTypeMenuBuilder.AddMenuEntry(LOCTEXT("StandaloneClient", "Standalone Client"), LOCTEXT("StandaloneClientActionHint", "Launch this instance as a standalone game client."), FSlateIcon(), StandaloneClientAction);
-
-			FUIAction ListenServerAction(FExecuteAction::CreateSP(this, &SProjectLauncherLaunchRoleEditor::HandleInstanceTypeMenuEntryClicked, ELauncherProfileRoleInstanceTypes::ListenServer));
-			InstanceTypeMenuBuilder.AddMenuEntry(LOCTEXT("ListenServer", "Listen Server"), LOCTEXT("ListenServerActionHint", "Launch this instance as a game client that can accept connections from other clients."), FSlateIcon(), ListenServerAction);
-
-			FUIAction DedicatedServerAction(FExecuteAction::CreateSP(this, &SProjectLauncherLaunchRoleEditor::HandleInstanceTypeMenuEntryClicked, ELauncherProfileRoleInstanceTypes::DedicatedServer));
-			InstanceTypeMenuBuilder.AddMenuEntry(LOCTEXT("DedicatedServer", "Dedicated Server"), LOCTEXT("DedicatedServerActionHint", "Launch this instance as a dedicated game server."), FSlateIcon(), DedicatedServerAction);
-
-			FUIAction UnrealEditorAction(FExecuteAction::CreateSP(this, &SProjectLauncherLaunchRoleEditor::HandleInstanceTypeMenuEntryClicked, ELauncherProfileRoleInstanceTypes::UnrealEditor));
-			InstanceTypeMenuBuilder.AddMenuEntry(LOCTEXT("UnrealEditor", "Unreal Editor"), LOCTEXT("UnrealEditorActionHint", "Launch this instance as an Unreal Editor."), FSlateIcon(), UnrealEditorAction);
-		}
-
-		ChildSlot
-		[
-			SNew(SVerticalBox)
-
-			+ SVerticalBox::Slot()
-				.AutoHeight()
-				[
-					SNew(SProjectLauncherFormLabel)
-						.LabelText(LOCTEXT("InstanceTypeComboBoxLabel", "Launch As:"))
-				]
-
-			+ SVerticalBox::Slot()
-				.AutoHeight()
-				.Padding(0.0f, 2.0f, 0.0f, 0.0f)
-				[
-					// instance type menu
-					SNew(SComboButton)
-						.ButtonContent()
-						[
-							SNew(STextBlock)
-								.Text(this, &SProjectLauncherLaunchRoleEditor::HandleInstanceTypeComboButtonContentText)
-						]
-						.ContentPadding(FMargin(6.0f, 2.0f))
-						.MenuContent()
-						[
-							InstanceTypeMenuBuilder.MakeWidget()
-						]
-				]
-									
-			+ SVerticalBox::Slot()
-				.AutoHeight()
-				.Padding(0.0f, 8.0f, 0.0f, 0.0f)
-				[
-					SNew(SProjectLauncherFormLabel)
-						.ErrorToolTipText(LOCTEXT("CultureNotAvailableError", "The selected culture is not being cooked or is not available."))
-						.ErrorVisibility(this, &SProjectLauncherLaunchRoleEditor::HandleCultureValidationErrorIconVisibility)
-						.LabelText(LOCTEXT("InitialCultureTextBoxLabel", "Initial Culture:"))
-				]
-
-			+ SVerticalBox::Slot()
-				.AutoHeight()
-				.Padding(0.0f, 4.0f, 0.0f, 0.0f)
-				[
-					// initial culture combo box
-					SAssignNew(CultureComboBox, STextComboBox)
-						.ColorAndOpacity(this, &SProjectLauncherLaunchRoleEditor::HandleCultureComboBoxColorAndOpacity)
-						.OptionsSource(&CultureList)
-						.OnSelectionChanged(this, &SProjectLauncherLaunchRoleEditor::HandleCultureComboBoxSelectionChanged)
-				]
-
-			+ SVerticalBox::Slot()
-				.AutoHeight()
-				.Padding(0.0f, 8.0f, 0.0f, 0.0f)
-				[
-					SNew(SProjectLauncherFormLabel)
-						.ErrorToolTipText(LOCTEXT("MapNotAvailableError", "The selected map is not being cooked or is not available."))
-						.ErrorVisibility(this, &SProjectLauncherLaunchRoleEditor::HandleMapValidationErrorIconVisibility)
-						.LabelText(LOCTEXT("InitialMapTextBoxLabel", "Initial Map:"))
-				]
-
-			+ SVerticalBox::Slot()
-				.AutoHeight()
-				.Padding(0.0f, 4.0f, 0.0f, 0.0f)
-				[
-					// initial map combo box
-					SAssignNew(MapComboBox, STextComboBox)
-						.ColorAndOpacity(this, &SProjectLauncherLaunchRoleEditor::HandleMapComboBoxColorAndOpacity)
-						.OptionsSource(&MapList)
-						.OnSelectionChanged(this, &SProjectLauncherLaunchRoleEditor::HandleMapComboBoxSelectionChanged)
-				]
-
-			+ SVerticalBox::Slot()
-				.AutoHeight()
-				.Padding(0.0f, 8.0f, 0.0f, 0.0f)
-				[
-					SNew(SProjectLauncherFormLabel)
-						.LabelText(LOCTEXT("CommandLineTextBoxLabel", "Additional Command Line Parameters:"))
-				]
-
-			+ SVerticalBox::Slot()
-				.AutoHeight()
-				.Padding(0.0f, 4.0f, 0.0f, 0.0f)
-				[
-					// command line text box
-					SAssignNew(CommandLineTextBox, SEditableTextBox)
-						.OnTextChanged(this, &SProjectLauncherLaunchRoleEditor::HandleCommandLineTextBoxTextChanged)
-				]
-
-			+ SVerticalBox::Slot()
-				.AutoHeight()
-				.Padding(0.0f, 12.0f, 0.0f, 0.0f)
-				[
-					// v-sync check box
-					SNew(SCheckBox)
-						.IsChecked(this, &SProjectLauncherLaunchRoleEditor::HandleVsyncCheckBoxIsChecked)
-						.OnCheckStateChanged(this, &SProjectLauncherLaunchRoleEditor::HandleVsyncCheckBoxCheckStateChanged)
-						.Padding(FMargin(4.0f, 0.0f))
-						.Content()
-						[
-							SNew(STextBlock)
-								.Text(LOCTEXT("VsyncCheckBoxText", "Synchronize Screen Refresh Rate (VSync)"))
-						]
-				]
-		];
-
-		Refresh(InArgs._InitialRole);
-	}
-	END_SLATE_FUNCTION_BUILD_OPTIMIZATION
+	void Construct(	const FArguments& InArgs );
 
 	/**
 	 * Refreshes the widget.
 	 *
 	 * @param InRole The role to edit, or NULL if no role is being edited.
 	 */
-	void Refresh( const ILauncherProfileLaunchRolePtr& InRole )
-	{
-		Role = InRole;
-
-		ILauncherProfileLaunchRolePtr RolePtr = Role.Pin();
-
-		CultureList.Reset();
-		CultureList.Add(MakeShareable(new FString(LOCTEXT("DefaultCultureText", "<default>").ToString())));
-
-		MapList.Reset();
-		MapList.Add(MakeShareable(new FString(LOCTEXT("DefaultMapText", "<default>").ToString())));
-
-		if (RolePtr.IsValid())
-		{
-			// refresh widgets
-			CommandLineTextBox->SetText(FText::FromString(*RolePtr->GetCommandLine()));
-
-			// rebuild culture list
-			if (AvailableCultures != NULL)
-			{
-				for (int32 AvailableCultureIndex = 0; AvailableCultureIndex < AvailableCultures->Num(); ++AvailableCultureIndex)
-				{
-					TSharedPtr<FString> Culture = MakeShareable(new FString((*AvailableCultures)[AvailableCultureIndex]));
-
-					CultureList.Add(Culture);
-
-					if (*Culture == RolePtr->GetInitialCulture())
-					{
-						CultureComboBox->SetSelectedItem(Culture);
-					}
-				}
-			}
-
-			if (RolePtr->GetInitialCulture().IsEmpty() || (CultureList.Num() == 0))
-			{
-				CultureComboBox->SetSelectedItem(CultureList[0]);
-			}
-
-			// rebuild map list
-			if (AvailableMaps != NULL)
-			{
-				for (int32 AvailableMapIndex = 0; AvailableMapIndex < AvailableMaps->Num(); ++AvailableMapIndex)
-				{
-					TSharedPtr<FString> Map = MakeShareable(new FString((*AvailableMaps)[AvailableMapIndex]));
-
-					MapList.Add(Map);
-
-					if (*Map == RolePtr->GetInitialMap())
-					{
-						MapComboBox->SetSelectedItem(Map);
-					}
-				}
-			}
-
-			if (RolePtr->GetInitialMap().IsEmpty() || (MapList.Num() == 0))
-			{
-				MapComboBox->SetSelectedItem(MapList[0]);
-			}
-		}
-		else
-		{
-			CommandLineTextBox->SetText(FText::GetEmpty());
-			CultureComboBox->ClearSelection();
-			MapComboBox->ClearSelection();
-		}
-
-		CultureComboBox->RefreshOptions();
-		MapComboBox->RefreshOptions();
-	}
+	void Refresh( const ILauncherProfileLaunchRolePtr& InRole );
 
 private:
 
@@ -458,6 +258,3 @@ private:
 	// Holds a pointer to the role that is being edited in this widget.
 	TWeakPtr<ILauncherProfileLaunchRole> Role;
 };
-
-
-#undef LOCTEXT_NAMESPACE
diff --git a/Engine/Source/Developer/ProjectLauncher/Private/Widgets/SProjectLauncher.cpp b/Engine/Source/Developer/ProjectLauncher/Private/Widgets/SProjectLauncher.cpp
--- a/Engine/Source/Developer/ProjectLauncher/Private/Widgets/SProjectLauncher.cpp
+++ b/Engine/Source/Developer/ProjectLauncher/Private/Widgets/SProjectLauncher.cpp
@@ -1,6 +1,7 @@
 // Copyright 1998-2014 Epic Games, Inc. All Rights Reserved.
 
 #include "ProjectLauncherPrivatePCH.h"
+#include "SWidgetSwitcher.h"
 
 
 #define LOCTEXT_NAMESPACE "SProjectLauncher"
diff --git a/Engine/Source/Developer/ProjectLauncher/Private/Widgets/SProjectLauncherSettings.cpp b/Engine/Source/Developer/ProjectLauncher/Private/Widgets/SProjectLauncherSettings.cpp
--- a/Engine/Source/Developer/ProjectLauncher/Private/Widgets/SProjectLauncherSettings.cpp
+++ b/Engine/Source/Developer/ProjectLauncher/Private/Widgets/SProjectLauncherSettings.cpp
@@ -1,6 +1,7 @@
 // Copyright 1998-2014 Epic Games, Inc. All Rights Reserved.
 
 #include "ProjectLauncherPrivatePCH.h"
+#include "SInlineEditableTextBlock.h"
 
 #define LOCTEXT_NAMESPACE "SProjectLauncherSettings"
 
diff --git a/Engine/Source/Developer/ProjectLauncher/Private/Widgets/SProjectLauncherTaskListRow.h b/Engine/Source/Developer/ProjectLauncher/Private/Widgets/SProjectLauncherTaskListRow.h
--- a/Engine/Source/Developer/ProjectLauncher/Private/Widgets/SProjectLauncherTaskListRow.h
+++ b/Engine/Source/Developer/ProjectLauncher/Private/Widgets/SProjectLauncherTaskListRow.h
@@ -42,71 +42,7 @@ public:
 	 * @param ColumnName The name of the column to generate the widget for.
 	 * @return The widget.
 	 */
-	BEGIN_SLATE_FUNCTION_BUILD_OPTIMIZATION
-	virtual TSharedRef<SWidget> GenerateWidgetForColumn( const FName& ColumnName ) override
-	{
-		if (ColumnName == "Duration")
-		{
-			return SNew(SBox)
-				.Padding(FMargin(4.0, 0.0))
-				.VAlign((VAlign_Center))
-				[
-					SNew(STextBlock)
-						.Text(this, &SProjectLauncherTaskListRow::HandleDurationText)
-				];
-		}
-		else if (ColumnName == "Icon")
-		{
-			return SNew(SOverlay)
-
-			+ SOverlay::Slot()
-				.HAlign(HAlign_Center)
-				.VAlign((VAlign_Center))
-				[
-					SNew(SThrobber)
-						.Animate(SThrobber::VerticalAndOpacity)
-						.NumPieces(1)
-						.Visibility(this, &SProjectLauncherTaskListRow::HandleThrobberVisibility)
-				]
-
-			+ SOverlay::Slot()
-				.HAlign(HAlign_Center)
-				.VAlign((VAlign_Center))
-				[
-					SNew(SImage)
-						.ColorAndOpacity(this, &SProjectLauncherTaskListRow::HandleIconColorAndOpacity)
-						.Image(this, &SProjectLauncherTaskListRow::HandleIconImage)
-				];
-		}
-		else if (ColumnName == "Status")
-		{
-			return SNew(SBox)
-				.Padding(FMargin(4.0, 0.0))
-				.VAlign((VAlign_Center))
-				[
-					SNew(STextBlock)
-						.Text(this, &SProjectLauncherTaskListRow::HandleStatusText)
-				];
-		}
-		else if (ColumnName == "Task")
-		{
-			ILauncherTaskPtr TaskPtr = Task.Pin();
-
-			if (TaskPtr.IsValid())
-			{
-				return SNew(SBox)
-					.Padding(FMargin(4.0, 0.0))
-					.VAlign((VAlign_Center))
-					[
-						SNew(STextBlock)
-							.Text(TaskPtr->GetDesc())
-					];
-			}
-		}
-
-		return SNullWidget::NullWidget;
-	}
-	END_SLATE_FUNCTION_BUILD_OPTIMIZATION
+	virtual TSharedRef<SWidget> GenerateWidgetForColumn( const FName& ColumnName ) override;
 
 private:
 
diff --git a/Engine/Source/Developer/ProjectLauncher/Private/Widgets/Shared/SProjectLauncherProfileNameDescEditor.h b/Engine/Source/Developer/ProjectLauncher/Private/Widgets/Shared/SProjectLauncherProfileNameDescEditor.h
--- a/Engine/Source/Developer/ProjectLauncher/Private/Widgets/Shared/SProjectLauncherProfileNameDescEditor.h
+++ b/Engine/Source/Developer/ProjectLauncher/Private/Widgets/Shared/SProjectLauncherProfileNameDescEditor.h
@@ -3,8 +3,6 @@
 #pragma once
 
 
-#define LOCTEXT_NAMESPACE "SProjectLauncherProfileNameDescEditor"
-
 /**
  * Implements a build configuration selector widget.
  */
@@ -25,70 +23,12 @@ public:
 	 * @param InArgs The Slate argument list.
 	 * @param InModel The data model.
 	 */
-	void Construct(const FArguments& InArgs, bool InShowAddDescriptionText)
-	{
-		EnterTextDescription = FText(LOCTEXT("LaunchProfileEnterDescription", "Enter a description here."));
-
-		LaunchProfileAttr = InArgs._LaunchProfile;
-		bShowAddDescriptionText = InShowAddDescriptionText;
-
-		ChildSlot
-		[
-			SNew(SHorizontalBox)
-
-			+ SHorizontalBox::Slot()
-			.AutoWidth()
-			.VAlign(VAlign_Center)
-			[
-				SNew(SBox)
-				.WidthOverride(40)
-				.HeightOverride(40)
-				[
-					SNew(SImage)
-					.Image(this, &SProjectLauncherProfileNameDescEditor::HandleProfileImage)
-				]
-			]
-
-			+ SHorizontalBox::Slot()
-			.FillWidth(1)
-			.VAlign(VAlign_Center)
-			[
-				SNew(SVerticalBox)
-
-				+ SVerticalBox::Slot()
-				.AutoHeight()
-				.Padding(2, 4, 2, 4)
-				[
-					SAssignNew(NameEditableTextBlock, SInlineEditableTextBlock)
-					.Text(this, &SProjectLauncherProfileNameDescEditor::OnGetNameText)
-					.OnTextCommitted(this, &SProjectLauncherProfileNameDescEditor::OnNameTextCommitted)
-					.Cursor(EMouseCursor::TextEditBeam)
-				]
-
-				+ SVerticalBox::Slot()
-				.AutoHeight()
-				.Padding(2, 4, 2, 4)
-				[
-					SNew(SInlineEditableTextBlock)
-					.Text(this, &SProjectLauncherProfileNameDescEditor::OnGetDescriptionText)
-					.Style(FCoreStyle::Get(), "InlineEditableTextBlockSmallStyle")
-					.OnTextCommitted(this, &SProjectLauncherProfileNameDescEditor::OnDescriptionTextCommitted)
-					.Cursor(EMouseCursor::TextEditBeam)
-				]
-			]
-		];
-	}
+	void Construct(const FArguments& InArgs, bool InShowAddDescriptionText);
 
 	/**
 	 * Triggers a name edit for the profile.
 	 */
-	void TriggerNameEdit()
-	{
-		if (NameEditableTextBlock.IsValid())
-		{
-			NameEditableTextBlock->EnterEditingMode();
-		}
-	}
+	void TriggerNameEdit();
 
 private:
 
@@ -167,5 +107,3 @@ private:
 	TSharedPtr<SInlineEditableTextBlock> NameEditableTextBlock;
 };
 
-
-#undef LOCTEXT_NAMESPACE
diff --git a/Engine/Source/Developer/SessionFrontend/Private/Widgets/Console/SSessionConsoleCommandBar.cpp b/Engine/Source/Developer/SessionFrontend/Private/Widgets/Console/SSessionConsoleCommandBar.cpp
--- a/Engine/Source/Developer/SessionFrontend/Private/Widgets/Console/SSessionConsoleCommandBar.cpp
+++ b/Engine/Source/Developer/SessionFrontend/Private/Widgets/Console/SSessionConsoleCommandBar.cpp
@@ -1,6 +1,7 @@
 // Copyright 1998-2014 Epic Games, Inc. All Rights Reserved.
 
 #include "SessionFrontendPrivatePCH.h"
+#include "SSuggestionTextBox.h"
 
 
 #define LOCTEXT_NAMESPACE "SSessionConsoleCommandBar"
diff --git a/Engine/Source/Developer/SessionFrontend/Private/Widgets/Console/SSessionConsoleFilterBar.cpp b/Engine/Source/Developer/SessionFrontend/Private/Widgets/Console/SSessionConsoleFilterBar.cpp
--- a/Engine/Source/Developer/SessionFrontend/Private/Widgets/Console/SSessionConsoleFilterBar.cpp
+++ b/Engine/Source/Developer/SessionFrontend/Private/Widgets/Console/SSessionConsoleFilterBar.cpp
@@ -1,7 +1,7 @@
 // Copyright 1998-2014 Epic Games, Inc. All Rights Reserved.
 
 #include "SessionFrontendPrivatePCH.h"
-
+#include "SSearchBox.h"
 
 #define LOCTEXT_NAMESPACE "SSessionConsoleFilterBar"
 
@@ -273,5 +273,10 @@ void SSessionConsoleFilterBar::HandleVerbosityFilterStateChanged( ELogVerbosity:
 	OnFilterChanged.ExecuteIfBound();
 }
 
+FText SSessionConsoleFilterBar::GetFilterText() const
+{
+	return FilterStringTextBox->GetText();
+}
+
 
 #undef LOCTEXT_NAMESPACE
diff --git a/Engine/Source/Developer/SessionFrontend/Private/Widgets/Console/SSessionConsoleFilterBar.h b/Engine/Source/Developer/SessionFrontend/Private/Widgets/Console/SSessionConsoleFilterBar.h
--- a/Engine/Source/Developer/SessionFrontend/Private/Widgets/Console/SSessionConsoleFilterBar.h
+++ b/Engine/Source/Developer/SessionFrontend/Private/Widgets/Console/SSessionConsoleFilterBar.h
@@ -40,10 +40,7 @@ public:
 	 *
 	 * @return Filter string.
 	 */
-	FText GetFilterText( ) const
-	{
-		return FilterStringTextBox->GetText();
-	}
+	FText GetFilterText( ) const;
 
 	/**
 	 * Resets the categories and filter counters.
diff --git a/Engine/Source/Developer/SessionFrontend/Private/Widgets/Console/SSessionConsoleShortcutWindow.cpp b/Engine/Source/Developer/SessionFrontend/Private/Widgets/Console/SSessionConsoleShortcutWindow.cpp
--- a/Engine/Source/Developer/SessionFrontend/Private/Widgets/Console/SSessionConsoleShortcutWindow.cpp
+++ b/Engine/Source/Developer/SessionFrontend/Private/Widgets/Console/SSessionConsoleShortcutWindow.cpp
@@ -1,6 +1,7 @@
 // Copyright 1998-2014 Epic Games, Inc. All Rights Reserved.
 
 #include "SessionFrontendPrivatePCH.h"
+#include "STextEntryPopup.h"
 
 
 #define LOCTEXT_NAMESPACE "SSessionConsoleShortcutWindow"
diff --git a/Engine/Source/Developer/SettingsEditor/Private/SettingsEditorModule.cpp b/Engine/Source/Developer/SettingsEditor/Private/SettingsEditorModule.cpp
--- a/Engine/Source/Developer/SettingsEditor/Private/SettingsEditorModule.cpp
+++ b/Engine/Source/Developer/SettingsEditor/Private/SettingsEditorModule.cpp
@@ -1,6 +1,7 @@
 // Copyright 1998-2014 Epic Games, Inc. All Rights Reserved.
 
 #include "SettingsEditorPrivatePCH.h"
+#include "SSettingsEditor.h"
 
 
 /**
diff --git a/Engine/Source/Developer/SettingsEditor/Private/SettingsEditorPrivatePCH.h b/Engine/Source/Developer/SettingsEditor/Private/SettingsEditorPrivatePCH.h
--- a/Engine/Source/Developer/SettingsEditor/Private/SettingsEditorPrivatePCH.h
+++ b/Engine/Source/Developer/SettingsEditor/Private/SettingsEditorPrivatePCH.h
@@ -20,4 +20,3 @@
 
 #include "SettingsEditorModel.h"
 
-#include "SSettingsEditor.h"
diff --git a/Engine/Source/Developer/SettingsEditor/Private/Widgets/SSettingsEditor.cpp b/Engine/Source/Developer/SettingsEditor/Private/Widgets/SSettingsEditor.cpp
--- a/Engine/Source/Developer/SettingsEditor/Private/Widgets/SSettingsEditor.cpp
+++ b/Engine/Source/Developer/SettingsEditor/Private/Widgets/SSettingsEditor.cpp
@@ -1,10 +1,14 @@
 // Copyright 1998-2014 Epic Games, Inc. All Rights Reserved.
 
 #include "SettingsEditorPrivatePCH.h"
+#include "SSettingsEditor.h"
 #include "EngineAnalytics.h"
 #include "AnalyticsEventAttribute.h"
 #include "IAnalyticsProvider.h"
 #include "SSettingsEditorCheckoutNotice.h"
+#include "SHyperlink.h"
+#include "SNotificationList.h"
+#include "NotificationManager.h"
 
 #define LOCTEXT_NAMESPACE "SSettingsEditor"
 
diff --git a/Engine/Source/Developer/SharedSettingsWidgets/Private/SPlatformSetupMessage.cpp b/Engine/Source/Developer/SharedSettingsWidgets/Private/SPlatformSetupMessage.cpp
--- a/Engine/Source/Developer/SharedSettingsWidgets/Private/SPlatformSetupMessage.cpp
+++ b/Engine/Source/Developer/SharedSettingsWidgets/Private/SPlatformSetupMessage.cpp
@@ -4,6 +4,9 @@
 #include "SPlatformSetupMessage.h"
 #include "ISourceControlModule.h"
 #include "EditorStyleSet.h"
+#include "SWidgetSwitcher.h"
+#include "SNotificationList.h"
+#include "NotificationManager.h"
 
 #define LOCTEXT_NAMESPACE "SPlatformSetupMessage"
 
diff --git a/Engine/Source/Developer/SharedSettingsWidgets/Private/SSettingsEditorCheckoutNotice.cpp b/Engine/Source/Developer/SharedSettingsWidgets/Private/SSettingsEditorCheckoutNotice.cpp
--- a/Engine/Source/Developer/SharedSettingsWidgets/Private/SSettingsEditorCheckoutNotice.cpp
+++ b/Engine/Source/Developer/SharedSettingsWidgets/Private/SSettingsEditorCheckoutNotice.cpp
@@ -4,6 +4,10 @@
 #include "EditorStyle.h"
 #include "SSettingsEditorCheckoutNotice.h"
 #include "ISourceControlModule.h"
+#include "SWidgetSwitcher.h"
+#include "SThrobber.h"
+#include "SNotificationList.h"
+#include "NotificationManager.h"
 
 
 #define LOCTEXT_NAMESPACE "SSettingsEditorCheckoutNotice"
diff --git a/Engine/Source/Developer/SlateReflector/Private/Widgets/SWidgetReflector.h b/Engine/Source/Developer/SlateReflector/Private/Widgets/SWidgetReflector.h
--- a/Engine/Source/Developer/SlateReflector/Private/Widgets/SWidgetReflector.h
+++ b/Engine/Source/Developer/SlateReflector/Private/Widgets/SWidgetReflector.h
@@ -2,15 +2,12 @@
 
 #pragma once
 
-
-/** A Delegate for passing along a string of a source code location to access */
-DECLARE_DELEGATE_RetVal_ThreeParams(bool, FAccessSourceCode, const FString& /*FileName*/, int32 /*InLineNumber*/, int32 /*InColumnNumber*/);
-
+#include "IWidgetReflector.h"
 
 /**
  * Abstract base class for widget reflectors.
  */
-class SWidgetReflector
+class SLATEREFLECTOR_API SWidgetReflector
 	: public SCompoundWidget
 	, public IWidgetReflector
 {
diff --git a/Engine/Source/Developer/SlateReflector/Private/Widgets/SWidgetReflectorTreeWidgetItem.cpp b/Engine/Source/Developer/SlateReflector/Private/Widgets/SWidgetReflectorTreeWidgetItem.cpp
new file mode 100644
--- /dev/null
+++ b/Engine/Source/Developer/SlateReflector/Private/Widgets/SWidgetReflectorTreeWidgetItem.cpp
@@ -0,0 +1,93 @@
+// Copyright 1998-2014 Epic Games, Inc. All Rights Reserved.
+
+#include "SlateReflectorPrivatePCH.h"
+#include "SWidgetReflectorTreeWidgetItem.h"
+#include "SHyperlink.h"
+
+BEGIN_SLATE_FUNCTION_BUILD_OPTIMIZATION
+TSharedRef<SWidget> SReflectorTreeWidgetItem::GenerateWidgetForColumn(const FName& ColumnName)
+{
+	if (ColumnName == TEXT("WidgetName"))
+	{
+		return SNew(SHorizontalBox)
+
+			+ SHorizontalBox::Slot()
+			.AutoWidth()
+			.VAlign(VAlign_Center)
+			[
+				SNew(SExpanderArrow, SharedThis(this))
+			]
+
+		+ SHorizontalBox::Slot()
+			.AutoWidth()
+			.Padding(2.0f, 0.0f)
+			.VAlign(VAlign_Center)
+			[
+				SNew(STextBlock)
+				.Text(this, &SReflectorTreeWidgetItem::GetWidgetType)
+				.ColorAndOpacity(this, &SReflectorTreeWidgetItem::GetTint)
+			];
+	}
+	else if (ColumnName == TEXT("WidgetInfo"))
+	{
+		return SNew(SBox)
+			.HAlign(HAlign_Left)
+			.VAlign(VAlign_Center)
+			.Padding(FMargin(2.0f, 0.0f))
+			[
+				SNew(SHyperlink)
+				.Text(this, &SReflectorTreeWidgetItem::GetReadableLocation)
+				.OnNavigate(this, &SReflectorTreeWidgetItem::HandleHyperlinkNavigate)
+			];
+	}
+	else if (ColumnName == "Visibility")
+	{
+		return SNew(SBox)
+			.HAlign(HAlign_Center)
+			.VAlign(VAlign_Center)
+			.Padding(FMargin(2.0f, 0.0f))
+			[
+				SNew(STextBlock)
+				.Text(this, &SReflectorTreeWidgetItem::GetVisibilityAsString)
+			];
+	}
+	else if (ColumnName == "ForegroundColor")
+	{
+		TSharedPtr<SWidget> ThisWidget = WidgetInfo.Get()->Widget.Pin();
+		FSlateColor Foreground = (ThisWidget.IsValid())
+			? ThisWidget->GetForegroundColor()
+			: FSlateColor::UseForeground();
+
+		return SNew(SBorder)
+			// Show unset color as an empty space.
+			.Visibility(Foreground.IsColorSpecified() ? EVisibility::Visible : EVisibility::Hidden)
+			// Show a checkerboard background so we can see alpha values well
+			.BorderImage(FCoreStyle::Get().GetBrush("Checkerboard"))
+			.VAlign(VAlign_Center)
+			.Padding(FMargin(2.0f, 0.0f))
+			[
+				// Show a color block
+				SNew(SColorBlock)
+				.Color(Foreground.GetSpecifiedColor())
+				.Size(FVector2D(16.0f, 16.0f))
+			];
+	}
+	else if (ColumnName == "Address")
+	{
+		const TSharedPtr<SWidget> TheWidget = WidgetInfo.Get()->Widget.Pin();
+		const FText Address = (TheWidget.IsValid())
+			? FText::FromString( FString::Printf(TEXT("0x%08X"), TheWidget.Get()) ) 
+			: NSLOCTEXT("SWidgetReflector","nullptr","nullptr");
+			
+		return
+			SNew(SHyperlink)
+			.ToolTipText(NSLOCTEXT("SWidgetReflector", "ClickToCopy", "Click to copy address."))
+			.Text( Address )
+			.OnNavigate_Lambda( [Address](){ FPlatformMisc::ClipboardCopy(*Address.ToString()); } ) ;
+	}
+	else
+	{
+		return SNullWidget::NullWidget;
+	}
+}
+END_SLATE_FUNCTION_BUILD_OPTIMIZATION
diff --git a/Engine/Source/Developer/SlateReflector/Private/Widgets/SWidgetReflectorTreeWidgetItem.h b/Engine/Source/Developer/SlateReflector/Private/Widgets/SWidgetReflectorTreeWidgetItem.h
--- a/Engine/Source/Developer/SlateReflector/Private/Widgets/SWidgetReflectorTreeWidgetItem.h
+++ b/Engine/Source/Developer/SlateReflector/Private/Widgets/SWidgetReflectorTreeWidgetItem.h
@@ -6,7 +6,7 @@
 /**
  * Widget that visualizes the contents of a FReflectorNode
  */
-class SReflectorTreeWidgetItem
+class SLATEREFLECTOR_API SReflectorTreeWidgetItem
 	: public SMultiColumnTableRow<TSharedPtr<FReflectorNode>>
 {
 public:
@@ -36,93 +36,7 @@ public:
 		SMultiColumnTableRow< TSharedPtr<FReflectorNode> >::Construct( SMultiColumnTableRow< TSharedPtr<FReflectorNode> >::FArguments().Padding(1), InOwnerTableView );
 	}
 
-	BEGIN_SLATE_FUNCTION_BUILD_OPTIMIZATION
-	virtual TSharedRef<SWidget> GenerateWidgetForColumn( const FName& ColumnName ) override
-	{
-		if (ColumnName == TEXT("WidgetName"))
-		{
-			return SNew(SHorizontalBox)
-			
-			+ SHorizontalBox::Slot()
-				.AutoWidth()
-				.VAlign(VAlign_Center)
-				[
-					SNew(SExpanderArrow, SharedThis(this))
-				]
-
-			+ SHorizontalBox::Slot()
-				.AutoWidth()
-				.Padding(2.0f, 0.0f)
-				.VAlign(VAlign_Center)
-				[
-					SNew(STextBlock)
-						.Text(this, &SReflectorTreeWidgetItem::GetWidgetType)
-						.ColorAndOpacity(this, &SReflectorTreeWidgetItem::GetTint)
-				];
-		}
-		else if (ColumnName == TEXT("WidgetInfo"))
-		{
-			return SNew(SBox)
-				.HAlign(HAlign_Left)
-				.VAlign(VAlign_Center)
-				.Padding(FMargin(2.0f, 0.0f))
-				[
-					SNew(SHyperlink)
-						.Text(this, &SReflectorTreeWidgetItem::GetReadableLocation)
-						.OnNavigate(this, &SReflectorTreeWidgetItem::HandleHyperlinkNavigate)
-				];
-		}
-		else if (ColumnName == "Visibility")
-		{
-			return SNew(SBox)
-				.HAlign(HAlign_Center)
-				.VAlign(VAlign_Center)
-				.Padding(FMargin(2.0f, 0.0f))
-				[
-					SNew(STextBlock)
-						.Text(this, &SReflectorTreeWidgetItem::GetVisibilityAsString)
-				];
-		}
-		else if (ColumnName == "ForegroundColor")
-		{
-			TSharedPtr<SWidget> ThisWidget = WidgetInfo.Get()->Widget.Pin();
-			FSlateColor Foreground = (ThisWidget.IsValid())
-				? ThisWidget->GetForegroundColor()
-				: FSlateColor::UseForeground();
-
-			return SNew(SBorder)
-				// Show unset color as an empty space.
-				.Visibility(Foreground.IsColorSpecified() ? EVisibility::Visible : EVisibility::Hidden)
-				// Show a checkerboard background so we can see alpha values well
-				.BorderImage( FCoreStyle::Get().GetBrush("Checkerboard") )
-				.VAlign(VAlign_Center)
-				.Padding(FMargin(2.0f, 0.0f))
-				[
-					// Show a color block
-					SNew(SColorBlock)
-						.Color(Foreground.GetSpecifiedColor())
-						.Size(FVector2D(16.0f, 16.0f))
-				];
-		}
-		else if (ColumnName == "Address")
-		{
-			const TSharedPtr<SWidget> TheWidget = WidgetInfo.Get()->Widget.Pin();
-			const FText Address = (TheWidget.IsValid())
-				? FText::FromString( FString::Printf(TEXT("0x%08X"), TheWidget.Get()) ) 
-				: NSLOCTEXT("SWidgetReflector","nullptr","nullptr");
-			
-			return
-				SNew(SHyperlink)
-				.ToolTipText(NSLOCTEXT("SWidgetReflector", "ClickToCopy", "Click to copy address."))
-				.Text( Address )
-				.OnNavigate_Lambda( [Address](){ FPlatformMisc::ClipboardCopy(*Address.ToString()); } ) ;
-		}
-		else
-		{
-			return SNullWidget::NullWidget;
-		}
-	}
-	END_SLATE_FUNCTION_BUILD_OPTIMIZATION
+	virtual TSharedRef<SWidget> GenerateWidgetForColumn( const FName& ColumnName ) override;
 
 protected:
 
diff --git a/Engine/Source/Developer/SourceControl/Private/SSourceControlLogin.cpp b/Engine/Source/Developer/SourceControl/Private/SSourceControlLogin.cpp
--- a/Engine/Source/Developer/SourceControl/Private/SSourceControlLogin.cpp
+++ b/Engine/Source/Developer/SourceControl/Private/SSourceControlLogin.cpp
@@ -5,6 +5,10 @@
 #include "SourceControlModule.h"
 #include "SSourceControlPicker.h"
 #include "MessageLog.h"
+#include "SExpandableArea.h"
+#include "SThrobber.h"
+#include "SNotificationList.h"
+#include "NotificationManager.h"
 
 #if SOURCE_CONTROL_WITH_SLATE
 
diff --git a/Engine/Source/Developer/SourceControl/Private/ScopedSourceControlProgress.cpp b/Engine/Source/Developer/SourceControl/Private/ScopedSourceControlProgress.cpp
--- a/Engine/Source/Developer/SourceControl/Private/ScopedSourceControlProgress.cpp
+++ b/Engine/Source/Developer/SourceControl/Private/ScopedSourceControlProgress.cpp
@@ -2,6 +2,7 @@
 
 #include "SourceControlPrivatePCH.h"
 #include "ScopedSourceControlProgress.h"
+#include "SThrobber.h"
 
 #define LOCTEXT_NAMESPACE "SourceControlProgress"
 
diff --git a/Engine/Source/Developer/TaskGraph/Private/SEventsTree.cpp b/Engine/Source/Developer/TaskGraph/Private/SEventsTree.cpp
--- a/Engine/Source/Developer/TaskGraph/Private/SEventsTree.cpp
+++ b/Engine/Source/Developer/TaskGraph/Private/SEventsTree.cpp
@@ -1,9 +1,10 @@
 // Copyright 1998-2014 Epic Games, Inc. All Rights Reserved.
 
-#include "Slate.h"
+#include "SlateBasics.h"
 #include "TaskGraphInterfaces.h"
 #include "VisualizerEvents.h"
 #include "SEventsTree.h"
+#include "SSearchBox.h"
 
 #define LOCTEXT_NAMESPACE "SEventsTree"
 
diff --git a/Engine/Source/Developer/TaskGraph/Private/STaskGraph.cpp b/Engine/Source/Developer/TaskGraph/Private/STaskGraph.cpp
--- a/Engine/Source/Developer/TaskGraph/Private/STaskGraph.cpp
+++ b/Engine/Source/Developer/TaskGraph/Private/STaskGraph.cpp
@@ -1,12 +1,13 @@
 // Copyright 1998-2014 Epic Games, Inc. All Rights Reserved.
 
-#include "Slate.h"
+#include "SlateBasics.h"
 #include "TaskGraphInterfaces.h"
 #include "VisualizerEvents.h"
 #include "STaskGraph.h"
 #include "SGraphBar.h"
 #include "Tickable.h"
 #include "SProfileVisualizer.h"
+#include "SDockTab.h"
 
 /**
  * Creates Visualizer using Visualizer profile data format
diff --git a/Engine/Source/Editor/AnimGraph/Private/AnimationStateMachineSchema.cpp b/Engine/Source/Editor/AnimGraph/Private/AnimationStateMachineSchema.cpp
--- a/Engine/Source/Editor/AnimGraph/Private/AnimationStateMachineSchema.cpp
+++ b/Engine/Source/Editor/AnimGraph/Private/AnimationStateMachineSchema.cpp
@@ -21,6 +21,7 @@
 #include "EdGraphUtilities.h"
 #include "AssetData.h"
 #include "Kismet2/KismetEditorUtilities.h"
+#include "GenericCommands.h"
 
 #define LOCTEXT_NAMESPACE "AnimationStateMachineSchema"
 
diff --git a/Engine/Source/Editor/BehaviorTreeEditor/Private/BehaviorTreeEditor.cpp b/Engine/Source/Editor/BehaviorTreeEditor/Private/BehaviorTreeEditor.cpp
--- a/Engine/Source/Editor/BehaviorTreeEditor/Private/BehaviorTreeEditor.cpp
+++ b/Engine/Source/Editor/BehaviorTreeEditor/Private/BehaviorTreeEditor.cpp
@@ -33,6 +33,8 @@
 #include "AssetRegistryModule.h"
 #include "AssetToolsModule.h"
 #include "ContentBrowserModule.h"
+#include "SDockTab.h"
+#include "GenericCommands.h"
 
 #define LOCTEXT_NAMESPACE "BehaviorTreeEditor"
 
diff --git a/Engine/Source/Editor/BehaviorTreeEditor/Private/SBehaviorTreeBlackboardView.cpp b/Engine/Source/Editor/BehaviorTreeEditor/Private/SBehaviorTreeBlackboardView.cpp
--- a/Engine/Source/Editor/BehaviorTreeEditor/Private/SBehaviorTreeBlackboardView.cpp
+++ b/Engine/Source/Editor/BehaviorTreeEditor/Private/SBehaviorTreeBlackboardView.cpp
@@ -12,6 +12,7 @@
 #include "BehaviorTreeEditorCommands.h"
 #include "ClassViewerModule.h"
 #include "ClassViewerFilter.h"
+#include "SInlineEditableTextBlock.h"
 
 #define LOCTEXT_NAMESPACE "SBehaviorTreeBlackboardView"
 
diff --git a/Engine/Source/Editor/Cascade/Private/Cascade.cpp b/Engine/Source/Editor/Cascade/Private/Cascade.cpp
--- a/Engine/Source/Editor/Cascade/Private/Cascade.cpp
+++ b/Engine/Source/Editor/Cascade/Private/Cascade.cpp
@@ -34,6 +34,12 @@
 
 #include "Runtime/Analytics/Analytics/Public/Interfaces/IAnalyticsProvider.h"
 #include "EngineAnalytics.h"
+#include "SDockTab.h"
+#include "SNotificationList.h"
+#include "NotificationManager.h"
+#include "SNumericEntryBox.h"
+#include "STextEntryPopup.h"
+#include "GenericCommands.h"
 
 static const FName Cascade_PreviewViewportTab("Cascade_PreviewViewport");
 static const FName Cascade_EmmitterCanvasTab("Cascade_EmitterCanvas");
diff --git a/Engine/Source/Editor/ContentBrowser/Private/AssetContextMenu.cpp b/Engine/Source/Editor/ContentBrowser/Private/AssetContextMenu.cpp
--- a/Engine/Source/Editor/ContentBrowser/Private/AssetContextMenu.cpp
+++ b/Engine/Source/Editor/ContentBrowser/Private/AssetContextMenu.cpp
@@ -29,6 +29,9 @@
 #include "EditorClassUtils.h"
 
 #include "SColorPicker.h"
+#include "GenericCommands.h"
+#include "SNotificationList.h"
+#include "NotificationManager.h"
 
 
 #define LOCTEXT_NAMESPACE "ContentBrowser"
diff --git a/Engine/Source/Editor/ContentBrowser/Private/SContentBrowser.cpp b/Engine/Source/Editor/ContentBrowser/Private/SContentBrowser.cpp
--- a/Engine/Source/Editor/ContentBrowser/Private/SContentBrowser.cpp
+++ b/Engine/Source/Editor/ContentBrowser/Private/SContentBrowser.cpp
@@ -9,6 +9,8 @@
 #include "ContentBrowserCommands.h"
 #include "CollectionManagerModule.h"
 #include "AssetRegistryModule.h"
+#include "SDockTab.h"
+#include "GenericCommands.h"
 
 #define LOCTEXT_NAMESPACE "ContentBrowser"
 
@@ -867,7 +869,7 @@ void SContentBrowser::OnContainingTabClosed(TSharedRef<SDockTab> DockTab)
 	FContentBrowserSingleton::Get().ContentBrowserClosed( SharedThis(this) );
 }
 
-void SContentBrowser::OnContainingTabActivated(TSharedRef<SDockTab> DockTab, ETabActivationCause::Type InActivationCause)
+void SContentBrowser::OnContainingTabActivated(TSharedRef<SDockTab> DockTab, ETabActivationCause InActivationCause)
 {
 	if(InActivationCause == ETabActivationCause::UserClickedOnTab)
 	{
diff --git a/Engine/Source/Editor/ContentBrowser/Private/SContentBrowser.h b/Engine/Source/Editor/ContentBrowser/Private/SContentBrowser.h
--- a/Engine/Source/Editor/ContentBrowser/Private/SContentBrowser.h
+++ b/Engine/Source/Editor/ContentBrowser/Private/SContentBrowser.h
@@ -82,7 +82,7 @@ private:
 	void OnContainingTabClosed(TSharedRef<SDockTab> DockTab);
 
 	/** Called when a containing tab is activated, if there is one */
-	void OnContainingTabActivated(TSharedRef<SDockTab> DockTab, ETabActivationCause::Type InActivationCause);
+	void OnContainingTabActivated(TSharedRef<SDockTab> DockTab, ETabActivationCause InActivationCause);
 
 	/** Loads settings from config based on the browser's InstanceName*/
 	void LoadSettings(const FName& InstanceName);
diff --git a/Engine/Source/Editor/ContentBrowser/Private/SourcesViewWidgets.cpp b/Engine/Source/Editor/ContentBrowser/Private/SourcesViewWidgets.cpp
--- a/Engine/Source/Editor/ContentBrowser/Private/SourcesViewWidgets.cpp
+++ b/Engine/Source/Editor/ContentBrowser/Private/SourcesViewWidgets.cpp
@@ -9,6 +9,7 @@
 #include "DragAndDrop/AssetPathDragDropOp.h"
 #include "ContentBrowserUtils.h"
 #include "CollectionViewUtils.h"
+#include "SInlineEditableTextBlock.h"
 
 #define LOCTEXT_NAMESPACE "ContentBrowser"
 
@@ -609,8 +610,4 @@ const FSlateBrush* SCollectionListItem::GetBorderImage() const
 }
 
 
-//////////////////////////
-// SSourcesTab
-//////////////////////////
-
 #undef LOCTEXT_NAMESPACE
diff --git a/Engine/Source/Editor/ContentBrowser/Private/SourcesViewWidgets.h b/Engine/Source/Editor/ContentBrowser/Private/SourcesViewWidgets.h
--- a/Engine/Source/Editor/ContentBrowser/Private/SourcesViewWidgets.h
+++ b/Engine/Source/Editor/ContentBrowser/Private/SourcesViewWidgets.h
@@ -232,9 +232,3 @@ private:
 	/** Widget to display the name of the collection item and allows for renaming */
 	TSharedPtr< SInlineEditableTextBlock > InlineRenameWidget;
 };
-
-/** A tab which may not be dragged */
-class SSourcesTab : public SDockableTab
-{
-
-};
\ No newline at end of file
diff --git a/Engine/Source/Editor/DataTableEditor/Private/DataTableEditor.cpp b/Engine/Source/Editor/DataTableEditor/Private/DataTableEditor.cpp
--- a/Engine/Source/Editor/DataTableEditor/Private/DataTableEditor.cpp
+++ b/Engine/Source/Editor/DataTableEditor/Private/DataTableEditor.cpp
@@ -5,6 +5,8 @@
 #include "DataTableEditor.h"
 #include "Toolkits/IToolkitHost.h"
 #include "Editor/WorkspaceMenuStructure/Public/WorkspaceMenuStructureModule.h"
+#include "SSearchBox.h"
+#include "SDockTab.h"
  
 #define LOCTEXT_NAMESPACE "DataTableEditor"
 
diff --git a/Engine/Source/Editor/DestructibleMeshEditor/Private/SDestructibleMeshEditorViewport.cpp b/Engine/Source/Editor/DestructibleMeshEditor/Private/SDestructibleMeshEditorViewport.cpp
--- a/Engine/Source/Editor/DestructibleMeshEditor/Private/SDestructibleMeshEditorViewport.cpp
+++ b/Engine/Source/Editor/DestructibleMeshEditor/Private/SDestructibleMeshEditorViewport.cpp
@@ -11,6 +11,7 @@
 #include "FbxImporter.h"
 #include "ComponentReregisterContext.h"
 #include "Engine/DestructibleMesh.h"
+#include "SDockableTab.h"
 
 DEFINE_LOG_CATEGORY_STATIC(LogDestructibleMeshEditor, Log, All);
 
diff --git a/Engine/Source/Editor/DetailCustomizations/Private/ActorDetails.cpp b/Engine/Source/Editor/DetailCustomizations/Private/ActorDetails.cpp
--- a/Engine/Source/Editor/DetailCustomizations/Private/ActorDetails.cpp
+++ b/Engine/Source/Editor/DetailCustomizations/Private/ActorDetails.cpp
@@ -29,6 +29,9 @@
 #include "ActorDetailsDelegates.h"
 #include "EditorCategoryUtils.h"
 #include "Engine/DocumentationActor.h"
+#include "SHyperlink.h"
+#include "SNotificationList.h"
+#include "NotificationManager.h"
 
 #define LOCTEXT_NAMESPACE "ActorDetails"
 
diff --git a/Engine/Source/Editor/DetailCustomizations/Private/AnimTransitionNodeDetails.cpp b/Engine/Source/Editor/DetailCustomizations/Private/AnimTransitionNodeDetails.cpp
--- a/Engine/Source/Editor/DetailCustomizations/Private/AnimTransitionNodeDetails.cpp
+++ b/Engine/Source/Editor/DetailCustomizations/Private/AnimTransitionNodeDetails.cpp
@@ -11,6 +11,8 @@
 #include "AnimTransitionNodeDetails.h"
 #include "KismetEditorUtilities.h"
 #include "SKismetLinearExpression.h"
+#include "STextEntryPopup.h"
+#include "SExpandableArea.h"
 
 #define LOCTEXT_NAMESPACE "FAnimStateNodeDetails"
 
diff --git a/Engine/Source/Editor/DetailCustomizations/Private/CameraDetails.cpp b/Engine/Source/Editor/DetailCustomizations/Private/CameraDetails.cpp
--- a/Engine/Source/Editor/DetailCustomizations/Private/CameraDetails.cpp
+++ b/Engine/Source/Editor/DetailCustomizations/Private/CameraDetails.cpp
@@ -2,6 +2,7 @@
 
 #include "DetailCustomizationsPrivatePCH.h"
 #include "CameraDetails.h"
+#include "SNumericEntryBox.h"
 
 #define LOCTEXT_NAMESPACE "CameraDetails"
 
diff --git a/Engine/Source/Editor/DetailCustomizations/Private/CurveTableCustomization.h b/Engine/Source/Editor/DetailCustomizations/Private/CurveTableCustomization.h
--- a/Engine/Source/Editor/DetailCustomizations/Private/CurveTableCustomization.h
+++ b/Engine/Source/Editor/DetailCustomizations/Private/CurveTableCustomization.h
@@ -105,33 +105,7 @@ private:
 	}
 
 	/** Returns the ListView for the ComboButton */
-	TSharedRef<SWidget> GetListContent()
-	{
-		SAssignNew( RowNameComboListView, SListView<TSharedPtr<FString> > )
-			.ListItemsSource( &RowNames )
-			.OnSelectionChanged( this, &FCurveTableCustomizationLayout::OnSelectionChanged )
-			.OnGenerateRow( this, &FCurveTableCustomizationLayout::HandleRowNameComboBoxGenarateWidget )
-			.SelectionMode(ESelectionMode::Single);
-
-
-		if( CurrentSelectedItem.IsValid() )
-		{
-			RowNameComboListView->SetSelection(CurrentSelectedItem);
-		}
-
-		return SNew( SVerticalBox )
-				+SVerticalBox::Slot()
-				.AutoHeight()
-				[
-					SNew( SSearchBox )
-					.OnTextChanged( this, &FCurveTableCustomizationLayout::OnFilterTextChanged )
-				]
-				+SVerticalBox::Slot()
-				.FillHeight( 1.f )
-				[
-					RowNameComboListView.ToSharedRef()
-				];				
-	}
+	TSharedRef<SWidget> GetListContent();
 
 	/** Delegate to refresh the drop down when the datatable changes */
 	void OnCurveTableChanged()
diff --git a/Engine/Source/Editor/DetailCustomizations/Private/DataTableCustomization.cpp b/Engine/Source/Editor/DetailCustomizations/Private/DataTableCustomization.cpp
new file mode 100644
--- /dev/null
+++ b/Engine/Source/Editor/DetailCustomizations/Private/DataTableCustomization.cpp
@@ -0,0 +1,35 @@
+// Copyright 1998-2014 Epic Games, Inc. All Rights Reserved.
+
+#include "DetailCustomizationsPrivatePCH.h"
+#include "DataTableCustomization.h"
+#include "SSearchBox.h"
+
+TSharedRef<SWidget> FDataTableCustomizationLayout::GetListContent()
+{
+	SAssignNew(RowNameComboListView, SListView<TSharedPtr<FString> >)
+		.ListItemsSource(&RowNames)
+		.OnSelectionChanged(this, &FDataTableCustomizationLayout::OnSelectionChanged)
+		.OnGenerateRow(this, &FDataTableCustomizationLayout::HandleRowNameComboBoxGenarateWidget)
+		.SelectionMode(ESelectionMode::Single);
+
+	// Ensure no filter is applied at the time the menu opens
+	OnFilterTextChanged(FText::GetEmpty());
+
+	if (CurrentSelectedItem.IsValid())
+	{
+		RowNameComboListView->SetSelection(CurrentSelectedItem);
+	}
+
+	return SNew(SVerticalBox)
+		+ SVerticalBox::Slot()
+		.AutoHeight()
+		[
+			SNew(SSearchBox)
+			.OnTextChanged(this, &FDataTableCustomizationLayout::OnFilterTextChanged)
+		]
+	+ SVerticalBox::Slot()
+		.FillHeight(1.f)
+		[
+			RowNameComboListView.ToSharedRef()
+		];
+}
diff --git a/Engine/Source/Editor/DetailCustomizations/Private/DataTableCustomization.h b/Engine/Source/Editor/DetailCustomizations/Private/DataTableCustomization.h
--- a/Engine/Source/Editor/DetailCustomizations/Private/DataTableCustomization.h
+++ b/Engine/Source/Editor/DetailCustomizations/Private/DataTableCustomization.h
@@ -109,35 +109,7 @@ private:
 	}
 
 	/** Returns the ListView for the ComboButton */
-	TSharedRef<SWidget> GetListContent()
-	{
-		SAssignNew( RowNameComboListView, SListView<TSharedPtr<FString> > )
-			.ListItemsSource( &RowNames )
-			.OnSelectionChanged( this, &FDataTableCustomizationLayout::OnSelectionChanged )			
-			.OnGenerateRow( this, &FDataTableCustomizationLayout::HandleRowNameComboBoxGenarateWidget )
-			.SelectionMode(ESelectionMode::Single);
-
-		// Ensure no filter is applied at the time the menu opens
-		OnFilterTextChanged(FText::GetEmpty());
-
-		if( CurrentSelectedItem.IsValid() )
-		{
-			RowNameComboListView->SetSelection(CurrentSelectedItem);
-		}
-
-		return SNew( SVerticalBox )
-				+SVerticalBox::Slot()
-				.AutoHeight()
-				[
-					SNew( SSearchBox )
-					.OnTextChanged( this, &FDataTableCustomizationLayout::OnFilterTextChanged )
-				]
-				+SVerticalBox::Slot()
-				.FillHeight( 1.f )
-				[
-					RowNameComboListView.ToSharedRef()
-				];				
-	}
+	TSharedRef<SWidget> GetListContent();
 
 	/** Delegate to refresh the drop down when the datatable changes */
 	void OnDataTableChanged()
diff --git a/Engine/Source/Editor/DetailCustomizations/Private/FbxImportUIDetails.cpp b/Engine/Source/Editor/DetailCustomizations/Private/FbxImportUIDetails.cpp
--- a/Engine/Source/Editor/DetailCustomizations/Private/FbxImportUIDetails.cpp
+++ b/Engine/Source/Editor/DetailCustomizations/Private/FbxImportUIDetails.cpp
@@ -3,6 +3,7 @@
 #include "DetailCustomizationsPrivatePCH.h"
 #include "FbxImportUIDetails.h"
 #include "Factories/FbxAnimSequenceImportData.h"
+#include "STextComboBox.h"
 
 #define LOCTEXT_NAMESPACE "FbxImportUIDetails"
 
diff --git a/Engine/Source/Editor/DeviceProfileEditor/Private/SDeviceProfileEditor.cpp b/Engine/Source/Editor/DeviceProfileEditor/Private/SDeviceProfileEditor.cpp
--- a/Engine/Source/Editor/DeviceProfileEditor/Private/SDeviceProfileEditor.cpp
+++ b/Engine/Source/Editor/DeviceProfileEditor/Private/SDeviceProfileEditor.cpp
@@ -7,6 +7,8 @@
 #include "ISourceControlModule.h"
 #include "PropertyEditorModule.h"
 #include "PropertyPath.h"
+#include "SDockTab.h"
+#include "SWidgetSwitcher.h"
 
 
 #define LOCTEXT_NAMESPACE "DeviceProfileEditor"
diff --git a/Engine/Source/Editor/DistCurveEditor/Private/SDistributionCurveEditor.cpp b/Engine/Source/Editor/DistCurveEditor/Private/SDistributionCurveEditor.cpp
--- a/Engine/Source/Editor/DistCurveEditor/Private/SDistributionCurveEditor.cpp
+++ b/Engine/Source/Editor/DistCurveEditor/Private/SDistributionCurveEditor.cpp
@@ -8,6 +8,10 @@
 #include "SDistributionCurveEditor.h"
 #include "Runtime/Engine/Public/Slate/SceneViewport.h"
 #include "SColorPicker.h"
+#include "STextComboBox.h"
+#include "STextEntryPopup.h"
+#include "SNotificationList.h"
+#include "NotificationManager.h"
 
 #define LOCTEXT_NAMESPACE "CurveEditor"
 
diff --git a/Engine/Source/Editor/Documentation/Private/SDocumentationToolTip.cpp b/Engine/Source/Editor/Documentation/Private/SDocumentationToolTip.cpp
--- a/Engine/Source/Editor/Documentation/Private/SDocumentationToolTip.cpp
+++ b/Engine/Source/Editor/Documentation/Private/SDocumentationToolTip.cpp
@@ -7,6 +7,9 @@
 #include "Developer/SourceControl/Public/SourceControlHelpers.h"
 #include "EngineAnalytics.h"
 #include "IAnalyticsProvider.h"
+#include "SHyperlink.h"
+#include "SNotificationList.h"
+#include "NotificationManager.h"
 
 void SDocumentationToolTip::Construct( const FArguments& InArgs )
 {
diff --git a/Engine/Source/Editor/Documentation/Private/UDNParser.cpp b/Engine/Source/Editor/Documentation/Private/UDNParser.cpp
--- a/Engine/Source/Editor/Documentation/Private/UDNParser.cpp
+++ b/Engine/Source/Editor/Documentation/Private/UDNParser.cpp
@@ -10,6 +10,9 @@
 #include "AssetEditorManager.h"
 #include "SourceCodeNavigation.h"
 #include "DesktopPlatformModule.h"
+#include "SNotificationList.h"
+#include "NotificationManager.h"
+#include "SHyperlink.h"
 
 #define LOCTEXT_NAMESPACE "IntroTutorials"
 
diff --git a/Engine/Source/Editor/GameProjectGeneration/Private/SGameProjectDialog.cpp b/Engine/Source/Editor/GameProjectGeneration/Private/SGameProjectDialog.cpp
--- a/Engine/Source/Editor/GameProjectGeneration/Private/SGameProjectDialog.cpp
+++ b/Engine/Source/Editor/GameProjectGeneration/Private/SGameProjectDialog.cpp
@@ -8,6 +8,7 @@
 #include "EngineAnalytics.h"
 #include "AnalyticsEventAttribute.h"
 #include "IAnalyticsProvider.h"
+#include "SWidgetSwitcher.h"
 
 #define LOCTEXT_NAMESPACE "GameProjectGeneration"
 
