Fixed some Clang compile errors

diff --git a/Engine/Plugins/2D/Paper2D/Source/Paper2D/Private/PaperCharacter.cpp b/Engine/Plugins/2D/Paper2D/Source/Paper2D/Private/PaperCharacter.cpp
--- a/Engine/Plugins/2D/Paper2D/Source/Paper2D/Private/PaperCharacter.cpp
+++ b/Engine/Plugins/2D/Paper2D/Source/Paper2D/Private/PaperCharacter.cpp
@@ -19,7 +19,7 @@ APaperCharacter::APaperCharacter(const FObjectInitializer& ObjectInitializer)
 		Sprite->bOwnerNoSee = false;
 		Sprite->bAffectDynamicIndirectLighting = true;
 		Sprite->PrimaryComponentTick.TickGroup = TG_PrePhysics;
-		Sprite->AttachParent = CapsuleComponent;
+		Sprite->AttachParent = GetCapsuleComponent();
 		static FName CollisionProfileName(TEXT("CharacterMesh"));
 		Sprite->SetCollisionProfileName(CollisionProfileName);
 		Sprite->bGenerateOverlapEvents = false;
@@ -35,9 +35,9 @@ void APaperCharacter::PostInitializeComponents()
 		if (Sprite)
 		{
 			// force animation tick after movement component updates
-			if (Sprite->PrimaryComponentTick.bCanEverTick && CharacterMovement.IsValid())
+			if (Sprite->PrimaryComponentTick.bCanEverTick && GetCharacterMovement())
 			{
-				Sprite->PrimaryComponentTick.AddPrerequisite(CharacterMovement, CharacterMovement->PrimaryComponentTick);
+				Sprite->PrimaryComponentTick.AddPrerequisite(GetCharacterMovement(), GetCharacterMovement()->PrimaryComponentTick);
 			}
 		}
 	}
diff --git a/Engine/Source/Editor/NiagaraEditor/Private/NiagaraCompiler_VectorVM.cpp b/Engine/Source/Editor/NiagaraEditor/Private/NiagaraCompiler_VectorVM.cpp
--- a/Engine/Source/Editor/NiagaraEditor/Private/NiagaraCompiler_VectorVM.cpp
+++ b/Engine/Source/Editor/NiagaraEditor/Private/NiagaraCompiler_VectorVM.cpp
@@ -593,3 +593,5 @@ void FNiagaraCompiler_VectorVM::Inverse_Internal(TArray<TNiagaraExprPtr>& InputE
 
 	
 }
+
+#undef LOCTEXT_NAMESPACE
