More TSubobjectPtr removal: removing TSubobjectPtrs from Paper2D

diff --git a/Engine/Plugins/2D/Paper2D/Source/Paper2D/Classes/PaperCharacter.h b/Engine/Plugins/2D/Paper2D/Source/Paper2D/Classes/PaperCharacter.h
--- a/Engine/Plugins/2D/Paper2D/Source/Paper2D/Classes/PaperCharacter.h
+++ b/Engine/Plugins/2D/Paper2D/Source/Paper2D/Classes/PaperCharacter.h
@@ -14,9 +14,14 @@ class PAPER2D_API APaperCharacter : public ACharacter
 	// Name of the Sprite component
 	static FName SpriteComponentName;
 
+private:
 	/** The main skeletal mesh associated with this Character (optional sub-object). */
-	UPROPERTY(Category = Character, VisibleAnywhere, BlueprintReadOnly)
-	TSubobjectPtr<class UPaperFlipbookComponent> Sprite;
+	UPROPERTY(Category = Character, VisibleAnywhere, BlueprintReadOnly, meta = (AllowPrivateAccess = "true"))
+	class UPaperFlipbookComponent* Sprite;
+public:
 
 	virtual void PostInitializeComponents() override;
+
+	/** Returns Sprite subobject **/
+	FORCEINLINE class UPaperFlipbookComponent* GetSprite() const { return Sprite; }
 };
diff --git a/Engine/Plugins/2D/Paper2D/Source/Paper2D/Classes/PaperFlipbookActor.h b/Engine/Plugins/2D/Paper2D/Source/Paper2D/Classes/PaperFlipbookActor.h
--- a/Engine/Plugins/2D/Paper2D/Source/Paper2D/Classes/PaperFlipbookActor.h
+++ b/Engine/Plugins/2D/Paper2D/Source/Paper2D/Classes/PaperFlipbookActor.h
@@ -10,13 +10,18 @@ class APaperFlipbookActor : public AActor
 {
 	GENERATED_UCLASS_BODY()
 
-	UPROPERTY(Category=Sprite, VisibleAnywhere, BlueprintReadOnly, meta=(ExposeFunctionCategories="Sprite,Rendering,Physics,Components|Flipbook"))
-	TSubobjectPtr<class UPaperFlipbookComponent> RenderComponent;
+private:
+	UPROPERTY(Category = Sprite, VisibleAnywhere, BlueprintReadOnly, meta = (ExposeFunctionCategories = "Sprite,Rendering,Physics,Components|Flipbook", AllowPrivateAccess = "true"))
+	class UPaperFlipbookComponent* RenderComponent;
+public:
 
 	// AActor interface
 #if WITH_EDITOR
 	virtual bool GetReferencedContentObjects(TArray<UObject*>& Objects) const override;
 #endif
 	// End of AActor interface
+
+	/** Returns RenderComponent subobject **/
+	FORCEINLINE class UPaperFlipbookComponent* GetRenderComponent() const { return RenderComponent; }
 };
 
diff --git a/Engine/Plugins/2D/Paper2D/Source/Paper2D/Classes/PaperSpriteActor.h b/Engine/Plugins/2D/Paper2D/Source/Paper2D/Classes/PaperSpriteActor.h
--- a/Engine/Plugins/2D/Paper2D/Source/Paper2D/Classes/PaperSpriteActor.h
+++ b/Engine/Plugins/2D/Paper2D/Source/Paper2D/Classes/PaperSpriteActor.h
@@ -16,14 +16,19 @@ class APaperSpriteActor : public AActor
 {
 	GENERATED_UCLASS_BODY()
 
-	UPROPERTY(Category=Sprite, VisibleAnywhere, BlueprintReadOnly, meta=(ExposeFunctionCategories="Sprite,Rendering,Physics,Components|Sprite"))
-	TSubobjectPtr<class UPaperSpriteComponent> RenderComponent;
+private:
+	UPROPERTY(Category = Sprite, VisibleAnywhere, BlueprintReadOnly, meta = (ExposeFunctionCategories = "Sprite,Rendering,Physics,Components|Sprite", AllowPrivateAccess = "true"))
+	class UPaperSpriteComponent* RenderComponent;
+public:
 
 	// AActor interface
 #if WITH_EDITOR
 	virtual bool GetReferencedContentObjects(TArray<UObject*>& Objects) const override;
 #endif
 	// End of AActor interface
+
+	/** Returns RenderComponent subobject **/
+	FORCEINLINE class UPaperSpriteComponent* GetRenderComponent() const { return RenderComponent; }
 };
 
 // Allow the old name to continue to work for one release
diff --git a/Engine/Plugins/2D/Paper2D/Source/Paper2D/Classes/PaperTerrainActor.h b/Engine/Plugins/2D/Paper2D/Source/Paper2D/Classes/PaperTerrainActor.h
--- a/Engine/Plugins/2D/Paper2D/Source/Paper2D/Classes/PaperTerrainActor.h
+++ b/Engine/Plugins/2D/Paper2D/Source/Paper2D/Classes/PaperTerrainActor.h
@@ -15,18 +15,27 @@ class PAPER2D_API APaperTerrainActor : public AActor
 {
 	GENERATED_UCLASS_BODY()
 
+private:
 	UPROPERTY()
-	TSubobjectPtr<class USceneComponent> DummyRoot;
+	class USceneComponent* DummyRoot;
 
 	UPROPERTY()
-	TSubobjectPtr<class UPaperTerrainSplineComponent> SplineComponent;
+	class UPaperTerrainSplineComponent* SplineComponent;
 
-	UPROPERTY(Category=Sprite, VisibleAnywhere, BlueprintReadOnly, meta=(ExposeFunctionCategories = "Sprite,Rendering,Physics,Components|Sprite"))
-	TSubobjectPtr<class UPaperTerrainComponent> RenderComponent;
+	UPROPERTY(Category = Sprite, VisibleAnywhere, BlueprintReadOnly, meta = (ExposeFunctionCategories = "Sprite,Rendering,Physics,Components|Sprite", AllowPrivateAccess = "true"))
+	class UPaperTerrainComponent* RenderComponent;
+public:
 
 	// AActor interface
 #if WITH_EDITOR
 	virtual bool GetReferencedContentObjects(TArray<UObject*>& Objects) const override;
 #endif
 	// End of AActor interface
+
+	/** Returns DummyRoot subobject **/
+	FORCEINLINE class USceneComponent* GetDummyRoot() const { return DummyRoot; }
+	/** Returns SplineComponent subobject **/
+	FORCEINLINE class UPaperTerrainSplineComponent* GetSplineComponent() const { return SplineComponent; }
+	/** Returns RenderComponent subobject **/
+	FORCEINLINE class UPaperTerrainComponent* GetRenderComponent() const { return RenderComponent; }
 };
diff --git a/Engine/Plugins/2D/Paper2D/Source/Paper2D/Classes/PaperTileMapActor.h b/Engine/Plugins/2D/Paper2D/Source/Paper2D/Classes/PaperTileMapActor.h
--- a/Engine/Plugins/2D/Paper2D/Source/Paper2D/Classes/PaperTileMapActor.h
+++ b/Engine/Plugins/2D/Paper2D/Source/Paper2D/Classes/PaperTileMapActor.h
@@ -11,12 +11,17 @@ class APaperTileMapActor : public AActor
 {
 	GENERATED_UCLASS_BODY()
 
+private:
 	UPROPERTY(Category=TileMapActor, VisibleAnywhere)
-	TSubobjectPtr<class UPaperTileMapRenderComponent> RenderComponent;
+	class UPaperTileMapRenderComponent* RenderComponent;
+public:
 
 	// AActor interface
 #if WITH_EDITOR
 	virtual bool GetReferencedContentObjects(TArray<UObject*>& Objects) const override;
 #endif
 	// End of AActor interface
+
+	/** Returns RenderComponent subobject **/
+	FORCEINLINE class UPaperTileMapRenderComponent* GetRenderComponent() const { return RenderComponent; }
 };
diff --git a/Engine/Plugins/2D/Paper2D/Source/Paper2DEditor/Private/PaperAnimatedSpriteActorFactory.cpp b/Engine/Plugins/2D/Paper2D/Source/Paper2DEditor/Private/PaperAnimatedSpriteActorFactory.cpp
--- a/Engine/Plugins/2D/Paper2D/Source/Paper2DEditor/Private/PaperAnimatedSpriteActorFactory.cpp
+++ b/Engine/Plugins/2D/Paper2D/Source/Paper2DEditor/Private/PaperAnimatedSpriteActorFactory.cpp
@@ -20,7 +20,7 @@ void UPaperFlipbookActorFactory::PostSpawnActor(UObject* Asset, AActor* NewActor
 		GEditor->SetActorLabelUnique(NewActor, Flipbook->GetName());
 
 		APaperFlipbookActor* TypedActor = CastChecked<APaperFlipbookActor>(NewActor);
-		UPaperFlipbookComponent* RenderComponent = TypedActor->RenderComponent;
+		UPaperFlipbookComponent* RenderComponent = TypedActor->GetRenderComponent();
 		check(RenderComponent);
 
 		RenderComponent->UnregisterComponent();
@@ -35,7 +35,7 @@ void UPaperFlipbookActorFactory::PostCreateBlueprint(UObject* Asset, AActor* CDO
 	{
 		if (APaperFlipbookActor* TypedActor = Cast<APaperFlipbookActor>(CDO))
 		{
-			UPaperFlipbookComponent* RenderComponent = TypedActor->RenderComponent;
+			UPaperFlipbookComponent* RenderComponent = TypedActor->GetRenderComponent();
 			check(RenderComponent);
 
 			RenderComponent->SetFlipbook(Flipbook);
diff --git a/Engine/Plugins/2D/Paper2D/Source/Paper2DEditor/Private/PaperSpriteActorFactory.cpp b/Engine/Plugins/2D/Paper2D/Source/Paper2DEditor/Private/PaperSpriteActorFactory.cpp
--- a/Engine/Plugins/2D/Paper2D/Source/Paper2DEditor/Private/PaperSpriteActorFactory.cpp
+++ b/Engine/Plugins/2D/Paper2D/Source/Paper2DEditor/Private/PaperSpriteActorFactory.cpp
@@ -20,7 +20,7 @@ void UPaperSpriteActorFactory::PostSpawnActor(UObject* Asset, AActor* NewActor)
 		GEditor->SetActorLabelUnique(NewActor, Sprite->GetName());
 
 		APaperSpriteActor* TypedActor = CastChecked<APaperSpriteActor>(NewActor);
-		UPaperSpriteComponent* RenderComponent = TypedActor->RenderComponent;
+		UPaperSpriteComponent* RenderComponent = TypedActor->GetRenderComponent();
 		check(RenderComponent);
 
 		RenderComponent->UnregisterComponent();
@@ -41,7 +41,7 @@ void UPaperSpriteActorFactory::PostCreateBlueprint(UObject* Asset, AActor* CDO)
 	{
 		if (APaperSpriteActor* TypedActor = Cast<APaperSpriteActor>(CDO))
 		{
-			UPaperSpriteComponent* RenderComponent = TypedActor->RenderComponent;
+			UPaperSpriteComponent* RenderComponent = TypedActor->GetRenderComponent();
 			check(RenderComponent);
 
 			RenderComponent->SetSprite(Sprite);
diff --git a/Engine/Plugins/2D/Paper2D/Source/Paper2DEditor/Private/TerrainSplineEditing/TerrainSplineActorFactory.cpp b/Engine/Plugins/2D/Paper2D/Source/Paper2DEditor/Private/TerrainSplineEditing/TerrainSplineActorFactory.cpp
--- a/Engine/Plugins/2D/Paper2D/Source/Paper2DEditor/Private/TerrainSplineEditing/TerrainSplineActorFactory.cpp
+++ b/Engine/Plugins/2D/Paper2D/Source/Paper2DEditor/Private/TerrainSplineEditing/TerrainSplineActorFactory.cpp
@@ -20,7 +20,7 @@ void UTerrainSplineActorFactory::PostSpawnActor(UObject* Asset, AActor* NewActor
 		GEditor->SetActorLabelUnique(NewActor, TerrainMaterial->GetName());
 
 		APaperTerrainActor* TypedActor = CastChecked<APaperTerrainActor>(NewActor);
-		UPaperTerrainComponent* RenderComponent = TypedActor->RenderComponent;
+		UPaperTerrainComponent* RenderComponent = TypedActor->GetRenderComponent();
 		check(RenderComponent);
 
 		RenderComponent->UnregisterComponent();
diff --git a/Engine/Plugins/2D/Paper2D/Source/Paper2DEditor/Private/TileMapEditing/TileMapActorFactory.cpp b/Engine/Plugins/2D/Paper2D/Source/Paper2DEditor/Private/TileMapEditing/TileMapActorFactory.cpp
--- a/Engine/Plugins/2D/Paper2D/Source/Paper2DEditor/Private/TileMapEditing/TileMapActorFactory.cpp
+++ b/Engine/Plugins/2D/Paper2D/Source/Paper2DEditor/Private/TileMapEditing/TileMapActorFactory.cpp
@@ -16,7 +16,7 @@ UTileMapActorFactory::UTileMapActorFactory(const FObjectInitializer& ObjectIniti
 void UTileMapActorFactory::PostSpawnActor(UObject* Asset, AActor* NewActor)
 {
 	APaperTileMapActor* TypedActor = CastChecked<APaperTileMapActor>(NewActor);
-	UPaperTileMapRenderComponent* RenderComponent = TypedActor->RenderComponent;
+	UPaperTileMapRenderComponent* RenderComponent = TypedActor->GetRenderComponent();
 	check(RenderComponent);
 
 	if (UPaperTileMap* TileMap = Cast<UPaperTileMap>(Asset))
@@ -45,7 +45,7 @@ void UTileMapActorFactory::PostCreateBlueprint(UObject* Asset, AActor* CDO)
 {
 	if (APaperTileMapActor* TypedActor = Cast<APaperTileMapActor>(CDO))
 	{
-		UPaperTileMapRenderComponent* RenderComponent = TypedActor->RenderComponent;
+		UPaperTileMapRenderComponent* RenderComponent = TypedActor->GetRenderComponent();
 		check(RenderComponent);
 
 		if (UPaperTileMap* TileMap = Cast<UPaperTileMap>(Asset))
