Paper2D: Make the flipbook asset type color green instead of cyan

diff --git a/Engine/Plugins/2D/Paper2D/Source/Paper2DEditor/Private/FlipbookAssetTypeActions.cpp b/Engine/Plugins/2D/Paper2D/Source/Paper2DEditor/Private/FlipbookAssetTypeActions.cpp
--- a/Engine/Plugins/2D/Paper2D/Source/Paper2DEditor/Private/FlipbookAssetTypeActions.cpp
+++ b/Engine/Plugins/2D/Paper2D/Source/Paper2DEditor/Private/FlipbookAssetTypeActions.cpp
@@ -17,7 +17,7 @@ FText FFlipbookAssetTypeActions::GetName() const
 
 FColor FFlipbookAssetTypeActions::GetTypeColor() const
 {
-	return FColor(0, 255, 255);
+	return FColor(129, 196, 115);
 }
 
 UClass* FFlipbookAssetTypeActions::GetSupportedClass() const
