Paper2D: Add DLL export to various actor classes, allowing them to be derived from in other modules

diff --git a/Engine/Plugins/2D/Paper2D/Source/Paper2D/Classes/PaperCharacter.h b/Engine/Plugins/2D/Paper2D/Source/Paper2D/Classes/PaperCharacter.h
--- a/Engine/Plugins/2D/Paper2D/Source/Paper2D/Classes/PaperCharacter.h
+++ b/Engine/Plugins/2D/Paper2D/Source/Paper2D/Classes/PaperCharacter.h
@@ -1,12 +1,13 @@
 // Copyright 1998-2014 Epic Games, Inc. All Rights Reserved.
 
 #pragma once
+
 #include "GameFramework/Character.h"
 #include "PaperCharacter.generated.h"
 
 // APaperCharacter behaves like ACharacter, but uses a UPaperFlipbookComponent instead of a USkeletalMeshComponent as a visual representation
 // Note: The variable named Mesh will not be set up on this actor!
-UCLASS(EarlyAccessPreview)
+UCLASS()
 class PAPER2D_API APaperCharacter : public ACharacter
 {
 	GENERATED_UCLASS_BODY()
diff --git a/Engine/Plugins/2D/Paper2D/Source/Paper2D/Classes/PaperFlipbookActor.h b/Engine/Plugins/2D/Paper2D/Source/Paper2D/Classes/PaperFlipbookActor.h
--- a/Engine/Plugins/2D/Paper2D/Source/Paper2D/Classes/PaperFlipbookActor.h
+++ b/Engine/Plugins/2D/Paper2D/Source/Paper2D/Classes/PaperFlipbookActor.h
@@ -5,14 +5,21 @@
 #include "PaperFlipbookComponent.h"
 #include "PaperFlipbookActor.generated.h"
 
-UCLASS(MinimalAPI)
-class APaperFlipbookActor : public AActor
+/**
+ * An instance of a UPaperFlipbook in a level.
+ *
+ * This actor is created when you drag a flipbook asset from the content browser into the level, and
+ * it is just a thin wrapper around a UPaperFlipbookComponent that actually references the asset.
+ */
+UCLASS()
+class PAPER2D_API APaperFlipbookActor : public AActor
 {
 	GENERATED_UCLASS_BODY()
 
 private:
-	UPROPERTY(Category = Sprite, VisibleAnywhere, BlueprintReadOnly, meta = (ExposeFunctionCategories = "Sprite,Rendering,Physics,Components|Flipbook", AllowPrivateAccess = "true"))
+	UPROPERTY(Category=Sprite, VisibleAnywhere, BlueprintReadOnly, meta=(ExposeFunctionCategories="Sprite,Rendering,Physics,Components|Flipbook", AllowPrivateAccess="true"))
 	class UPaperFlipbookComponent* RenderComponent;
+
 public:
 
 	// AActor interface
diff --git a/Engine/Plugins/2D/Paper2D/Source/Paper2D/Classes/PaperSpriteActor.h b/Engine/Plugins/2D/Paper2D/Source/Paper2D/Classes/PaperSpriteActor.h
--- a/Engine/Plugins/2D/Paper2D/Source/Paper2D/Classes/PaperSpriteActor.h
+++ b/Engine/Plugins/2D/Paper2D/Source/Paper2D/Classes/PaperSpriteActor.h
@@ -11,13 +11,13 @@
  * This actor is created when you drag a sprite asset from the content browser into the level, and
  * it is just a thin wrapper around a UPaperSpriteComponent that actually references the asset.
  */
-UCLASS(MinimalAPI, meta=(ChildCanTick))
-class APaperSpriteActor : public AActor
+UCLASS(meta=(ChildCanTick))
+class PAPER2D_API APaperSpriteActor : public AActor
 {
 	GENERATED_UCLASS_BODY()
 
 private:
-	UPROPERTY(Category = Sprite, VisibleAnywhere, BlueprintReadOnly, meta = (ExposeFunctionCategories = "Sprite,Rendering,Physics,Components|Sprite", AllowPrivateAccess = "true"))
+	UPROPERTY(Category=Sprite, VisibleAnywhere, BlueprintReadOnly, meta=(ExposeFunctionCategories="Sprite,Rendering,Physics,Components|Sprite", AllowPrivateAccess="true"))
 	class UPaperSpriteComponent* RenderComponent;
 public:
 
diff --git a/Engine/Plugins/2D/Paper2D/Source/Paper2D/Classes/PaperTileMapActor.h b/Engine/Plugins/2D/Paper2D/Source/Paper2D/Classes/PaperTileMapActor.h
--- a/Engine/Plugins/2D/Paper2D/Source/Paper2D/Classes/PaperTileMapActor.h
+++ b/Engine/Plugins/2D/Paper2D/Source/Paper2D/Classes/PaperTileMapActor.h
@@ -6,13 +6,19 @@
 
 #include "PaperTileMapActor.generated.h"
 
-UCLASS(MinimalAPI)
-class APaperTileMapActor : public AActor
+/**
+ * An instance of a UPaperTileMap in a level.
+ *
+ * This actor is created when you drag a tile map asset from the content browser into the level, and
+ * it is just a thin wrapper around a UPaperTileMapRenderComponent that actually references the asset.
+ */
+UCLASS()
+class PAPER2D_API APaperTileMapActor : public AActor
 {
 	GENERATED_UCLASS_BODY()
 
 private:
-	UPROPERTY(Category=TileMapActor, VisibleAnywhere)
+	UPROPERTY(Category=TileMapActor, VisibleAnywhere, BlueprintReadOnly, meta=(ExposeFunctionCategories="Sprite,Rendering,Physics,Components|Sprite", AllowPrivateAccess="true"))
 	class UPaperTileMapRenderComponent* RenderComponent;
 public:
 
