Paper2D: Fix flipbook thumbnail preview bounds zooming in/out from frame to frame (now uses the union render bounds for all frames)

diff --git a/Engine/Plugins/2D/Paper2D/Source/Paper2DEditor/Classes/PaperSpriteThumbnailRenderer.h b/Engine/Plugins/2D/Paper2D/Source/Paper2DEditor/Classes/PaperSpriteThumbnailRenderer.h
--- a/Engine/Plugins/2D/Paper2D/Source/Paper2DEditor/Classes/PaperSpriteThumbnailRenderer.h
+++ b/Engine/Plugins/2D/Paper2D/Source/Paper2DEditor/Classes/PaperSpriteThumbnailRenderer.h
@@ -12,7 +12,7 @@ class UPaperSpriteThumbnailRenderer : public UDefaultSizedThumbnailRenderer
 	// End of UThumbnailRenderer interface
 
 protected:
-	void DrawFrame(class UPaperSprite* Sprite, int32 X, int32 Y, uint32 Width, uint32 Height, FRenderTarget*, FCanvas* Canvas);
+	void DrawFrame(class UPaperSprite* Sprite, int32 X, int32 Y, uint32 Width, uint32 Height, FRenderTarget*, FCanvas* Canvas, FBoxSphereBounds* OverrideRenderBounds);
 
 	void DrawGrid(int32 X, int32 Y, uint32 Width, uint32 Height, FCanvas* Canvas);
 };
diff --git a/Engine/Plugins/2D/Paper2D/Source/Paper2DEditor/Private/PaperAnimatedSpriteThumbnailRenderer.cpp b/Engine/Plugins/2D/Paper2D/Source/Paper2DEditor/Private/PaperAnimatedSpriteThumbnailRenderer.cpp
deleted file mode 100644
diff --git a/Engine/Plugins/2D/Paper2D/Source/Paper2DEditor/Private/PaperFlipbookThumbnailRenderer.cpp b/Engine/Plugins/2D/Paper2D/Source/Paper2DEditor/Private/PaperFlipbookThumbnailRenderer.cpp
new file mode 100644
--- /dev/null
+++ b/Engine/Plugins/2D/Paper2D/Source/Paper2DEditor/Private/PaperFlipbookThumbnailRenderer.cpp
@@ -0,0 +1,43 @@
+// Copyright 1998-2014 Epic Games, Inc. All Rights Reserved.
+
+#include "Paper2DEditorPrivatePCH.h"
+
+//////////////////////////////////////////////////////////////////////////
+// UPaperFlipbookThumbnailRenderer
+
+UPaperFlipbookThumbnailRenderer::UPaperFlipbookThumbnailRenderer(const FObjectInitializer& ObjectInitializer)
+	: Super(ObjectInitializer)
+{
+}
+
+void UPaperFlipbookThumbnailRenderer::Draw(UObject* Object, int32 X, int32 Y, uint32 Width, uint32 Height, FRenderTarget* RenderTarget, FCanvas* Canvas)
+{
+	if (UPaperFlipbook* Flipbook = Cast<UPaperFlipbook>(Object))
+	{
+		const double DeltaTime = FApp::GetCurrentTime() - GStartTime;
+		const float TotalDuration = Flipbook->GetTotalDuration();
+		const float PlayTime = (TotalDuration > 0.0f) ? FMath::Fmod(DeltaTime, TotalDuration) : 0.0f;
+
+		if (UPaperSprite* Sprite = Flipbook->GetSpriteAtTime(PlayTime))
+		{
+			FBoxSphereBounds FlipbookBounds = Flipbook->GetRenderBounds();
+			DrawFrame(Sprite, X, Y, Width, Height, RenderTarget, Canvas, &FlipbookBounds);
+			return;
+		}
+		else
+		{
+			// Fallback for empty frames or newly created flipbooks
+			DrawGrid(X, Y, Width, Height, Canvas);
+		}
+
+		if (TotalDuration == 0.0f)
+		{
+			// Warning text for no frames
+			const FText ErrorText = NSLOCTEXT("FlipbookEditorApp", "ThumbnailWarningNoFrames", "No frames");
+			FCanvasTextItem TextItem(FVector2D(5.0f, 5.0f), ErrorText, GEngine->GetLargeFont(), FLinearColor::Red);
+			TextItem.EnableShadow(FLinearColor::Black);
+			TextItem.Scale = FVector2D(Width / 128.0f, Height / 128.0f);
+			TextItem.Draw(Canvas);
+		}
+	}
+}
\ No newline at end of file
diff --git a/Engine/Plugins/2D/Paper2D/Source/Paper2DEditor/Private/PaperSpriteThumbnailRenderer.cpp b/Engine/Plugins/2D/Paper2D/Source/Paper2DEditor/Private/PaperSpriteThumbnailRenderer.cpp
--- a/Engine/Plugins/2D/Paper2D/Source/Paper2DEditor/Private/PaperSpriteThumbnailRenderer.cpp
+++ b/Engine/Plugins/2D/Paper2D/Source/Paper2DEditor/Private/PaperSpriteThumbnailRenderer.cpp
@@ -13,7 +13,7 @@ UPaperSpriteThumbnailRenderer::UPaperSpriteThumbnailRenderer(const FObjectInitia
 void UPaperSpriteThumbnailRenderer::Draw(UObject* Object, int32 X, int32 Y, uint32 Width, uint32 Height, FRenderTarget* RenderTarget, FCanvas* Canvas)
 {
 	UPaperSprite* Sprite = Cast<UPaperSprite>(Object);
-	DrawFrame(Sprite, X, Y, Width, Height, RenderTarget, Canvas);
+	DrawFrame(Sprite, X, Y, Width, Height, RenderTarget, Canvas, nullptr);
 }
 
 void UPaperSpriteThumbnailRenderer::DrawGrid(int32 X, int32 Y, uint32 Width, uint32 Height, FCanvas* Canvas)
@@ -40,7 +40,7 @@ void UPaperSpriteThumbnailRenderer::DrawGrid(int32 X, int32 Y, uint32 Width, uin
 		bAlphaBlend);
 }
 
-void UPaperSpriteThumbnailRenderer::DrawFrame(class UPaperSprite* Sprite, int32 X, int32 Y, uint32 Width, uint32 Height, FRenderTarget*, FCanvas* Canvas)
+void UPaperSpriteThumbnailRenderer::DrawFrame(class UPaperSprite* Sprite, int32 X, int32 Y, uint32 Width, uint32 Height, FRenderTarget*, FCanvas* Canvas, FBoxSphereBounds* OverrideRenderBounds)
 {
 	if (const UTexture2D* SourceTexture = (Sprite != NULL) ? Sprite->GetSourceTexture() : NULL)
 	{
@@ -65,21 +65,25 @@ void UPaperSpriteThumbnailRenderer::DrawFrame(class UPaperSprite* Sprite, int32
 			new(CanvasUVs)FVector2D(BakedRenderData[Vertex].Z, BakedRenderData[Vertex].W);
 		}
 
-		FVector2D MinPoint(BIG_NUMBER, BIG_NUMBER);
-		FVector2D MaxPoint(-BIG_NUMBER, -BIG_NUMBER);
-		for (int Vertex = 0; Vertex < CanvasPositions.Num(); ++Vertex)
+		// Determine the bounds to use
+		FBoxSphereBounds* RenderBounds = OverrideRenderBounds;
+		FBoxSphereBounds FrameBounds;
+		if (RenderBounds == nullptr)
 		{
-			MinPoint.X = FMath::Min(MinPoint.X, CanvasPositions[Vertex].X);
-			MinPoint.Y = FMath::Min(MinPoint.Y, CanvasPositions[Vertex].Y);
-			MaxPoint.X = FMath::Max(MaxPoint.X, CanvasPositions[Vertex].X);
-			MaxPoint.Y = FMath::Max(MaxPoint.Y, CanvasPositions[Vertex].Y);
+			FrameBounds = Sprite->GetRenderBounds();
+			RenderBounds = &FrameBounds;
 		}
 
-		float ScaleFactor = 1;
+		const FVector MinPoint3D = RenderBounds->GetBoxExtrema(0);
+		const FVector MaxPoint3D = RenderBounds->GetBoxExtrema(1);
+		const FVector2D MinPoint(FVector::DotProduct(MinPoint3D, PaperAxisX), FVector::DotProduct(MinPoint3D, PaperAxisY));
+		const FVector2D MaxPoint(FVector::DotProduct(MaxPoint3D, PaperAxisX), FVector::DotProduct(MaxPoint3D, PaperAxisY));
+
+		float ScaleFactor = 1.0f;
 		float UnscaledWidth = MaxPoint.X - MinPoint.X;
 		float UnscaledHeight = MaxPoint.Y - MinPoint.Y;
-		FVector2D Origin(X + Width / 2, Y + Height / 2);
-		if (UnscaledWidth > 0 && UnscaledHeight > 0 && UnscaledWidth > UnscaledHeight)
+		FVector2D Origin(X + Width / 2.0f, Y + Height / 2.0f);
+		if ((UnscaledWidth > 0.0f) && (UnscaledHeight > 0.0f) && (UnscaledWidth > UnscaledHeight))
 		{ 
 			ScaleFactor = Width / UnscaledWidth;
 		}
@@ -102,7 +106,7 @@ void UPaperSpriteThumbnailRenderer::DrawFrame(class UPaperSprite* Sprite, int32
 			const FLinearColor SpriteColor(FLinearColor::White);
 			for (int Vertex = 0; Vertex < CanvasPositions.Num(); Vertex += 3)
 			{
-				FCanvasUVTri *Triangle = new(Triangles)FCanvasUVTri();
+				FCanvasUVTri* Triangle = new (Triangles) FCanvasUVTri();
 				Triangle->V0_Pos = CanvasPositions[Vertex + 0]; Triangle->V0_UV = CanvasUVs[Vertex + 0]; Triangle->V0_Color = SpriteColor;
 				Triangle->V1_Pos = CanvasPositions[Vertex + 1]; Triangle->V1_UV = CanvasUVs[Vertex + 1]; Triangle->V1_Color = SpriteColor;
 				Triangle->V2_Pos = CanvasPositions[Vertex + 2]; Triangle->V2_UV = CanvasUVs[Vertex + 2]; Triangle->V2_Color = SpriteColor;
