Paper2D: Enable FPaperTileMapDetailsCustomization to work with raw assets as well as actors containing the component

diff --git a/Engine/Plugins/2D/Paper2D/Source/Paper2DEditor/Private/Paper2DEditorModule.cpp b/Engine/Plugins/2D/Paper2D/Source/Paper2DEditor/Private/Paper2DEditorModule.cpp
--- a/Engine/Plugins/2D/Paper2D/Source/Paper2DEditor/Private/Paper2DEditorModule.cpp
+++ b/Engine/Plugins/2D/Paper2D/Source/Paper2DEditor/Private/Paper2DEditorModule.cpp
@@ -111,8 +111,9 @@ public:
 		// Register the details customizations
 		{
 			FPropertyEditorModule& PropertyModule = FModuleManager::LoadModuleChecked<FPropertyEditorModule>("PropertyEditor");
-			PropertyModule.RegisterCustomClassLayout("PaperTileMapRenderComponent", FOnGetDetailCustomizationInstance::CreateStatic(&FPaperTileMapDetailsCustomization::MakeInstance));
-			PropertyModule.RegisterCustomClassLayout("PaperSprite", FOnGetDetailCustomizationInstance::CreateStatic(&FSpriteDetailsCustomization::MakeInstance));
+			PropertyModule.RegisterCustomClassLayout(UPaperTileMapRenderComponent::StaticClass()->GetFName(), FOnGetDetailCustomizationInstance::CreateStatic(&FPaperTileMapDetailsCustomization::MakeInstance));
+			PropertyModule.RegisterCustomClassLayout(UPaperTileMap::StaticClass()->GetFName(), FOnGetDetailCustomizationInstance::CreateStatic(&FPaperTileMapDetailsCustomization::MakeInstance));
+			PropertyModule.RegisterCustomClassLayout(UPaperSprite::StaticClass()->GetFName(), FOnGetDetailCustomizationInstance::CreateStatic(&FSpriteDetailsCustomization::MakeInstance));
 
 			//@TODO: Struct registration should happen using ::StaticStruct, not by string!!!
 			//PropertyModule.RegisterCustomPropertyTypeLayout( "SpritePolygonCollection", FOnGetPropertyTypeCustomizationInstance::CreateStatic( &FSpritePolygonCollectionCustomization::MakeInstance ) );
diff --git a/Engine/Plugins/2D/Paper2D/Source/Paper2DEditor/Private/TileMapEditing/PaperTileMapDetailsCustomization.cpp b/Engine/Plugins/2D/Paper2D/Source/Paper2DEditor/Private/TileMapEditing/PaperTileMapDetailsCustomization.cpp
--- a/Engine/Plugins/2D/Paper2D/Source/Paper2DEditor/Private/TileMapEditing/PaperTileMapDetailsCustomization.cpp
+++ b/Engine/Plugins/2D/Paper2D/Source/Paper2DEditor/Private/TileMapEditing/PaperTileMapDetailsCustomization.cpp
@@ -27,7 +27,8 @@ void FPaperTileMapDetailsCustomization::CustomizeDetails(IDetailLayoutBuilder& D
 	UPaperTileMap* TileMap = NULL;
 	for (int32 ObjectIndex = 0; ObjectIndex < SelectedObjects.Num(); ++ObjectIndex)
 	{
-		if (AActor* CurrentActor = Cast<AActor>(SelectedObjects[ObjectIndex].Get()))
+		UObject* TestObject = SelectedObjects[ObjectIndex].Get();
+		if (AActor* CurrentActor = Cast<AActor>(TestObject))
 		{
 			if (UPaperTileMapRenderComponent* CurrentTileMap = CurrentActor->FindComponentByClass<UPaperTileMapRenderComponent>())
 			{
@@ -35,6 +36,11 @@ void FPaperTileMapDetailsCustomization::CustomizeDetails(IDetailLayoutBuilder& D
 				break;
 			}
 		}
+		else if (UPaperTileMap* TestTileMap = Cast<UPaperTileMap>(TestObject))
+		{
+			TileMap = TestTileMap;
+			break;
+		}
 	}
 	TileMapPtr = TileMap;
 
