Paper2D: Experimenting with hosting editor modes in an asset editor (WIP)

diff --git a/Engine/Plugins/2D/Paper2D/Source/Paper2DEditor/Private/TileMapEditing/TileMapEditor.cpp b/Engine/Plugins/2D/Paper2D/Source/Paper2DEditor/Private/TileMapEditing/TileMapEditor.cpp
--- a/Engine/Plugins/2D/Paper2D/Source/Paper2DEditor/Private/TileMapEditing/TileMapEditor.cpp
+++ b/Engine/Plugins/2D/Paper2D/Source/Paper2DEditor/Private/TileMapEditing/TileMapEditor.cpp
@@ -28,12 +28,14 @@ struct FTileMapEditorTabs
 	// Tab identifiers
 	static const FName DetailsID;
 	static const FName ViewportID;
+	static const FName ToolboxHostID;
 };
 
 //////////////////////////////////////////////////////////////////////////
 
 const FName FTileMapEditorTabs::DetailsID(TEXT("Details"));
 const FName FTileMapEditorTabs::ViewportID(TEXT("Viewport"));
+const FName FTileMapEditorTabs::ToolboxHostID(TEXT("Toolbox"));
 
 //////////////////////////////////////////////////////////////////////////
 // STileMapEditorViewport
@@ -217,6 +219,20 @@ TSharedRef<SDockTab> FTileMapEditor::SpawnTab_Viewport(const FSpawnTabArgs& Args
 		];
 }
 
+TSharedRef<SDockTab> FTileMapEditor::SpawnTab_ToolboxHost(const FSpawnTabArgs& Args)
+{
+	TSharedPtr<FTileMapEditor> TileMapEditorPtr = SharedThis(this);
+
+	// Spawn the tab
+	return SNew(SDockTab)
+		.Icon(FEditorStyle::GetBrush("LevelEditor.Tabs.Modes"))
+		.Label(LOCTEXT("ToolboxHost_Title", "Toolbox"))
+		[
+			SNew(STextBlock)
+			.Text(LOCTEXT("ToolkitControlsGohere", "Toolkit controls go here"))
+		];
+}
+
 TSharedRef<SDockTab> FTileMapEditor::SpawnTab_Details(const FSpawnTabArgs& Args)
 {
 	TSharedPtr<FTileMapEditor> TileMapEditorPtr = SharedThis(this);
@@ -242,6 +258,11 @@ void FTileMapEditor::RegisterTabSpawners(const TSharedRef<class FTabManager>& Ta
 		.SetGroup(WorkspaceMenuCategoryRef)
 		.SetIcon(FSlateIcon(FEditorStyle::GetStyleSetName(), "LevelEditor.Tabs.Viewports"));
 
+	TabManager->RegisterTabSpawner(FTileMapEditorTabs::ToolboxHostID, FOnSpawnTab::CreateSP(this, &FTileMapEditor::SpawnTab_ToolboxHost))
+		.SetDisplayName(LOCTEXT("ToolboxHostLabel", "Toolbox"))
+		.SetGroup(WorkspaceMenuCategoryRef)
+		.SetIcon(FSlateIcon(FEditorStyle::GetStyleSetName(), "LevelEditor.Tabs.Modes"));
+
 	TabManager->RegisterTabSpawner(FTileMapEditorTabs::DetailsID, FOnSpawnTab::CreateSP(this, &FTileMapEditor::SpawnTab_Details))
 		.SetDisplayName(LOCTEXT("DetailsTabLabel", "Details"))
 		.SetGroup(WorkspaceMenuCategoryRef)
@@ -253,6 +274,7 @@ void FTileMapEditor::UnregisterTabSpawners(const TSharedRef<class FTabManager>&
 	FAssetEditorToolkit::UnregisterTabSpawners(TabManager);
 
 	TabManager->UnregisterTabSpawner(FTileMapEditorTabs::ViewportID);
+	TabManager->UnregisterTabSpawner(FTileMapEditorTabs::ToolboxHostID);
 	TabManager->UnregisterTabSpawner(FTileMapEditorTabs::DetailsID);
 }
 
@@ -268,7 +290,7 @@ void FTileMapEditor::InitTileMapEditor(const EToolkitMode::Type Mode, const TSha
 	ViewportPtr = SNew(STileMapEditorViewport, SharedThis(this));
 
 	// Default layout
-	const TSharedRef<FTabManager::FLayout> StandaloneDefaultLayout = FTabManager::NewLayout("Standalone_TileMapEditor_Layout_v1")
+	const TSharedRef<FTabManager::FLayout> StandaloneDefaultLayout = FTabManager::NewLayout("Standalone_TileMapEditor_Layout_v2")
 		->AddArea
 		(
 			FTabManager::NewPrimaryArea()
@@ -286,6 +308,13 @@ void FTileMapEditor::InitTileMapEditor(const EToolkitMode::Type Mode, const TSha
 				->SetOrientation(Orient_Horizontal)
 				->SetSizeCoefficient(0.9f)
 				->Split
+				(
+					FTabManager::NewStack()
+					->SetSizeCoefficient(0.3f)
+					->SetHideTabWell(true)
+					->AddTab(FTileMapEditorTabs::ToolboxHostID, ETabState::OpenedTab)
+				)
+				->Split
 				(
 					FTabManager::NewStack()
 					->SetSizeCoefficient(0.8f)
@@ -307,13 +336,20 @@ void FTileMapEditor::InitTileMapEditor(const EToolkitMode::Type Mode, const TSha
 			)
 		);
 
-	// Initialize the asset editor and spawn nothing (dummy layout)
+	// Initialize the asset editor and spawn the layout above
 	InitAssetEditor(Mode, InitToolkitHost, TileMapEditorAppName, StandaloneDefaultLayout, /*bCreateDefaultStandaloneMenu=*/ true, /*bCreateDefaultToolbar=*/ true, InitTileMap);
 
 	// Extend things
 	ExtendMenu();
 	ExtendToolbar();
 	RegenerateMenusAndToolbars();
+
+	// Activate the tile map edit mode
+	EditorModeToolsInstance.SetDefaultMode(FEdModeTileMap::EM_TileMap);
+	EditorModeToolsInstance.ActivateDefaultMode();
+
+	//@TODO: Need to be able to pass ToolkitHost.Get() into EditorModeToolsInstance, and have it in turn pass it into Enter/Leave on the individual modes I think
+	//@TODO: Need to be able to register the widget in the toolbox panel with ToolkitHost, so it can instance the ed mode widgets into it
 }
 
 void FTileMapEditor::BindCommands()
diff --git a/Engine/Plugins/2D/Paper2D/Source/Paper2DEditor/Private/TileMapEditing/TileMapEditor.h b/Engine/Plugins/2D/Paper2D/Source/Paper2DEditor/Private/TileMapEditing/TileMapEditor.h
--- a/Engine/Plugins/2D/Paper2D/Source/Paper2DEditor/Private/TileMapEditing/TileMapEditor.h
+++ b/Engine/Plugins/2D/Paper2D/Source/Paper2DEditor/Private/TileMapEditing/TileMapEditor.h
@@ -36,6 +36,7 @@ public:
 protected:
 	UPaperTileMap* TileMapBeingEdited;
 	TSharedPtr<class STileMapEditorViewport> ViewportPtr;
+	FEditorModeTools EditorModeToolsInstance;
 
 protected:
 	void BindCommands();
@@ -43,6 +44,7 @@ protected:
 	void ExtendToolbar();
 
 	TSharedRef<SDockTab> SpawnTab_Viewport(const FSpawnTabArgs& Args);
+	TSharedRef<SDockTab> SpawnTab_ToolboxHost(const FSpawnTabArgs& Args);
 	TSharedRef<SDockTab> SpawnTab_Details(const FSpawnTabArgs& Args);
 
 	void CreateModeToolbarWidgets(FToolBarBuilder& ToolbarBuilder);
