Paper2D: Prevent showing the 'edit tile map button' in the tile map asset editor (the mode will always be active)

diff --git a/Engine/Plugins/2D/Paper2D/Source/Paper2DEditor/Private/TileMapEditing/PaperTileMapDetailsCustomization.cpp b/Engine/Plugins/2D/Paper2D/Source/Paper2DEditor/Private/TileMapEditing/PaperTileMapDetailsCustomization.cpp
--- a/Engine/Plugins/2D/Paper2D/Source/Paper2DEditor/Private/TileMapEditing/PaperTileMapDetailsCustomization.cpp
+++ b/Engine/Plugins/2D/Paper2D/Source/Paper2DEditor/Private/TileMapEditing/PaperTileMapDetailsCustomization.cpp
@@ -24,6 +24,8 @@ void FPaperTileMapDetailsCustomization::CustomizeDetails(IDetailLayoutBuilder& D
 {
 	const TArray< TWeakObjectPtr<UObject> >& SelectedObjects = DetailLayout.GetDetailsView().GetSelectedObjects();
 	
+	bool bEditingActor = false;
+
 	UPaperTileMap* TileMap = NULL;
 	for (int32 ObjectIndex = 0; ObjectIndex < SelectedObjects.Num(); ++ObjectIndex)
 	{
@@ -32,6 +34,7 @@ void FPaperTileMapDetailsCustomization::CustomizeDetails(IDetailLayoutBuilder& D
 		{
 			if (UPaperTileMapRenderComponent* CurrentTileMap = CurrentActor->FindComponentByClass<UPaperTileMapRenderComponent>())
 			{
+				bEditingActor = true;
 				TileMap = CurrentTileMap->TileMap;
 				break;
 			}
@@ -46,17 +49,14 @@ void FPaperTileMapDetailsCustomization::CustomizeDetails(IDetailLayoutBuilder& D
 
 	IDetailCategoryBuilder& TileMapCategory = DetailLayout.EditCategory("TileMap");
 
-	TileMapCategory
-	.AddCustomRow(TEXT("EnterEditMode"))
-	[
-		SNew(SVerticalBox)
-		+SVerticalBox::Slot()
+	if (bEditingActor)
+	{
+		TileMapCategory
+		.AddCustomRow(TEXT("EnterEditMode"))
 		[
-// 			SNew(SHorizontalBox)
-// 			+SHorizontalBox::Slot()
-// 			.AutoWidth()
-// 			.Padding(10,5)
-// 			[
+			SNew(SVerticalBox)
+			+SVerticalBox::Slot()
+			[
 				SNew(SButton)
 				.ContentPadding(3)
 				.VAlign(VAlign_Center)
@@ -64,9 +64,9 @@ void FPaperTileMapDetailsCustomization::CustomizeDetails(IDetailLayoutBuilder& D
 				.OnClicked(this, &FPaperTileMapDetailsCustomization::EnterTileMapEditingMode)
 				.Visibility(this, &FPaperTileMapDetailsCustomization::GetNonEditModeVisibility)
 				.Text( LOCTEXT( "EnterTileMapEditMode", "Enter Edit Mode" ) )
-//			]
-		]
-	];
+			]
+		];
+	}
 
 	//@TODO: Handle showing layers when multiple tile maps are selected
 	if (TileMap != NULL)
