Paper2D: Tile map asset editor using an ed mode WIP

diff --git a/Engine/Plugins/2D/Paper2D/Source/Paper2DEditor/Private/PaperEditorViewportClient.cpp b/Engine/Plugins/2D/Paper2D/Source/Paper2DEditor/Private/PaperEditorViewportClient.cpp
--- a/Engine/Plugins/2D/Paper2D/Source/Paper2DEditor/Private/PaperEditorViewportClient.cpp
+++ b/Engine/Plugins/2D/Paper2D/Source/Paper2DEditor/Private/PaperEditorViewportClient.cpp
@@ -9,14 +9,14 @@
 // FPaperEditorViewportClient
 
 FPaperEditorViewportClient::FPaperEditorViewportClient()
-	: FEditorViewportClient(GLevelEditorModeTools())
+	: FEditorViewportClient(*( new FEditorModeTools() ))
 	, CheckerboardTexture(NULL)
 {
 	ZoomPos = FVector2D::ZeroVector;
 	ZoomAmount = 1.0f;
 
 	//ModifyCheckerboardTextureColors();
-
+	//@TODO: ModeTools->SetToolkitHost
 
 	//@TODO: Pretty lame hardcoding
 	//@TODO: Doesn't handle negatives either (not really)
@@ -49,6 +49,7 @@ FPaperEditorViewportClient::FPaperEditorViewportClient()
 FPaperEditorViewportClient::~FPaperEditorViewportClient()
 {
 	//DestroyCheckerboardTexture();
+	delete ModeTools;
 }
 
 void FPaperEditorViewportClient::Draw(FViewport* Viewport, FCanvas* Canvas)
diff --git a/Engine/Plugins/2D/Paper2D/Source/Paper2DEditor/Private/TileMapEditing/TileMapEdModeToolkit.cpp b/Engine/Plugins/2D/Paper2D/Source/Paper2DEditor/Private/TileMapEditing/TileMapEdModeToolkit.cpp
--- a/Engine/Plugins/2D/Paper2D/Source/Paper2DEditor/Private/TileMapEditing/TileMapEdModeToolkit.cpp
+++ b/Engine/Plugins/2D/Paper2D/Source/Paper2DEditor/Private/TileMapEditing/TileMapEdModeToolkit.cpp
@@ -13,6 +13,11 @@
 //////////////////////////////////////////////////////////////////////////
 // FTileMapEdModeToolkit
 
+FTileMapEdModeToolkit::FTileMapEdModeToolkit(class FEdModeTileMap* InOwningMode)
+{
+	TileMapEditor = InOwningMode;
+}
+
 void FTileMapEdModeToolkit::RegisterTabSpawners(const TSharedRef<class FTabManager>& TabManager)
 {
 }
@@ -47,7 +52,7 @@ FText FTileMapEdModeToolkit::GetToolkitName() const
 
 FEdMode* FTileMapEdModeToolkit::GetEditorMode() const
 {
-	return GLevelEditorModeTools().GetActiveMode(FEdModeTileMap::EM_TileMap);
+	return TileMapEditor;
 }
 
 TSharedPtr<SWidget> FTileMapEdModeToolkit::GetInlineContent() const
@@ -196,51 +201,27 @@ void FTileMapEdModeToolkit::BindCommands()
 
 void FTileMapEdModeToolkit::OnSelectTool(ETileMapEditorTool::Type NewTool)
 {
-	if (FEdModeTileMap* TileMapEditor = GLevelEditorModeTools().GetActiveModeTyped<FEdModeTileMap>(FEdModeTileMap::EM_TileMap))
-	{
-		TileMapEditor->SetActiveTool(NewTool);
-	}
+	TileMapEditor->SetActiveTool(NewTool);
 }
 
 bool FTileMapEdModeToolkit::IsToolSelected(ETileMapEditorTool::Type QueryTool) const
 {
-	if (FEdModeTileMap* TileMapEditor = GLevelEditorModeTools().GetActiveModeTyped<FEdModeTileMap>(FEdModeTileMap::EM_TileMap))
-	{
-		return (TileMapEditor->GetActiveTool() == QueryTool);
-	}
-	else
-	{
-		return false;
-	}
+	return (TileMapEditor->GetActiveTool() == QueryTool);
 }
 
 void FTileMapEdModeToolkit::OnSelectLayerPaintingMode(ETileMapLayerPaintingMode::Type NewMode)
 {
-	if (FEdModeTileMap* TileMapEditor = GLevelEditorModeTools().GetActiveModeTyped<FEdModeTileMap>(FEdModeTileMap::EM_TileMap))
-	{
-		TileMapEditor->SetActiveLayerPaintingMode(NewMode);
-	}
+	TileMapEditor->SetActiveLayerPaintingMode(NewMode);
 }
 
 bool FTileMapEdModeToolkit::IsLayerPaintingModeSelected(ETileMapLayerPaintingMode::Type PaintingMode) const
 {
-	if (FEdModeTileMap* TileMapEditor = GLevelEditorModeTools().GetActiveModeTyped<FEdModeTileMap>(FEdModeTileMap::EM_TileMap))
-	{
-		return (TileMapEditor->GetActiveLayerPaintingMode() == PaintingMode);
-	}
-	else
-	{
-		return false;
-	}
+	return (TileMapEditor->GetActiveLayerPaintingMode() == PaintingMode);
 }
 
 EVisibility FTileMapEdModeToolkit::GetTileSetSelectorVisibility() const
 {
-	bool bShouldShowSelector = false;
-	if (FEdModeTileMap* TileMapEditor = GLevelEditorModeTools().GetActiveModeTyped<FEdModeTileMap>(FEdModeTileMap::EM_TileMap))
-	{
-		bShouldShowSelector = (TileMapEditor->GetActiveLayerPaintingMode() == ETileMapLayerPaintingMode::VisualLayers);
-	}
+	bool bShouldShowSelector = (TileMapEditor->GetActiveLayerPaintingMode() == ETileMapLayerPaintingMode::VisualLayers);
 	
 	return bShouldShowSelector ? EVisibility::Visible : EVisibility::Collapsed;
 }
diff --git a/Engine/Plugins/2D/Paper2D/Source/Paper2DEditor/Private/TileMapEditing/TileMapEdModeToolkit.h b/Engine/Plugins/2D/Paper2D/Source/Paper2DEditor/Private/TileMapEditing/TileMapEdModeToolkit.h
--- a/Engine/Plugins/2D/Paper2D/Source/Paper2DEditor/Private/TileMapEditing/TileMapEdModeToolkit.h
+++ b/Engine/Plugins/2D/Paper2D/Source/Paper2DEditor/Private/TileMapEditing/TileMapEdModeToolkit.h
@@ -10,6 +10,8 @@
 class FTileMapEdModeToolkit : public FModeToolkit
 {
 public:
+	FTileMapEdModeToolkit(class FEdModeTileMap* InOwningMode);
+
 	// IToolkit interface
 	virtual void RegisterTabSpawners(const TSharedRef<class FTabManager>& TabManager) override;
 	virtual void UnregisterTabSpawners(const TSharedRef<class FTabManager>& TabManager) override;
@@ -40,6 +42,8 @@ protected:
 
 	TSharedRef<SWidget> BuildToolBar() const;
 private:
+	class FEdModeTileMap* TileMapEditor;
+
 	TWeakObjectPtr<class UPaperTileSet> CurrentTileSetPtr;
 
 	// All of the inline content for this toolkit
diff --git a/Engine/Plugins/2D/Paper2D/Source/Paper2DEditor/Private/TileMapEditing/TileMapEditor.cpp b/Engine/Plugins/2D/Paper2D/Source/Paper2DEditor/Private/TileMapEditing/TileMapEditor.cpp
--- a/Engine/Plugins/2D/Paper2D/Source/Paper2DEditor/Private/TileMapEditing/TileMapEditor.cpp
+++ b/Engine/Plugins/2D/Paper2D/Source/Paper2DEditor/Private/TileMapEditing/TileMapEditor.cpp
@@ -67,6 +67,11 @@ public:
 	{
 		EditorViewportClient->NotifyTileMapBeingEditedHasChanged();
 	}
+
+	void ActivateEditMode()
+	{
+		EditorViewportClient->ActivateEditMode();
+	}
 private:
 	// Pointer back to owning tile map editor instance (the keeper of state)
 	TWeakPtr<class FTileMapEditor> TileMapEditorPtr;
@@ -340,17 +345,13 @@ void FTileMapEditor::InitTileMapEditor(const EToolkitMode::Type Mode, const TSha
 	// Initialize the asset editor and spawn the layout above
 	InitAssetEditor(Mode, InitToolkitHost, TileMapEditorAppName, StandaloneDefaultLayout, /*bCreateDefaultStandaloneMenu=*/ true, /*bCreateDefaultToolbar=*/ true, InitTileMap);
 
+	// Activate the edit mode
+	ViewportPtr->ActivateEditMode();
+
 	// Extend things
 	ExtendMenu();
 	ExtendToolbar();
 	RegenerateMenusAndToolbars();
-
-	// Activate the tile map edit mode
-	EditorModeToolsInstance.SetDefaultMode(FEdModeTileMap::EM_TileMap);
-	EditorModeToolsInstance.ActivateDefaultMode();
-
-	//@TODO: Need to be able to pass ToolkitHost.Get() into EditorModeToolsInstance, and have it in turn pass it into Enter/Leave on the individual modes I think
-	//@TODO: Need to be able to register the widget in the toolbox panel with ToolkitHost, so it can instance the ed mode widgets into it
 }
 
 void FTileMapEditor::BindCommands()
diff --git a/Engine/Plugins/2D/Paper2D/Source/Paper2DEditor/Private/TileMapEditing/TileMapEditor.h b/Engine/Plugins/2D/Paper2D/Source/Paper2DEditor/Private/TileMapEditing/TileMapEditor.h
--- a/Engine/Plugins/2D/Paper2D/Source/Paper2DEditor/Private/TileMapEditing/TileMapEditor.h
+++ b/Engine/Plugins/2D/Paper2D/Source/Paper2DEditor/Private/TileMapEditing/TileMapEditor.h
@@ -36,7 +36,6 @@ public:
 protected:
 	UPaperTileMap* TileMapBeingEdited;
 	TSharedPtr<class STileMapEditorViewport> ViewportPtr;
-	FEditorModeTools EditorModeToolsInstance;
 
 protected:
 	void BindCommands();
diff --git a/Engine/Plugins/2D/Paper2D/Source/Paper2DEditor/Private/TileMapEditing/TileMapEditorViewportClient.cpp b/Engine/Plugins/2D/Paper2D/Source/Paper2DEditor/Private/TileMapEditing/TileMapEditorViewportClient.cpp
--- a/Engine/Plugins/2D/Paper2D/Source/Paper2DEditor/Private/TileMapEditing/TileMapEditorViewportClient.cpp
+++ b/Engine/Plugins/2D/Paper2D/Source/Paper2DEditor/Private/TileMapEditing/TileMapEditorViewportClient.cpp
@@ -3,6 +3,7 @@
 #include "Paper2DEditorPrivatePCH.h"
 #include "TileMapEditorViewportClient.h"
 #include "SceneViewport.h"
+#include "EdModeTileMap.h"
 
 #include "PreviewScene.h"
 #include "ScopedTransaction.h"
@@ -15,8 +16,8 @@
 // FTileMapEditorViewportClient
 
 FTileMapEditorViewportClient::FTileMapEditorViewportClient(TWeakPtr<FTileMapEditor> InTileMapEditor, TWeakPtr<class STileMapEditorViewport> InTileMapEditorViewportPtr)
-	: TileMapEditorPtr(InTileMapEditor)
-	, TileMapEditorViewportPtr(InTileMapEditorViewportPtr)
+: TileMapEditorPtr(InTileMapEditor)
+, TileMapEditorViewportPtr(InTileMapEditorViewportPtr)
 {
 	check(TileMapEditorPtr.IsValid() && TileMapEditorViewportPtr.IsValid());
 
@@ -46,6 +47,16 @@ FTileMapEditorViewportClient::FTileMapEditorViewportClient(TWeakPtr<FTileMapEdit
 	}
 }
 
+void FTileMapEditorViewportClient::ActivateEditMode()
+{
+	// Activate the tile map edit mode
+	ModeTools->SetToolkitHost(TileMapEditorPtr.Pin()->GetToolkitHost());
+	ModeTools->SetDefaultMode(FEdModeTileMap::EM_TileMap);
+	ModeTools->ActivateDefaultMode();
+	
+	//@TODO: Need to be able to register the widget in the toolbox panel with ToolkitHost, so it can instance the ed mode widgets into it
+}
+
 void FTileMapEditorViewportClient::DrawBoundsAsText(FViewport& InViewport, FSceneView& View, FCanvas& Canvas, int32& YPos)
 {
 	FNumberFormattingOptions NoDigitGroupingFormat;
diff --git a/Engine/Plugins/2D/Paper2D/Source/Paper2DEditor/Private/TileMapEditing/TileMapEditorViewportClient.h b/Engine/Plugins/2D/Paper2D/Source/Paper2DEditor/Private/TileMapEditing/TileMapEditorViewportClient.h
--- a/Engine/Plugins/2D/Paper2D/Source/Paper2DEditor/Private/TileMapEditing/TileMapEditorViewportClient.h
+++ b/Engine/Plugins/2D/Paper2D/Source/Paper2DEditor/Private/TileMapEditing/TileMapEditorViewportClient.h
@@ -50,6 +50,9 @@ public:
 
 	// Invalidate any references to the tile map being edited; it has changed
 	void NotifyTileMapBeingEditedHasChanged();
+
+	// Note: Has to be delayed due to an unfortunate init ordering
+	void ActivateEditMode();
 private:
 	// The preview scene
 	FPreviewScene OwnedPreviewScene;
