Paper2D: Removed some dead code from the tile map ed mode

diff --git a/Engine/Plugins/2D/Paper2D/Source/Paper2DEditor/Private/TileMapEditing/EdModeTileMap.cpp b/Engine/Plugins/2D/Paper2D/Source/Paper2DEditor/Private/TileMapEditing/EdModeTileMap.cpp
--- a/Engine/Plugins/2D/Paper2D/Source/Paper2DEditor/Private/TileMapEditing/EdModeTileMap.cpp
+++ b/Engine/Plugins/2D/Paper2D/Source/Paper2DEditor/Private/TileMapEditing/EdModeTileMap.cpp
@@ -358,10 +358,7 @@ UPaperTileLayer* FEdModeTileMap::GetSelectedLayerUnderCursor(const FViewportCurs
 					OutTileY = FMath::FloorToInt(NormalizedY * TileMap->MapHeight);
 					
 					if ((OutTileX > -BrushWidth) && (OutTileX < TileMap->MapWidth) && (OutTileY > -BrushHeight) && (OutTileY < TileMap->MapHeight))
-						//(NormalizedX >= 0.0f) && (NormalizedX < 1.0f) && (NormalizedY >= 0.0f) && (NormalizedY < 1.0f))
 					{
-						//OutTileX = FMath::Clamp<int>(NormalizedX * TileMap->MapWidth, 0, TileMap->MapWidth-1);
-						//OutTileY = FMath::Clamp<int>(NormalizedY * TileMap->MapHeight, 0, TileMap->MapHeight-1);
 						return Layer;
 					}
 				}
@@ -744,27 +741,6 @@ FViewportCursorLocation FEdModeTileMap::CalculateViewRay(FEditorViewportClient*
 	return MouseViewportRay;
 }
 
-AActor* FEdModeTileMap::GetFirstSelectedActorContainingTileMapComponent()
-{
-	for (FSelectionIterator It(GEditor->GetSelectedActorIterator()); It; ++It)
-	{
-		if (AActor* Actor = Cast<AActor>(*It))
-		{
-			if (Actor->FindComponentByClass<UPaperTileMapRenderComponent>() != NULL)
-			{
-				return Actor;
-			}
-		}
-	}
-
-	return NULL;
-}
-
-void FEdModeTileMap::CreateModeButtonInModeTray(FToolBarBuilder& Builder)
-{
-	Builder.AddToolBarButton( FPaperEditorCommands::Get().EnterTileMapEditMode );
-}
-
 void FEdModeTileMap::SetActiveTool(ETileMapEditorTool::Type NewTool)
 {
 	ActiveTool = NewTool;
diff --git a/Engine/Plugins/2D/Paper2D/Source/Paper2DEditor/Private/TileMapEditing/EdModeTileMap.h b/Engine/Plugins/2D/Paper2D/Source/Paper2DEditor/Private/TileMapEditing/EdModeTileMap.h
--- a/Engine/Plugins/2D/Paper2D/Source/Paper2DEditor/Private/TileMapEditing/EdModeTileMap.h
+++ b/Engine/Plugins/2D/Paper2D/Source/Paper2DEditor/Private/TileMapEditing/EdModeTileMap.h
@@ -90,10 +90,6 @@ protected:
 
 	// Compute a world space ray from the screen space mouse coordinates
 	FViewportCursorLocation CalculateViewRay(FEditorViewportClient* InViewportClient, FViewport* InViewport);
-
-	static AActor* GetFirstSelectedActorContainingTileMapComponent();
-
-	void CreateModeButtonInModeTray(FToolBarBuilder& Builder);
 	
 	TSharedRef<FExtender> AddCreationModeExtender(const TSharedRef<FUICommandList> InCommandList);
 
