Paper2D: Wired up toolkit hosting panel in tile map editor

diff --git a/Engine/Plugins/2D/Paper2D/Source/Paper2DEditor/Private/TileMapEditing/TileMapEdModeToolkit.cpp b/Engine/Plugins/2D/Paper2D/Source/Paper2DEditor/Private/TileMapEditing/TileMapEdModeToolkit.cpp
--- a/Engine/Plugins/2D/Paper2D/Source/Paper2DEditor/Private/TileMapEditing/TileMapEdModeToolkit.cpp
+++ b/Engine/Plugins/2D/Paper2D/Source/Paper2DEditor/Private/TileMapEditing/TileMapEdModeToolkit.cpp
@@ -81,6 +81,8 @@ void FTileMapEdModeToolkit::Init(const TSharedPtr<IToolkitHost>& InitToolkitHost
 		}
 	}
 
+	TileSetPalette = SNew(STileSetSelectorViewport, CurrentTileSetPtr.Get(), TileMapEditor);
+
 	// Create the contents of the editor mode toolkit
 	MyWidget = 
 		SNew(SBorder)
@@ -139,7 +141,7 @@ void FTileMapEdModeToolkit::Init(const TSharedPtr<IToolkitHost>& InitToolkitHost
 					+SHorizontalBox::Slot()
 					.HAlign(HAlign_Fill)
 					[
-						SAssignNew(TileSetPalette, STileSetSelectorViewport, CurrentTileSetPtr.Get())
+						TileSetPalette.ToSharedRef()
 					]
 				]
 			]
diff --git a/Engine/Plugins/2D/Paper2D/Source/Paper2DEditor/Private/TileMapEditing/TileMapEditor.cpp b/Engine/Plugins/2D/Paper2D/Source/Paper2DEditor/Private/TileMapEditing/TileMapEditor.cpp
--- a/Engine/Plugins/2D/Paper2D/Source/Paper2DEditor/Private/TileMapEditing/TileMapEditor.cpp
+++ b/Engine/Plugins/2D/Paper2D/Source/Paper2DEditor/Private/TileMapEditing/TileMapEditor.cpp
@@ -234,8 +234,7 @@ TSharedRef<SDockTab> FTileMapEditor::SpawnTab_ToolboxHost(const FSpawnTabArgs& A
 		.Icon(FEditorStyle::GetBrush("LevelEditor.Tabs.Modes"))
 		.Label(LOCTEXT("ToolboxHost_Title", "Toolbox"))
 		[
-			SNew(STextBlock)
-			.Text(LOCTEXT("ToolkitControlsGohere", "Toolkit controls go here"))
+			ToolboxPtr.ToSharedRef()
 		];
 }
 
@@ -294,6 +293,9 @@ void FTileMapEditor::InitTileMapEditor(const EToolkitMode::Type Mode, const TSha
 	BindCommands();
 
 	ViewportPtr = SNew(STileMapEditorViewport, SharedThis(this));
+	ToolboxPtr = SNew(SBorder)
+		.BorderImage(FEditorStyle::GetBrush("ToolPanel.GroupBorder"))
+		.Padding(0.f);
 
 	// Default layout
 	const TSharedRef<FTabManager::FLayout> StandaloneDefaultLayout = FTabManager::NewLayout("Standalone_TileMapEditor_Layout_v2")
@@ -406,6 +408,41 @@ FString FTileMapEditor::GetDocumentationLink() const
 	return TEXT("Engine/Paper2D/TileMapEditor");
 }
 
+void FTileMapEditor::OnToolkitHostingStarted(const TSharedRef< class IToolkit >& Toolkit)
+{
+	TSharedPtr<SWidget> InlineContent = Toolkit->GetInlineContent();
+	if (InlineContent.IsValid())
+	{
+		ToolboxPtr->SetContent(InlineContent.ToSharedRef());
+	}
+}
+
+void FTileMapEditor::OnToolkitHostingFinished(const TSharedRef< class IToolkit >& Toolkit)
+{
+	ToolboxPtr->SetContent(SNullWidget::NullWidget);
+
+	//@TODO: MODETOOLS: How to handle multiple ed modes at once in a standalone asset editor?
+#if 0
+	bool FoundAnotherToolkit = false;
+	const TArray< TSharedPtr< IToolkit > >& HostedToolkits = LevelEditor.Pin()->GetHostedToolkits();
+	for (auto HostedToolkitIt = HostedToolkits.CreateConstIterator(); HostedToolkitIt; ++HostedToolkitIt)
+	{
+		if ((*HostedToolkitIt) != Toolkit)
+		{
+			UpdateInlineContent((*HostedToolkitIt)->GetInlineContent());
+			FoundAnotherToolkit = true;
+			break;
+		}
+	}
+
+	if (!FoundAnotherToolkit)
+	{
+		UpdateInlineContent(SNullWidget::NullWidget);
+	}
+
+#endif
+}
+
 FLinearColor FTileMapEditor::GetWorldCentricTabColorScale() const
 {
 	return FLinearColor::White;
diff --git a/Engine/Plugins/2D/Paper2D/Source/Paper2DEditor/Private/TileMapEditing/TileMapEditor.h b/Engine/Plugins/2D/Paper2D/Source/Paper2DEditor/Private/TileMapEditing/TileMapEditor.h
--- a/Engine/Plugins/2D/Paper2D/Source/Paper2DEditor/Private/TileMapEditing/TileMapEditor.h
+++ b/Engine/Plugins/2D/Paper2D/Source/Paper2DEditor/Private/TileMapEditing/TileMapEditor.h
@@ -23,6 +23,8 @@ public:
 	virtual FLinearColor GetWorldCentricTabColorScale() const override;
 	virtual FString GetWorldCentricTabPrefix() const override;
 	virtual FString GetDocumentationLink() const override;
+	virtual void OnToolkitHostingStarted(const TSharedRef< class IToolkit >& Toolkit) override;
+	virtual void OnToolkitHostingFinished(const TSharedRef< class IToolkit >& Toolkit) override;
 	// End of FAssetEditorToolkit
 
 	// FSerializableObject interface
@@ -36,6 +38,7 @@ public:
 protected:
 	UPaperTileMap* TileMapBeingEdited;
 	TSharedPtr<class STileMapEditorViewport> ViewportPtr;
+	TSharedPtr<class SBorder> ToolboxPtr;
 
 protected:
 	void BindCommands();
diff --git a/Engine/Plugins/2D/Paper2D/Source/Paper2DEditor/Private/TileSetEditor.cpp b/Engine/Plugins/2D/Paper2D/Source/Paper2DEditor/Private/TileSetEditor.cpp
--- a/Engine/Plugins/2D/Paper2D/Source/Paper2DEditor/Private/TileSetEditor.cpp
+++ b/Engine/Plugins/2D/Paper2D/Source/Paper2DEditor/Private/TileSetEditor.cpp
@@ -133,9 +133,10 @@ STileSetSelectorViewport::~STileSetSelectorViewport()
 	TypedViewportClient = NULL;
 }
 
-void STileSetSelectorViewport::Construct(const FArguments& InArgs, UPaperTileSet* InTileSet)
+void STileSetSelectorViewport::Construct(const FArguments& InArgs, UPaperTileSet* InTileSet, FEdModeTileMap* InTileMapEditor)
 {
 	TileSetPtr = InTileSet;
+	TileMapEditor = InTileMapEditor;
 
 	TypedViewportClient = MakeShareable(new FTileSetEditorViewportClient(InTileSet));
 
@@ -203,17 +204,17 @@ void STileSetSelectorViewport::OnSelectionChanged(FMarqueeOperation Marquee, boo
 	const bool bHasSelection = (SelectionDimensions.X + SelectionDimensions.Y > 0);
 	if (bIsPreview && bHasSelection)
 	{
-		if (FEdModeTileMap* TileMapEditor = GLevelEditorModeTools().GetActiveModeTyped<FEdModeTileMap>(FEdModeTileMap::EM_TileMap))
+		if (TileMapEditor != nullptr)
 		{
 			TileMapEditor->SetActivePaint(TileSetBeingEdited, SelectionTopLeft, SelectionDimensions);
+		}
 
-			if (FTileSetEditorViewportClient* Client = TypedViewportClient.Get())
-			{
-				Client->bHasValidPaintRectangle = true;
-				Client->ValidPaintRectangle.Color = FLinearColor::White;
-				Client->ValidPaintRectangle.Dimensions = FVector2D(SelectionDimensions.X * TileSetBeingEdited->TileWidth, SelectionDimensions.Y * TileSetBeingEdited->TileHeight);
-				Client->ValidPaintRectangle.TopLeft = FVector2D(SelectionTopLeft.X * TileSetBeingEdited->TileWidth, SelectionTopLeft.Y * TileSetBeingEdited->TileHeight);
-			}
+		if (FTileSetEditorViewportClient* Client = TypedViewportClient.Get())
+		{
+			Client->bHasValidPaintRectangle = true;
+			Client->ValidPaintRectangle.Color = FLinearColor::White;
+			Client->ValidPaintRectangle.Dimensions = FVector2D(SelectionDimensions.X * TileSetBeingEdited->TileWidth, SelectionDimensions.Y * TileSetBeingEdited->TileHeight);
+			Client->ValidPaintRectangle.TopLeft = FVector2D(SelectionTopLeft.X * TileSetBeingEdited->TileWidth, SelectionTopLeft.Y * TileSetBeingEdited->TileHeight);
 		}
 	}
 }
@@ -239,10 +240,8 @@ public:
 	virtual TSharedRef<SWidget> CreateTabBody(const FWorkflowTabSpawnInfo& Info) const override
 	{
 		TSharedPtr<FTileSetEditor> TileSetEditorPtr = StaticCastSharedPtr<FTileSetEditor>(HostingApp.Pin());
-		//TSharedRef<FTileSetEditorViewportClient> ViewportClientPtr = MakeShareable(new FTileSetEditorViewportClient(TileSetEditorPtr));
-		//TWeakPtr<FTileSetEditorViewportClient> WeakViewportClient = ViewportClientPtr;
 
-		return SNew(STileSetSelectorViewport, TileSetEditorPtr->GetTileSetBeingEdited());
+		return SNew(STileSetSelectorViewport, TileSetEditorPtr->GetTileSetBeingEdited(), /*EdMode=*/ nullptr);
 	}
 };
 
diff --git a/Engine/Plugins/2D/Paper2D/Source/Paper2DEditor/Private/TileSetEditor.h b/Engine/Plugins/2D/Paper2D/Source/Paper2DEditor/Private/TileSetEditor.h
--- a/Engine/Plugins/2D/Paper2D/Source/Paper2DEditor/Private/TileSetEditor.h
+++ b/Engine/Plugins/2D/Paper2D/Source/Paper2DEditor/Private/TileSetEditor.h
@@ -46,7 +46,7 @@ public:
 
 	~STileSetSelectorViewport();
 
-	void Construct(const FArguments& InArgs, UPaperTileSet* InTileSet);
+	void Construct(const FArguments& InArgs, UPaperTileSet* InTileSet, FEdModeTileMap* InTileMapEditor);
 
 	void ChangeTileSet(UPaperTileSet* InTileSet);
 protected:
@@ -60,4 +60,5 @@ private:
 private:
 	TWeakObjectPtr<class UPaperTileSet> TileSetPtr;
 	TSharedPtr<class FTileSetEditorViewportClient> TypedViewportClient;
+	FEdModeTileMap* TileMapEditor;
 };
\ No newline at end of file
