Paper2D: Fix non-unity error in TileSetEditor.h

diff --git a/Engine/Plugins/2D/Paper2D/Source/Paper2DEditor/Private/TileSetEditor.h b/Engine/Plugins/2D/Paper2D/Source/Paper2DEditor/Private/TileSetEditor.h
--- a/Engine/Plugins/2D/Paper2D/Source/Paper2DEditor/Private/TileSetEditor.h
+++ b/Engine/Plugins/2D/Paper2D/Source/Paper2DEditor/Private/TileSetEditor.h
@@ -46,7 +46,7 @@ public:
 
 	~STileSetSelectorViewport();
 
-	void Construct(const FArguments& InArgs, UPaperTileSet* InTileSet, FEdModeTileMap* InTileMapEditor);
+	void Construct(const FArguments& InArgs, UPaperTileSet* InTileSet, class FEdModeTileMap* InTileMapEditor);
 
 	void ChangeTileSet(UPaperTileSet* InTileSet);
 protected:
@@ -60,5 +60,5 @@ private:
 private:
 	TWeakObjectPtr<class UPaperTileSet> TileSetPtr;
 	TSharedPtr<class FTileSetEditorViewportClient> TypedViewportClient;
-	FEdModeTileMap* TileMapEditor;
+	class FEdModeTileMap* TileMapEditor;
 };
\ No newline at end of file
