Editor: Add selection sets to FEditorModeTools, so code calling GEditor->GetSelectedActors/etc... can start to be refactored

diff --git a/Engine/Plugins/2D/Paper2D/Source/Paper2DEditor/Private/PaperEditorViewportClient.cpp b/Engine/Plugins/2D/Paper2D/Source/Paper2DEditor/Private/PaperEditorViewportClient.cpp
--- a/Engine/Plugins/2D/Paper2D/Source/Paper2DEditor/Private/PaperEditorViewportClient.cpp
+++ b/Engine/Plugins/2D/Paper2D/Source/Paper2DEditor/Private/PaperEditorViewportClient.cpp
@@ -5,11 +5,43 @@
 #include "PaperEditorViewportClient.h"
 #include "CanvasTypes.h"
 
+//////////////////////////////////////////////////////////////////////////
+// FAssetEditorModeTools
+
+FAssetEditorModeTools::FAssetEditorModeTools()
+{
+	ActorSet = NewObject<USelection>();
+	ActorSet->SetFlags(RF_Transactional);
+	ActorSet->AddToRoot();
+
+	ObjectSet = NewObject<USelection>();
+	ObjectSet->SetFlags(RF_Transactional);
+	ObjectSet->AddToRoot();
+}
+
+FAssetEditorModeTools::~FAssetEditorModeTools()
+{
+	ActorSet->RemoveFromRoot();
+	ActorSet = nullptr;
+	ObjectSet->RemoveFromRoot();
+	ObjectSet = nullptr;
+}
+
+USelection* FAssetEditorModeTools::GetSelectedActors() const
+{
+	return ActorSet;
+}
+
+USelection* FAssetEditorModeTools::GetSelectedObjects() const
+{
+	return ObjectSet;
+}
+
 //////////////////////////////////////////////////////////////////////////
 // FPaperEditorViewportClient
 
 FPaperEditorViewportClient::FPaperEditorViewportClient()
-	: FEditorViewportClient(*( new FEditorModeTools() ))
+	: FEditorViewportClient(*( new FAssetEditorModeTools() ))
 	, CheckerboardTexture(NULL)
 {
 	ZoomPos = FVector2D::ZeroVector;
diff --git a/Engine/Plugins/2D/Paper2D/Source/Paper2DEditor/Private/PaperEditorViewportClient.h b/Engine/Plugins/2D/Paper2D/Source/Paper2DEditor/Private/PaperEditorViewportClient.h
--- a/Engine/Plugins/2D/Paper2D/Source/Paper2DEditor/Private/PaperEditorViewportClient.h
+++ b/Engine/Plugins/2D/Paper2D/Source/Paper2DEditor/Private/PaperEditorViewportClient.h
@@ -11,6 +11,25 @@ struct FViewportSelectionRectangle
 	FLinearColor Color;
 };
 
+//////////////////////////////////////////////////////////////////////////
+// FAssetEditorModeTools
+
+class FAssetEditorModeTools : public FEditorModeTools
+{
+public:
+	FAssetEditorModeTools();
+	virtual ~FAssetEditorModeTools();
+
+	// FEditorModeTools interface
+	virtual class USelection* GetSelectedActors() const override;
+	virtual class USelection* GetSelectedObjects() const override;
+	// End of FEditorModeTools interface
+
+protected:
+	class USelection* ActorSet;
+	class USelection* ObjectSet;
+};
+
 //////////////////////////////////////////////////////////////////////////
 // FPaperEditorViewportClient
 
diff --git a/Engine/Source/Editor/UnrealEd/Private/EditorModes.cpp b/Engine/Source/Editor/UnrealEd/Private/EditorModes.cpp
--- a/Engine/Source/Editor/UnrealEd/Private/EditorModes.cpp
+++ b/Engine/Source/Editor/UnrealEd/Private/EditorModes.cpp
@@ -1242,6 +1242,16 @@ void FEditorModeTools::SetToolkitHost(TSharedRef<class IToolkitHost> InHost)
 	ToolkitHost = InHost;
 }
 
+class USelection* FEditorModeTools::GetSelectedActors() const
+{
+	return GEditor->GetSelectedActors();
+}
+
+class USelection* FEditorModeTools::GetSelectedObjects() const
+{
+	return GEditor->GetSelectedObjects();
+}
+
 void FEditorModeTools::OnEditorSelectionChanged(UObject* NewSelection)
 {
 	// If selecting an actor, move the pivot location.
diff --git a/Engine/Source/Editor/UnrealEd/Public/EditorModes.h b/Engine/Source/Editor/UnrealEd/Public/EditorModes.h
--- a/Engine/Source/Editor/UnrealEd/Public/EditorModes.h
+++ b/Engine/Source/Editor/UnrealEd/Public/EditorModes.h
@@ -444,7 +444,7 @@ public:
 	/**
 	 * Set the default editor mode for these tools
 	 */
-	void SetDefaultMode ( FEditorModeID DefaultID );
+	void SetDefaultMode(FEditorModeID DefaultID);
 
 	/**
 	 * Activates the default mode defined by this class
@@ -790,6 +790,15 @@ public:
 	/** Returns the host for toolkits created via modes from this mode manager */
 	TSharedPtr<class IToolkitHost> GetToolkitHost() const;
 
+	/**
+	 * Returns the set of selected actors.
+	 */
+	virtual class USelection* GetSelectedActors() const;
+
+	/**
+	 * @return the set of selected non-actor objects.
+	 */
+	virtual class USelection* GetSelectedObjects() const;
 protected:
 	/** 
 	 * Delegate handlers
