Paper2D: Tile map ed mode inside of standalone tile map editor WIP
- Remove overrides from FTileMapEditorViewportClient that get in the way of the tile map ed mode
- Select the render component on the mode tools when opened
- Change FEdModeTileMap::GetSelectedLayerUnderCursor to use the mode tools selection set instead of the global one

diff --git a/Engine/Plugins/2D/Paper2D/Source/Paper2DEditor/Private/TileMapEditing/EdModeTileMap.cpp b/Engine/Plugins/2D/Paper2D/Source/Paper2DEditor/Private/TileMapEditing/EdModeTileMap.cpp
--- a/Engine/Plugins/2D/Paper2D/Source/Paper2DEditor/Private/TileMapEditing/EdModeTileMap.cpp
+++ b/Engine/Plugins/2D/Paper2D/Source/Paper2DEditor/Private/TileMapEditing/EdModeTileMap.cpp
@@ -308,59 +308,75 @@ UPaperTileLayer* FEdModeTileMap::GetSelectedLayerUnderCursor(const FViewportCurs
 
 	const bool bCollisionPainting = (GetActiveLayerPaintingMode() == ETileMapLayerPaintingMode::CollisionLayers);
 
-	USelection* SelectedActors = GEditor->GetSelectedActors();
+	UPaperTileMapRenderComponent* TileMapComponent = nullptr;
+
+	USelection* SelectedActors = Owner->GetSelectedActors();
 	for (FSelectionIterator Iter(*SelectedActors); Iter; ++Iter)
 	{
 		AActor* Actor = CastChecked<AActor>(*Iter);
 
-		UPaperTileMap* TileMap = nullptr;
-		UPaperTileMapRenderComponent* TileMapComponent = Actor->FindComponentByClass<UPaperTileMapRenderComponent>();
+		TileMapComponent = Actor->FindComponentByClass<UPaperTileMapRenderComponent>();
 		if (TileMapComponent != nullptr)
 		{
-			TileMap = TileMapComponent->TileMap;
+			break;
 		}
+	}
 
-		if (TileMap != nullptr)
+	if (TileMapComponent == nullptr)
+	{
+		USelection* SelectedObjects = Owner->GetSelectedObjects();
+		for (FSelectionIterator Iter(*SelectedObjects); Iter; ++Iter)
 		{
-			// Find the first visible layer
-			int32 LayerIndex = 0;
-			for (; LayerIndex < TileMap->TileLayers.Num(); ++LayerIndex)
+			UObject* Foo = *Iter;
+			TileMapComponent = Cast<UPaperTileMapRenderComponent>(Foo);
+
+			if (TileMapComponent != nullptr)
 			{
-				UPaperTileLayer* Layer = TileMap->TileLayers[LayerIndex];
-				if (!Layer->bHiddenInEditor && Layer->bCollisionLayer == bCollisionPainting)
-				{
-					break;
-				}
+				break;
 			}
+		}
+	}
 
-			// If there was a visible layer, pick it
-			if (LayerIndex < TileMap->TileLayers.Num())
+	if (UPaperTileMap* TileMap = TileMapComponent->TileMap)
+	{
+		// Find the first visible layer
+		int32 LayerIndex = 0;
+		for (; LayerIndex < TileMap->TileLayers.Num(); ++LayerIndex)
+		{
+			UPaperTileLayer* Layer = TileMap->TileLayers[LayerIndex];
+			if (!Layer->bHiddenInEditor && Layer->bCollisionLayer == bCollisionPainting)
 			{
-				UPaperTileLayer* Layer = TileMap->TileLayers[LayerIndex];
+				break;
+			}
+		}
 
-				const float WX = TileMap->MapWidth * TileMap->TileWidth;
-				const float WY = TileMap->MapHeight * TileMap->TileHeight;
+		// If there was a visible layer, pick it
+		if (LayerIndex < TileMap->TileLayers.Num())
+		{
+			UPaperTileLayer* Layer = TileMap->TileLayers[LayerIndex];
 
-				ComponentToWorld = (TileMapComponent != nullptr) ? TileMapComponent->ComponentToWorld : FTransform::Identity;
-				FVector LocalStart = ComponentToWorld.InverseTransformPosition(TraceStart);
-				FVector LocalDirection = ComponentToWorld.InverseTransformVector(TraceDir);
+			const float WX = TileMap->MapWidth * TileMap->TileWidth;
+			const float WY = TileMap->MapHeight * TileMap->TileHeight;
 
-				FVector Intersection;
-				FPlane Plane(FVector(1, 0, 0), FVector::ZeroVector, FVector(0, 0, 1));
+			ComponentToWorld = (TileMapComponent != nullptr) ? TileMapComponent->ComponentToWorld : FTransform::Identity;
+			FVector LocalStart = ComponentToWorld.InverseTransformPosition(TraceStart);
+			FVector LocalDirection = ComponentToWorld.InverseTransformVector(TraceDir);
 
-				if (FMath::SegmentPlaneIntersection(LocalStart, LocalDirection * HALF_WORLD_MAX, Plane, /*out*/ Intersection))
-				{
-					//@TODO: Ideally tile pivots weren't in the center!
-					const float NormalizedX = (Intersection.X + 0.5f * TileMap->TileWidth) / WX;
-					const float NormalizedY = (-Intersection.Z + 0.5f * TileMap->TileHeight) / WY;
+			FVector Intersection;
+			FPlane Plane(FVector(1, 0, 0), FVector::ZeroVector, FVector(0, 0, 1));
 
-					OutTileX = FMath::FloorToInt(NormalizedX * TileMap->MapWidth);
-					OutTileY = FMath::FloorToInt(NormalizedY * TileMap->MapHeight);
+			if (FMath::SegmentPlaneIntersection(LocalStart, LocalDirection * HALF_WORLD_MAX, Plane, /*out*/ Intersection))
+			{
+				//@TODO: Ideally tile pivots weren't in the center!
+				const float NormalizedX = (Intersection.X + 0.5f * TileMap->TileWidth) / WX;
+				const float NormalizedY = (-Intersection.Z + 0.5f * TileMap->TileHeight) / WY;
+
+				OutTileX = FMath::FloorToInt(NormalizedX * TileMap->MapWidth);
+				OutTileY = FMath::FloorToInt(NormalizedY * TileMap->MapHeight);
 					
-					if ((OutTileX > -BrushWidth) && (OutTileX < TileMap->MapWidth) && (OutTileY > -BrushHeight) && (OutTileY < TileMap->MapHeight))
-					{
-						return Layer;
-					}
+				if ((OutTileX > -BrushWidth) && (OutTileX < TileMap->MapWidth) && (OutTileY > -BrushHeight) && (OutTileY < TileMap->MapHeight))
+				{
+					return Layer;
 				}
 			}
 		}
diff --git a/Engine/Plugins/2D/Paper2D/Source/Paper2DEditor/Private/TileMapEditing/TileMapEditorViewportClient.cpp b/Engine/Plugins/2D/Paper2D/Source/Paper2DEditor/Private/TileMapEditing/TileMapEditorViewportClient.cpp
--- a/Engine/Plugins/2D/Paper2D/Source/Paper2DEditor/Private/TileMapEditing/TileMapEditorViewportClient.cpp
+++ b/Engine/Plugins/2D/Paper2D/Source/Paper2DEditor/Private/TileMapEditing/TileMapEditorViewportClient.cpp
@@ -45,6 +45,9 @@ FTileMapEditorViewportClient::FTileMapEditorViewportClient(TWeakPtr<FTileMapEdit
 
 		PreviewScene->AddComponent(RenderTileMapComponent, FTransform::Identity);
 	}
+
+	// Select the render component
+	ModeTools->GetSelectedObjects()->Select(RenderTileMapComponent);
 }
 
 void FTileMapEditorViewportClient::ActivateEditMode()
@@ -149,87 +152,6 @@ bool FTileMapEditorViewportClient::IsShowMeshEdgesChecked() const
 	return EngineShowFlags.MeshEdges;
 }
 
-void FTileMapEditorViewportClient::UpdateMouseDelta()
-{
-	FPaperEditorViewportClient::UpdateMouseDelta();
-}
-
-void FTileMapEditorViewportClient::ProcessClick(FSceneView& View, HHitProxy* HitProxy, FKey Key, EInputEvent Event, uint32 HitX, uint32 HitY)
-{
-	FPaperEditorViewportClient::ProcessClick(View, HitProxy, Key, Event, HitX, HitY);
-}
-
-bool FTileMapEditorViewportClient::InputKey(FViewport* Viewport, int32 ControllerId, FKey Key, EInputEvent Event, float AmountDepressed, bool bGamepad)
-{
-	bool bHandled = false;
-
-	// Start the drag
-	//@TODO: EKeys::LeftMouseButton
-	//@TODO: Event.IE_Pressed
-	// Implement InputAxis
-	// StartTracking
-
-	// Pass keys to standard controls, if we didn't consume input
-	return (bHandled) ? true : FEditorViewportClient::InputKey(Viewport,  ControllerId, Key, Event, AmountDepressed, bGamepad);
-}
-
-bool FTileMapEditorViewportClient::InputWidgetDelta(FViewport* Viewport, EAxisList::Type CurrentAxis, FVector& Drag, FRotator& Rot, FVector& Scale)
-{
-	bool bHandled = false;
-	return bHandled;
-}
-
-void FTileMapEditorViewportClient::TrackingStarted(const struct FInputEventState& InInputState, bool bIsDragging, bool bNudge)
-{
-	if (!bManipulating && bIsDragging)
-	{
-		BeginTransaction(LOCTEXT("ModificationInViewport", "Modification in Viewport"));
-		bManipulating = true;
-		bManipulationDirtiedSomething = false;
-	}
-}
-
-void FTileMapEditorViewportClient::TrackingStopped()
-{
-	if (bManipulating)
-	{
-		EndTransaction();
-		bManipulating = false;
-	}
-}
-
-FWidget::EWidgetMode FTileMapEditorViewportClient::GetWidgetMode() const
-{
-	return FWidget::WM_None;
-}
-
-FVector FTileMapEditorViewportClient::GetWidgetLocation() const
-{
-
-	return FVector::ZeroVector;
-}
-
-FMatrix FTileMapEditorViewportClient::GetWidgetCoordSystem() const
-{
-	return FMatrix::Identity;
-}
-
-ECoordSystem FTileMapEditorViewportClient::GetWidgetCoordSystemSpace() const
-{
-	return COORD_World;
-}
-
-void FTileMapEditorViewportClient::BeginTransaction(const FText& SessionName)
-{
-	if (ScopedTransaction == NULL)
-	{
-		ScopedTransaction = new FScopedTransaction(SessionName);
-
-		UPaperTileMap* TileMap = GetTileMapBeingEdited();
-		TileMap->Modify();
-	}
-}
-
 void FTileMapEditorViewportClient::EndTransaction()
 {
 	if (bManipulationDirtiedSomething)
diff --git a/Engine/Plugins/2D/Paper2D/Source/Paper2DEditor/Private/TileMapEditing/TileMapEditorViewportClient.h b/Engine/Plugins/2D/Paper2D/Source/Paper2DEditor/Private/TileMapEditing/TileMapEditorViewportClient.h
--- a/Engine/Plugins/2D/Paper2D/Source/Paper2DEditor/Private/TileMapEditing/TileMapEditorViewportClient.h
+++ b/Engine/Plugins/2D/Paper2D/Source/Paper2DEditor/Private/TileMapEditing/TileMapEditorViewportClient.h
@@ -27,16 +27,6 @@ public:
 	// End of FViewportClient interface
 
 	// FEditorViewportClient interface
-	virtual void UpdateMouseDelta() override;
-	virtual void ProcessClick(FSceneView& View, HHitProxy* HitProxy, FKey Key, EInputEvent Event, uint32 HitX, uint32 HitY) override;
-	virtual bool InputKey(FViewport* Viewport, int32 ControllerId, FKey Key, EInputEvent Event, float AmountDepressed, bool bGamepad) override;
-	virtual bool InputWidgetDelta(FViewport* Viewport, EAxisList::Type CurrentAxis, FVector& Drag, FRotator& Rot, FVector& Scale) override;
-	virtual void TrackingStarted(const struct FInputEventState& InInputState, bool bIsDragging, bool bNudge) override;
-	virtual void TrackingStopped() override;
-	virtual FWidget::EWidgetMode GetWidgetMode() const override;
-	virtual FVector GetWidgetLocation() const;
-	virtual FMatrix GetWidgetCoordSystem() const;
-	virtual ECoordSystem GetWidgetCoordSystemSpace() const;
 	// End of FEditorViewportClient interface
 
 	void ToggleShowPivot() { bShowPivot = !bShowPivot; Invalidate(); }
