Paper2D: Opt in to using mode tools in tile map editor

diff --git a/Engine/Plugins/2D/Paper2D/Source/Paper2DEditor/Private/TileMapEditing/TileMapEditorViewportClient.cpp b/Engine/Plugins/2D/Paper2D/Source/Paper2DEditor/Private/TileMapEditing/TileMapEditorViewportClient.cpp
--- a/Engine/Plugins/2D/Paper2D/Source/Paper2DEditor/Private/TileMapEditing/TileMapEditorViewportClient.cpp
+++ b/Engine/Plugins/2D/Paper2D/Source/Paper2DEditor/Private/TileMapEditing/TileMapEditorViewportClient.cpp
@@ -16,9 +16,12 @@
 // FTileMapEditorViewportClient
 
 FTileMapEditorViewportClient::FTileMapEditorViewportClient(TWeakPtr<FTileMapEditor> InTileMapEditor, TWeakPtr<class STileMapEditorViewport> InTileMapEditorViewportPtr)
-: TileMapEditorPtr(InTileMapEditor)
-, TileMapEditorViewportPtr(InTileMapEditorViewportPtr)
+	: TileMapEditorPtr(InTileMapEditor)
+	, TileMapEditorViewportPtr(InTileMapEditorViewportPtr)
 {
+	// The tile map editor fully supports mode tools and isn't doing any incompatible stuff with the Widget
+	Widget->SetUsesEditorModeTools(ModeTools);
+
 	check(TileMapEditorPtr.IsValid() && TileMapEditorViewportPtr.IsValid());
 
 	PreviewScene = &OwnedPreviewScene;
