Paper2D: Add ContainsSprite method to UPaperFlipbook

diff --git a/Engine/Plugins/2D/Paper2D/Source/Paper2D/Classes/PaperFlipbook.h b/Engine/Plugins/2D/Paper2D/Source/Paper2D/Classes/PaperFlipbook.h
--- a/Engine/Plugins/2D/Paper2D/Source/Paper2D/Classes/PaperFlipbook.h
+++ b/Engine/Plugins/2D/Paper2D/Source/Paper2D/Classes/PaperFlipbook.h
@@ -121,6 +121,9 @@ public:
 #endif
 	// End of UObject interface
 
+	// Returns true if the flipbook asset contains the specified sprite asset in any frames
+	bool ContainsSprite(UPaperSprite* SpriteAsset) const;
+
 private:
 	friend class FScopedFlipbookMutator;
 };
diff --git a/Engine/Plugins/2D/Paper2D/Source/Paper2D/Private/PaperFlipbook.cpp b/Engine/Plugins/2D/Paper2D/Source/Paper2D/Private/PaperFlipbook.cpp
--- a/Engine/Plugins/2D/Paper2D/Source/Paper2D/Private/PaperFlipbook.cpp
+++ b/Engine/Plugins/2D/Paper2D/Source/Paper2D/Private/PaperFlipbook.cpp
@@ -209,3 +209,16 @@ void UPaperFlipbook::InvalidateCachedData()
 {
 	// No cached data yet, but the functions that currently have to iterate over all frames can use cached data in the future
 }
+
+bool UPaperFlipbook::ContainsSprite(UPaperSprite* SpriteAsset) const
+{
+	for (const FPaperFlipbookKeyFrame& KeyFrame : KeyFrames)
+	{
+		if (KeyFrame.Sprite == SpriteAsset)
+		{
+			return true;
+		}
+	}
+
+	return false;
+}
