Paper2D: Optimize which components are re-registered when a sprite asset is edited

diff --git a/Engine/Plugins/2D/Paper2D/Source/Paper2D/Private/PaperSprite.cpp b/Engine/Plugins/2D/Paper2D/Source/Paper2D/Private/PaperSprite.cpp
--- a/Engine/Plugins/2D/Paper2D/Source/Paper2D/Private/PaperSprite.cpp
+++ b/Engine/Plugins/2D/Paper2D/Source/Paper2D/Private/PaperSprite.cpp
@@ -279,11 +279,64 @@ UPaperSprite::UPaperSprite(const FObjectInitializer& ObjectInitializer)
 }
 
 #if WITH_EDITOR
+
+
+/** Removes all components that use the specified sprite asset from their scenes for the lifetime of the class. */
+class FSpriteReregisterContext
+{
+public:
+	/** Initialization constructor. */
+	FSpriteReregisterContext(UPaperSprite* TargetAsset)
+	{
+		// Look at sprite components
+		for (TObjectIterator<UPaperSpriteComponent> SpriteIt; SpriteIt; ++SpriteIt)
+		{
+			if (UPaperSpriteComponent* TestComponent = *SpriteIt)
+			{
+				if (TestComponent->GetSprite() == TargetAsset)
+				{
+					AddComponentToRefresh(TestComponent);
+				}
+			}
+		}
+
+		// Look at flipbook components
+		for (TObjectIterator<UPaperFlipbookComponent> FlipbookIt; FlipbookIt; ++FlipbookIt)
+		{
+			if (UPaperFlipbookComponent* TestComponent = *FlipbookIt)
+			{
+				if (UPaperFlipbook* Flipbook = TestComponent->GetFlipbook())
+				{
+					if (Flipbook->ContainsSprite(TargetAsset))
+					{
+						AddComponentToRefresh(TestComponent);
+					}
+				}
+			}
+		}
+	}
+
+protected:
+	void AddComponentToRefresh(UActorComponent* Component)
+	{
+		if (ComponentContexts.Num() == 0)
+		{
+			// wait until resources are released
+			FlushRenderingCommands();
+		}
+
+		new (ComponentContexts) FComponentReregisterContext(Component);
+	}
+
+private:
+	/** The recreate contexts for the individual components. */
+	TIndirectArray<FComponentReregisterContext> ComponentContexts;
+};
+
 void UPaperSprite::PostEditChangeProperty(FPropertyChangedEvent& PropertyChangedEvent)
 {
-	//@TODO: Determine when these are really needed, as they're seriously expensive!
-	TComponentReregisterContext<UPaperSpriteComponent> ReregisterStaticComponents;
-	TComponentReregisterContext<UPaperFlipbookComponent> ReregisterAnimatedComponents;
+	//@TODO: Determine when this is really needed, as it is seriously expensive!
+	FSpriteReregisterContext ReregisterExistingComponents(this);
 
 	// Look for changed properties
 	const FName PropertyName = (PropertyChangedEvent.Property != NULL) ? PropertyChangedEvent.Property->GetFName() : NAME_None;
