Paper2D: Add a property restriction to prevent sprite collision geometry type from being set to Diced

diff --git a/Engine/Plugins/2D/Paper2D/Source/Paper2D/Classes/SpriteEditorOnlyTypes.h b/Engine/Plugins/2D/Paper2D/Source/Paper2D/Classes/SpriteEditorOnlyTypes.h
--- a/Engine/Plugins/2D/Paper2D/Source/Paper2D/Classes/SpriteEditorOnlyTypes.h
+++ b/Engine/Plugins/2D/Paper2D/Source/Paper2D/Classes/SpriteEditorOnlyTypes.h
@@ -51,7 +51,7 @@ namespace ESpritePolygonMode
 		// Fully custom geometry; edited by hand
 		FullyCustom,
 
-		// Diced (split up into smaller squares, including only non-empty ones in the final geometry)
+		// Diced (split up into smaller squares, including only non-empty ones in the final geometry).  This option is only supported for Render geometry and will be ignored for Collision geometry.
 		Diced
 	};
 }
diff --git a/Engine/Plugins/2D/Paper2D/Source/Paper2DEditor/Private/SpriteEditor/SpriteDetailsCustomization.cpp b/Engine/Plugins/2D/Paper2D/Source/Paper2DEditor/Private/SpriteEditor/SpriteDetailsCustomization.cpp
--- a/Engine/Plugins/2D/Paper2D/Source/Paper2DEditor/Private/SpriteEditor/SpriteDetailsCustomization.cpp
+++ b/Engine/Plugins/2D/Paper2D/Source/Paper2DEditor/Private/SpriteEditor/SpriteDetailsCustomization.cpp
@@ -6,6 +6,7 @@
 #include "PhysicsEngine/BodySetup.h"
 #include "Editor/Documentation/Public/IDocumentation.h"
 #include "PhysicsEngine/PhysicsSettings.h"
+#include "PropertyRestriction.h"
 
 #define LOCTEXT_NAMESPACE "SpriteEditor"
 
@@ -91,7 +92,7 @@ void FSpriteDetailsCustomization::BuildSpriteSection(IDetailCategoryBuilder& Spr
 
 void FSpriteDetailsCustomization::BuildRenderingSection(IDetailCategoryBuilder& RenderingCategory, IDetailLayoutBuilder& DetailLayout)
 {
-	// Add the collision geometry mode into the parent container (renamed)
+	// Add the rendering geometry mode into the parent container (renamed)
 	const FString RenderGeometryTypePropertyPath = FString::Printf(TEXT("%s.%s"), GET_MEMBER_NAME_STRING_CHECKED(UPaperSprite, RenderGeometry), GET_MEMBER_NAME_STRING_CHECKED(FSpritePolygonCollection, GeometryType));
 	RenderingCategory.AddProperty(DetailLayout.GetProperty(*RenderGeometryTypePropertyPath))
 		.DisplayName(LOCTEXT("RenderGeometryType", "Render Geometry Type").ToString());
@@ -129,10 +130,21 @@ void FSpriteDetailsCustomization::BuildCollisionSection(IDetailCategoryBuilder&
 		.Visibility(WarnAbout2DQueriesBeingDisabledVisibility);
 
 	// Add the collision geometry mode into the parent container (renamed)
-	const FString CollisionGeometryTypePropertyPath = FString::Printf(TEXT("%s.%s"), GET_MEMBER_NAME_STRING_CHECKED(UPaperSprite, CollisionGeometry), GET_MEMBER_NAME_STRING_CHECKED(FSpritePolygonCollection, GeometryType));
-	CollisionCategory.AddProperty(DetailLayout.GetProperty(*CollisionGeometryTypePropertyPath))
-		.DisplayName(LOCTEXT("CollisionGeometryType", "Collision Geometry Type").ToString())
-		.Visibility(ParticipatesInPhysics);
+	{
+		// Restrict the diced value
+		TSharedPtr<FPropertyRestriction> PreventDicedRestriction = MakeShareable(new FPropertyRestriction(LOCTEXT("CollisionGeometryDoesNotSupportDiced", "Collision geometry can not be set to Diced")));
+		PreventDicedRestriction->AddValue(TEXT("Diced"));
+
+		// Find and add the property
+		const FString CollisionGeometryTypePropertyPath = FString::Printf(TEXT("%s.%s"), GET_MEMBER_NAME_STRING_CHECKED(UPaperSprite, CollisionGeometry), GET_MEMBER_NAME_STRING_CHECKED(FSpritePolygonCollection, GeometryType));
+		TSharedPtr<IPropertyHandle> CollisionGeometryTypeProperty = DetailLayout.GetProperty(*CollisionGeometryTypePropertyPath);
+
+		CollisionGeometryTypeProperty->AddRestriction(PreventDicedRestriction.ToSharedRef());
+
+		CollisionCategory.AddProperty(CollisionGeometryTypeProperty)
+			.DisplayName(LOCTEXT("CollisionGeometryType", "Collision Geometry Type").ToString())
+			.Visibility(ParticipatesInPhysics);
+	}
 
 	// Show the collision geometry when not None
 	CollisionCategory.AddProperty( DetailLayout.GetProperty(GET_MEMBER_NAME_CHECKED(UPaperSprite, CollisionGeometry)) )
