Paper2D: Allow adding additional vertices on top of an existing edge (Shift+click always adds rather than preferentially selecting existing proxies)

diff --git a/Engine/Plugins/2D/Paper2D/Source/Paper2DEditor/Private/SpriteEditor/SpriteEditorViewportClient.cpp b/Engine/Plugins/2D/Paper2D/Source/Paper2DEditor/Private/SpriteEditor/SpriteEditorViewportClient.cpp
--- a/Engine/Plugins/2D/Paper2D/Source/Paper2DEditor/Private/SpriteEditor/SpriteEditorViewportClient.cpp
+++ b/Engine/Plugins/2D/Paper2D/Source/Paper2DEditor/Private/SpriteEditor/SpriteEditorViewportClient.cpp
@@ -970,12 +970,13 @@ void FSpriteEditorViewportClient::ProcessClick(FSceneView& View, HHitProxy* HitP
 	const bool bIsShiftKeyDown = Viewport->KeyState(EKeys::LeftShift) || Viewport->KeyState(EKeys::RightShift);
 	const bool bIsAltKeyDown = Viewport->KeyState(EKeys::LeftAlt) || Viewport->KeyState(EKeys::RightAlt);
 	bool bHandled = false;
-	bool bAllowSelectVertex = !(IsEditingGeometry() && IsAddingPolygon());
-	HSpriteSelectableObjectHitProxy* SelectedItemProxy = HitProxyCast<HSpriteSelectableObjectHitProxy>(HitProxy);
 
-	bool bClearSelectionModifier = bIsCtrlKeyDown;
-	bool bDeleteClickedVertex = bIsAltKeyDown;
-	bool bInsertVertexModifier = bIsShiftKeyDown;
+	const bool bAllowSelectVertex = !(IsEditingGeometry() && IsAddingPolygon()) && !bIsShiftKeyDown;
+
+	const bool bClearSelectionModifier = bIsCtrlKeyDown;
+	const bool bDeleteClickedVertex = bIsAltKeyDown;
+	const bool bInsertVertexModifier = bIsShiftKeyDown;
+	HSpriteSelectableObjectHitProxy* SelectedItemProxy = HitProxyCast<HSpriteSelectableObjectHitProxy>(HitProxy);
 
 	if (bAllowSelectVertex && SelectedItemProxy)
 	{
