Paper2D: Improve sprite extraction behavior, now sorts rectangles by row and names with an explicit desired number suffix

diff --git a/Engine/Plugins/2D/Paper2D/Source/Paper2DEditor/Private/ContentBrowserExtensions/ContentBrowserExtensions.cpp b/Engine/Plugins/2D/Paper2D/Source/Paper2DEditor/Private/ContentBrowserExtensions/ContentBrowserExtensions.cpp
--- a/Engine/Plugins/2D/Paper2D/Source/Paper2DEditor/Private/ContentBrowserExtensions/ContentBrowserExtensions.cpp
+++ b/Engine/Plugins/2D/Paper2D/Source/Paper2DEditor/Private/ContentBrowserExtensions/ContentBrowserExtensions.cpp
@@ -94,6 +94,48 @@ struct FCreateSpriteFromTextureExtension : public FContentBrowserSelectedAssetEx
 				TArray<FIntRect> ExtractedRects;
 				UPaperSprite::ExtractRectsFromTexture(Texture, /*out*/ ExtractedRects);
 
+				// Sort the rectangles by approximate row
+				struct FRectangleSortHelper
+				{
+					FRectangleSortHelper(TArray<FIntRect>& InOutSprites)
+					{
+						// Sort by Y, then by X (top left corner), descending order (so we can use it as a stack from the top row down)
+						TArray<FIntRect> SpritesLeft = InOutSprites;
+						SpritesLeft.Sort([](const FIntRect& A, const FIntRect& B) { return (A.Min.Y == B.Min.Y) ? (A.Min.X > B.Min.X) : (A.Min.Y > B.Min.Y); });
+						InOutSprites.Reset();
+
+						// Start pulling sprites out, the first one in each row will dominate remaining ones and cause them to get labeled
+						TArray<FIntRect> DominatedSprites;
+						DominatedSprites.Empty(SpritesLeft.Num());
+ 						while (SpritesLeft.Num())
+ 						{
+							FIntRect DominatingSprite = SpritesLeft.Pop();
+							DominatedSprites.Add(DominatingSprite);
+
+							// Find the sprites that are dominated (intersect the infinite horizontal band described by the dominating sprite)
+							for (int32 Index = 0; Index < SpritesLeft.Num();)
+							{
+								const FIntRect& CurElement = SpritesLeft[Index];
+								if ((CurElement.Min.Y <= DominatingSprite.Max.Y) && (CurElement.Max.Y >= DominatingSprite.Min.Y))
+								{
+									DominatedSprites.Add(CurElement);
+									SpritesLeft.RemoveAt(Index, /*Count=*/ 1, /*bAllowShrinking=*/ false);
+								}
+								else
+								{
+									++Index;
+								}
+							}
+
+							// Sort the sprites in the band by X and add them to the result
+							DominatedSprites.Sort([](const FIntRect& A, const FIntRect& B) { return (A.Min.X < B.Min.X); });
+							InOutSprites.Append(DominatedSprites);
+							DominatedSprites.Reset();
+ 						}
+					}
+				};
+				FRectangleSortHelper RectSorter(ExtractedRects);
+
 				Feedback.TotalAmountOfWork = ExtractedRects.Num();
 
 				for (int ExtractedRectIndex = 0; ExtractedRectIndex < ExtractedRects.Num(); ++ExtractedRectIndex)
@@ -106,7 +148,8 @@ struct FCreateSpriteFromTextureExtension : public FContentBrowserSelectedAssetEx
 					SpriteFactory->InitialSourceDimension = FVector2D(ExtractedRect.Width(), ExtractedRect.Height());
 
 					// Get a unique name for the sprite
-					AssetToolsModule.Get().CreateUniqueAssetName(Texture->GetOutermost()->GetName(), DefaultSuffix, /*out*/ PackageName, /*out*/ Name);
+					const FString Suffix = FString::Printf(TEXT("%s_%d"), *DefaultSuffix, ExtractedRectIndex);
+					AssetToolsModule.Get().CreateUniqueAssetName(Texture->GetOutermost()->GetName(), Suffix, /*out*/ PackageName, /*out*/ Name);
 					const FString PackagePath = FPackageName::GetLongPackagePath(PackageName);
 
 					if (UObject* NewAsset = AssetToolsModule.Get().CreateAsset(Name, PackagePath, UPaperSprite::StaticClass(), SpriteFactory))
