Paper2D: Delete unused method ExtractSpriteNumber

diff --git a/Engine/Plugins/2D/Paper2D/Source/PaperJsonImporter/Private/PaperSpriteSheetAssetTypeActions.cpp b/Engine/Plugins/2D/Paper2D/Source/PaperJsonImporter/Private/PaperSpriteSheetAssetTypeActions.cpp
--- a/Engine/Plugins/2D/Paper2D/Source/PaperJsonImporter/Private/PaperSpriteSheetAssetTypeActions.cpp
+++ b/Engine/Plugins/2D/Paper2D/Source/PaperJsonImporter/Private/PaperSpriteSheetAssetTypeActions.cpp
@@ -70,60 +70,6 @@ void FPaperSpriteSheetAssetTypeActions::GetActions(const TArray<UObject*>& InObj
 
 //////////////////////////////////////////////////////////////////////////
 
-static bool ExtractSpriteNumber(const FString& String, FString& BareString, int& Number)
-{
-	bool bExtracted = false;
-
-	int LastCharacter = String.Len() - 1;
-	if (LastCharacter >= 0)
-	{
-		// Find the last character that isn't a digit (Handle sprite names with numbers inside inverted commas / parentheses)
-		while (LastCharacter >= 0 && !FChar::IsDigit(String[LastCharacter]))
-		{
-			LastCharacter--;
-		}
-
-		// Only proceed if we found a number in the sprite name
-		if (LastCharacter >= 0)
-		{
-			while (LastCharacter > 0 && FChar::IsDigit(String[LastCharacter - 1]))
-			{
-				LastCharacter--;
-			}
-
-			if (LastCharacter >= 0)
-			{
-				int FirstDigit = LastCharacter;
-				int EndCharacter = FirstDigit;
-				while (EndCharacter > 0 && !FChar::IsAlnum(String[EndCharacter - 1]))
-				{
-					EndCharacter--;
-				}
-
-				if (EndCharacter == 0)
-				{
-					// This string consists of non alnum + number, eg. _42
-					// The flipbook / category name in this case will be _
-					// Otherwise, we strip out all trailing non-alnum chars
-					EndCharacter = FirstDigit;
-				}
-
-				FString NumberString = String.Mid(FirstDigit);
-				BareString = String.Left(EndCharacter);
-				Number = FCString::Atoi(*NumberString);
-				bExtracted = true;
-			}
-		}
-	}
-
-	if (!bExtracted)
-	{
-		BareString = String;
-		Number = -1;
-	}
-	return bExtracted;
-}
-
 void FPaperSpriteSheetAssetTypeActions::ExecuteCreateFlipbooks(TArray<TWeakObjectPtr<UPaperSpriteSheet>> Objects)
 {	
 	for (int SpriteSheetIndex = 0; SpriteSheetIndex < Objects.Num(); ++SpriteSheetIndex)
