Removed DrawDynamicElements, PreRenderView and dependencies.  These have been disabled since 11/04, and the GetDynamicMeshElements path is now the only path for dynamic mesh rendering.

diff --git a/Engine/Plugins/2D/Paper2D/Source/Paper2D/Private/PaperRenderSceneProxy.cpp b/Engine/Plugins/2D/Paper2D/Source/Paper2D/Private/PaperRenderSceneProxy.cpp
--- a/Engine/Plugins/2D/Paper2D/Source/Paper2D/Private/PaperRenderSceneProxy.cpp
+++ b/Engine/Plugins/2D/Paper2D/Source/Paper2D/Private/PaperRenderSceneProxy.cpp
@@ -292,97 +292,6 @@ void FPaperRenderSceneProxy::GetDynamicMeshElements(const TArray<const FSceneVie
 	}
 }
 
-void FPaperRenderSceneProxy::DrawDynamicElements(FPrimitiveDrawInterface* PDI, const FSceneView* View)
-{
-	QUICK_SCOPE_CYCLE_COUNTER(STAT_PaperRenderSceneProxy_DrawDynamicElements);
-
-	checkSlow(IsInRenderingThread());
-
-
-	const FEngineShowFlags& EngineShowFlags = View->Family->EngineShowFlags;
-
-	bool bDrawSimpleCollision = false;
-	bool bDrawComplexCollision = false;
-	const bool bInCollisionView = IsCollisionView(EngineShowFlags, bDrawSimpleCollision, bDrawComplexCollision);
-
-	// Sprites don't distinguish between simple and complex collision; when viewing visibility we should still render simple collision
-	bDrawSimpleCollision |= bDrawComplexCollision;
-	
-	// Draw simple collision as wireframe if 'show collision', collision is enabled
-	const bool bDrawWireframeCollision = EngineShowFlags.Collision && IsCollisionEnabled();
-
-	if ((bDrawSimpleCollision || bDrawWireframeCollision) && AllowDebugViewmodes())
-	{
-		if (BodySetup)
-		{
-			if (FMath::Abs(GetLocalToWorld().Determinant()) < SMALL_NUMBER)
-			{
-				// Catch this here or otherwise GeomTransform below will assert
-				// This spams so commented out
-				//UE_LOG(LogStaticMesh, Log, TEXT("Zero scaling not supported (%s)"), *StaticMesh->GetPathName());
-			}
-			else
-			{
-				const bool bDrawSolid = !bDrawWireframeCollision;
-
-				if (bDrawSolid)
-				{
-					// Make a material for drawing solid collision stuff
-					const FColoredMaterialRenderProxy SolidMaterialInstance(
-						GEngine->ShadedLevelColorationUnlitMaterial->GetRenderProxy(IsSelected(), IsHovered()),
-						WireframeColor
-						);
-
-					FTransform GeomTransform(GetLocalToWorld());
-					BodySetup->AggGeom.DrawAggGeom(PDI, GeomTransform, WireframeColor, &SolidMaterialInstance, false, true, UseEditorDepthTest());
-				}
-				else
-				{
-					// wireframe
-					FColor CollisionColor = FColor(157, 149, 223, 255);
-					FTransform GeomTransform(GetLocalToWorld());
-					BodySetup->AggGeom.DrawAggGeom(PDI, GeomTransform, GetSelectionColor(CollisionColor, IsSelected(), IsHovered()), NULL, (Owner == NULL), false, UseEditorDepthTest());
-				}
-			}
-		}
-	}
-
-	// Draw bounds
-#if !(UE_BUILD_SHIPPING || UE_BUILD_TEST)
-	if (EngineShowFlags.Paper2DSprites)
-	{
-		RenderBounds(PDI, EngineShowFlags, GetBounds(), (Owner == NULL) || IsSelected());
-	}
-#endif
-
-	const bool bDrawSprite = !bInCollisionView;
-
-	if (bDrawSprite)
-	{
-		// Set the selection/hover color from the current engine setting.
-		FLinearColor OverrideColor;
-		bool bUseOverrideColor = false;
-
-		const bool bShowAsSelected = !(GIsEditor && View->Family->EngineShowFlags.Selection) || IsSelected();
-		if (bShowAsSelected || IsHovered())
-		{
-			bUseOverrideColor = true;
-			OverrideColor = GetSelectionColor(FLinearColor::White, bShowAsSelected, IsHovered());
-		}
-
-		bUseOverrideColor = false;
-
-		// Sprites of locked actors draw in red.
-		//FLinearColor LevelColorToUse = IsSelected() ? ColorToUse : (FLinearColor)LevelColor;
-		//FLinearColor PropertyColorToUse = PropertyColor;
-
-#if TEST_BATCHING
-#else
-		DrawDynamicElements_RichMesh(PDI, View, bUseOverrideColor, OverrideColor);
-#endif
-	}
-}
-
 FVertexFactory* FPaperRenderSceneProxy::GetPaperSpriteVertexFactory() const
 {
 	static bool bInitAxes = false;
@@ -405,83 +314,6 @@ void FPaperRenderSceneProxy::GetDynamicMeshElementsForView(const FSceneView* Vie
 	GetNewBatchMeshes(View, bUseOverrideColor, OverrideColor, ViewIndex, Collector);
 }
 
-void FPaperRenderSceneProxy::DrawDynamicElements_RichMesh(FPrimitiveDrawInterface* PDI, const FSceneView* View, bool bUseOverrideColor, const FLinearColor& OverrideColor)
-{
-	if (Material != nullptr)
-	{
-		DrawBatch(PDI, View, bUseOverrideColor, OverrideColor, Material, BatchedSprites);
-	}
-	DrawNewBatches(PDI, View, bUseOverrideColor, OverrideColor);
-}
-
-void FPaperRenderSceneProxy::DrawNewBatches(FPrimitiveDrawInterface* PDI, const FSceneView* View, bool bUseOverrideColor, const FLinearColor& OverrideColor)
-{
-	//@TODO: Doesn't support OverrideColor yet
-	if (BatchedSections.Num() == 0)
-	{
-		return;
-	}
-
-	const uint8 DPG = GetDepthPriorityGroup(View);
-	
-	FVertexFactory* VertexFactory = GetPaperSpriteVertexFactory();
-
-	FMeshBatch Mesh;
-
-	Mesh.UseDynamicData = true;
-	Mesh.DynamicVertexData = BatchedVertices.GetData();
-	Mesh.DynamicVertexStride = sizeof(FPaperSpriteVertex);
-
-	Mesh.VertexFactory = VertexFactory;
-	Mesh.LCI = NULL;
-	Mesh.ReverseCulling = IsLocalToWorldDeterminantNegative() ? true : false;
-	Mesh.CastShadow = bCastShadow;
-	Mesh.DepthPriorityGroup = DPG;
-	Mesh.Type = PT_TriangleList;
-	Mesh.bDisableBackfaceCulling = true;
-
-	// Set up the FMeshBatchElement.
-	FMeshBatchElement& BatchElement = Mesh.Elements[0];
-	BatchElement.IndexBuffer = NULL;
-	BatchElement.DynamicIndexData = NULL;
-	BatchElement.DynamicIndexStride = 0;
-	BatchElement.PrimitiveUniformBuffer = GetUniformBuffer();
-
-	const bool bIsWireframeView = View->Family->EngineShowFlags.Wireframe;
-
-	for (const FSpriteRenderSection& Batch : BatchedSections)
-	{
-		const FLinearColor EffectiveWireframeColor = (Batch.Material->GetBlendMode() != BLEND_Opaque) ? WireframeColor : FLinearColor::Green;
-		FMaterialRenderProxy* ParentMaterialProxy = Batch.Material->GetRenderProxy((View->Family->EngineShowFlags.Selection) && IsSelected(), IsHovered());
-
-		FTexture* TextureResource = (Batch.Texture != nullptr) ? Batch.Texture->Resource : nullptr;
-		if ((TextureResource != nullptr) && (Batch.NumVertices > 0))
-		{
-			// Create a texture override material proxy and register it as a dynamic resource so that it won't be deleted until the rendering thread has finished with it
-			FTextureOverrideRenderProxy* TextureOverrideMaterialProxy = new FTextureOverrideRenderProxy(ParentMaterialProxy, Batch.Texture, TEXT("SpriteTexture"));
-			PDI->RegisterDynamicResource(TextureOverrideMaterialProxy);
-
-			Mesh.MaterialRenderProxy = TextureOverrideMaterialProxy;
-
-			BatchElement.FirstIndex = Batch.VertexOffset;
-			BatchElement.MinVertexIndex = Batch.VertexOffset;
-			BatchElement.MaxVertexIndex = Batch.VertexOffset + Batch.NumVertices;
-			BatchElement.NumPrimitives = Batch.NumVertices / 3;
-
-			DrawRichMesh(
-				PDI, 
-				Mesh,
-				EffectiveWireframeColor,
-				FLinearColor(1.0f, 1.0f, 1.0f),//LevelColor,
-				FLinearColor(1.0f, 1.0f, 1.0f),//PropertyColor,
-				this,
-				IsSelected(),
-				bIsWireframeView
-				);
-		}
-	}
-}
-
 void FPaperRenderSceneProxy::GetNewBatchMeshes(const FSceneView* View, bool bUseOverrideColor, const FLinearColor& OverrideColor, int32 ViewIndex, FMeshElementCollector& Collector) const
 {
 	//@TODO: Doesn't support OverrideColor yet
@@ -651,83 +483,6 @@ void FPaperRenderSceneProxy::GetBatchMesh(const FSceneView* View, bool bUseOverr
 	}
 }
 
-void FPaperRenderSceneProxy::DrawBatch(FPrimitiveDrawInterface* PDI, const FSceneView* View, bool bUseOverrideColor, const FLinearColor& OverrideColor, class UMaterialInterface* BatchMaterial, const TArray<FSpriteDrawCallRecord>& Batch)
-{
-	const uint8 DPG = GetDepthPriorityGroup(View);
-	FVertexFactory* VertexFactory = GetPaperSpriteVertexFactory();
-
-	const bool bIsWireframeView = View->Family->EngineShowFlags.Wireframe;
-	const FLinearColor EffectiveWireframeColor = (BatchMaterial->GetBlendMode() != BLEND_Opaque) ? WireframeColor : FLinearColor::Green;
-
-	for (const FSpriteDrawCallRecord& Record : Batch)
-	{
-		FTexture* TextureResource = (Record.Texture != nullptr) ? Record.Texture->Resource : nullptr;
-		if ((TextureResource != nullptr) && (Record.RenderVerts.Num() > 0))
-		{
-			const FLinearColor SpriteColor = bUseOverrideColor ? OverrideColor : Record.Color;
-
-			const FVector EffectiveOrigin = Record.Destination;
-
-			TArray< FPaperSpriteVertex, TInlineAllocator<6> > Vertices;
-			Vertices.Empty(Record.RenderVerts.Num());
-			for (int32 SVI = 0; SVI < Record.RenderVerts.Num(); ++SVI)
-			{
-				const FVector4& SourceVert = Record.RenderVerts[SVI];
-
-				const FVector Pos((PaperAxisX * SourceVert.X) + (PaperAxisY * SourceVert.Y) + EffectiveOrigin);
-				const FVector2D UV(SourceVert.Z, SourceVert.W);
-
-				new (Vertices) FPaperSpriteVertex(Pos, UV, SpriteColor);
-			}
-
-			// Set up the FMeshElement.
-			FMeshBatch Mesh;
-
-			Mesh.UseDynamicData = true;
-			Mesh.DynamicVertexData = Vertices.GetData();
-			Mesh.DynamicVertexStride = sizeof(FPaperSpriteVertex);
-
-			Mesh.VertexFactory = VertexFactory;
-
-			FMaterialRenderProxy* ParentMaterialProxy = BatchMaterial->GetRenderProxy((View->Family->EngineShowFlags.Selection) && IsSelected(), IsHovered());
-
-			// Create a texture override material proxy and register it as a dynamic resource so that it won't be deleted until the rendering thread has finished with it
-			FTextureOverrideRenderProxy* TextureOverrideMaterialProxy = new FTextureOverrideRenderProxy(ParentMaterialProxy, Record.Texture, TEXT("SpriteTexture"));
-			PDI->RegisterDynamicResource( TextureOverrideMaterialProxy );
-
-			Mesh.MaterialRenderProxy = TextureOverrideMaterialProxy;
-			Mesh.LCI = NULL;
-			Mesh.ReverseCulling = IsLocalToWorldDeterminantNegative() ? true : false;
-			Mesh.CastShadow = bCastShadow;
-			Mesh.DepthPriorityGroup = DPG;
-			Mesh.Type = PT_TriangleList;
-			Mesh.bDisableBackfaceCulling = true;
-
-			// Set up the FMeshBatchElement.
-			FMeshBatchElement& BatchElement = Mesh.Elements[0];
-			BatchElement.IndexBuffer = NULL;
-			BatchElement.DynamicIndexData = NULL;
-			BatchElement.DynamicIndexStride = 0;
-			BatchElement.FirstIndex = 0;
-			BatchElement.MinVertexIndex = 0;
-			BatchElement.MaxVertexIndex = Vertices.Num();
-			BatchElement.PrimitiveUniformBuffer = GetUniformBuffer();
-			BatchElement.NumPrimitives = Vertices.Num() / 3;
-
-			DrawRichMesh(
-				PDI, 
-				Mesh,
-				EffectiveWireframeColor,
-				FLinearColor(1.0f, 1.0f, 1.0f),//LevelColor,
-				FLinearColor(1.0f, 1.0f, 1.0f),//PropertyColor,
-				this,
-				IsSelected(),
-				bIsWireframeView
-				);
-		}
-	}
-}
-
 FPrimitiveViewRelevance FPaperRenderSceneProxy::GetViewRelevance(const FSceneView* View)
 {
 	const FEngineShowFlags& EngineShowFlags = View->Family->EngineShowFlags;
@@ -755,29 +510,6 @@ FPrimitiveViewRelevance FPaperRenderSceneProxy::GetViewRelevance(const FSceneVie
 
 	Result.bDynamicRelevance = true;
 
-// 	if (
-// #if SUPPORT_EXTRA_RENDERING
-// 		IsRichView(View) ||
-// 		EngineShowFlags.Collision ||
-// 		bInCollisionView ||
-// 		EngineShowFlags.Bounds ||
-// #endif
-// 		// Force down dynamic rendering path if invalid lightmap settings, so we can apply an error material in DrawRichMesh
-// 		(HasStaticLighting() && !HasValidSettingsForStaticLighting()) ||
-// 		HasViewDependentDPG()
-// #if WITH_EDITOR
-// 		//only check these in the editor
-// 		|| IsSelected() || IsHovered()
-// #endif
-// 		)
-// 	{
-//		Result.bDynamicRelevance = true;
-// 	}
-// 	else
-// 	{
-// 		Result.bStaticRelevance = true;
-// 	}
-
 	if (!EngineShowFlags.Materials
 #if SUPPORT_EXTRA_RENDERING
 		|| bInCollisionView
diff --git a/Engine/Plugins/2D/Paper2D/Source/Paper2D/Private/PaperRenderSceneProxy.h b/Engine/Plugins/2D/Paper2D/Source/Paper2D/Private/PaperRenderSceneProxy.h
--- a/Engine/Plugins/2D/Paper2D/Source/Paper2D/Private/PaperRenderSceneProxy.h
+++ b/Engine/Plugins/2D/Paper2D/Source/Paper2D/Private/PaperRenderSceneProxy.h
@@ -90,7 +90,6 @@ public:
 
 	// FPrimitiveSceneProxy interface.
 	virtual void GetDynamicMeshElements(const TArray<const FSceneView*>& Views, const FSceneViewFamily& ViewFamily, uint32 VisibilityMap, FMeshElementCollector& Collector) const override;
-	virtual void DrawDynamicElements(FPrimitiveDrawInterface* PDI, const FSceneView* View) override;
 	virtual FPrimitiveViewRelevance GetViewRelevance(const FSceneView* View) override;
 	virtual void OnTransformChanged() override;
 	virtual uint32 GetMemoryFootprint() const override;
@@ -106,12 +105,8 @@ public:
 	}
 
 protected:
-	virtual void DrawDynamicElements_RichMesh(FPrimitiveDrawInterface* PDI, const FSceneView* View, bool bUseOverrideColor, const FLinearColor& OverrideColor);
 	virtual void GetDynamicMeshElementsForView(const FSceneView* View, int32 ViewIndex, bool bUseOverrideColor, const FLinearColor& OverrideColor, FMeshElementCollector& Collector) const;
 
-	void DrawBatch(FPrimitiveDrawInterface* PDI, const FSceneView* View, bool bUseOverrideColor, const FLinearColor& OverrideColor, class UMaterialInterface* BatchMaterial, const TArray<FSpriteDrawCallRecord>& Batch);
-	void DrawNewBatches(FPrimitiveDrawInterface* PDI, const FSceneView* View, bool bUseOverrideColor, const FLinearColor& OverrideColor);
-
 	void GetBatchMesh(const FSceneView* View, bool bUseOverrideColor, const FLinearColor& OverrideColor, class UMaterialInterface* BatchMaterial, const TArray<FSpriteDrawCallRecord>& Batch, int32 ViewIndex, FMeshElementCollector& Collector) const;
 	void GetNewBatchMeshes(const FSceneView* View, bool bUseOverrideColor, const FLinearColor& OverrideColor, int32 ViewIndex, FMeshElementCollector& Collector) const;
 
diff --git a/Engine/Plugins/2D/Paper2D/Source/Paper2D/Private/PaperSpriteSceneProxy.cpp b/Engine/Plugins/2D/Paper2D/Source/Paper2D/Private/PaperSpriteSceneProxy.cpp
--- a/Engine/Plugins/2D/Paper2D/Source/Paper2D/Private/PaperSpriteSceneProxy.cpp
+++ b/Engine/Plugins/2D/Paper2D/Source/Paper2D/Private/PaperSpriteSceneProxy.cpp
@@ -20,62 +20,6 @@ FPaperSpriteSceneProxy::FPaperSpriteSceneProxy(const UPaperSpriteComponent* InCo
 	SourceSprite = InComponent->SourceSprite; //@TODO: This is totally not threadsafe, and won't keep up to date if the actor's sprite changes, etc....
 }
 
-void FPaperSpriteSceneProxy::DrawDynamicElements(FPrimitiveDrawInterface* PDI, const FSceneView* View)
-{
-	if (SourceSprite != nullptr)
-	{
-		// Show 3D physics
-		if ((View->Family->EngineShowFlags.Collision /*@TODO: && bIsCollisionEnabled*/) && AllowDebugViewmodes())
-		{
-			if (UBodySetup2D* BodySetup2D = Cast<UBodySetup2D>(SourceSprite->BodySetup))
-			{
-				//@TODO: Draw 2D debugging geometry
-			}
-			else if (UBodySetup* BodySetup = SourceSprite->BodySetup)
-			{
-				if (FMath::Abs(GetLocalToWorld().Determinant()) < SMALL_NUMBER)
-				{
-					// Catch this here or otherwise GeomTransform below will assert
-					// This spams so commented out
-					//UE_LOG(LogStaticMesh, Log, TEXT("Zero scaling not supported (%s)"), *StaticMesh->GetPathName());
-				}
-				else
-				{
-					// Make a material for drawing solid collision stuff
-					const UMaterial* LevelColorationMaterial = View->Family->EngineShowFlags.Lighting 
-						? GEngine->ShadedLevelColorationLitMaterial : GEngine->ShadedLevelColorationUnlitMaterial;
-
-					const FColoredMaterialRenderProxy CollisionMaterialInstance(
-						LevelColorationMaterial->GetRenderProxy(IsSelected(), IsHovered()),
-						WireframeColor
-						);
-
-					// Draw the static mesh's body setup.
-
-					// Get transform without scaling.
-					FTransform GeomTransform(GetLocalToWorld());
-
-					// In old wireframe collision mode, always draw the wireframe highlighted (selected or not).
-					bool bDrawWireSelected = IsSelected();
-					if (View->Family->EngineShowFlags.Collision)
-					{
-						bDrawWireSelected = true;
-					}
-
-					// Differentiate the color based on bBlockNonZeroExtent.  Helps greatly with skimming a level for optimization opportunities.
-					const FColor CollisionColor = FColor(220,149,223,255);
-
-					const bool bUseSeparateColorPerHull = (Owner == NULL);
-					const bool bDrawSolid = false;
-					BodySetup->AggGeom.DrawAggGeom(PDI, GeomTransform, GetSelectionColor(CollisionColor, bDrawWireSelected, IsHovered()), &CollisionMaterialInstance, bUseSeparateColorPerHull, bDrawSolid, false);
-				}
-			}
-		}
-	}
-
-	FPaperRenderSceneProxy::DrawDynamicElements(PDI, View);
-}
-
 void FPaperSpriteSceneProxy::GetDynamicMeshElements(const TArray<const FSceneView*>& Views, const FSceneViewFamily& ViewFamily, uint32 VisibilityMap, FMeshElementCollector& Collector) const
 {
 	if (SourceSprite != nullptr)
@@ -167,19 +111,6 @@ void FPaperSpriteSceneProxy::SetSprite_RenderThread(const FSpriteDrawCallRecord&
 	}
 }
 
-void FPaperSpriteSceneProxy::DrawDynamicElements_RichMesh(FPrimitiveDrawInterface* PDI, const FSceneView* View, bool bUseOverrideColor, const FLinearColor& OverrideColor)
-{
-	if (Material != nullptr)
-	{
-		DrawBatch(PDI, View, bUseOverrideColor, OverrideColor, Material, BatchedSprites);
-	}
-
-	if ((AlternateMaterial != nullptr) && (AlternateBatchedSprites.Num() > 0))
-	{
-		DrawBatch(PDI, View, bUseOverrideColor, OverrideColor, AlternateMaterial, AlternateBatchedSprites);
-	}
-}
-
 void FPaperSpriteSceneProxy::GetDynamicMeshElementsForView(const FSceneView* View, int32 ViewIndex, bool bUseOverrideColor, const FLinearColor& OverrideColor, FMeshElementCollector& Collector) const
 {
 	if (Material != nullptr)
diff --git a/Engine/Plugins/2D/Paper2D/Source/Paper2D/Private/PaperSpriteSceneProxy.h b/Engine/Plugins/2D/Paper2D/Source/Paper2D/Private/PaperSpriteSceneProxy.h
--- a/Engine/Plugins/2D/Paper2D/Source/Paper2D/Private/PaperSpriteSceneProxy.h
+++ b/Engine/Plugins/2D/Paper2D/Source/Paper2D/Private/PaperSpriteSceneProxy.h
@@ -14,7 +14,6 @@ public:
 
 	// FPrimitiveSceneProxy interface
 	virtual void GetDynamicMeshElements(const TArray<const FSceneView*>& Views, const FSceneViewFamily& ViewFamily, uint32 VisibilityMap, FMeshElementCollector& Collector) const override;
-	virtual void DrawDynamicElements(FPrimitiveDrawInterface* PDI, const FSceneView* View) override;
 	// End of FPrimitiveSceneProxy interface
 
 	void SetSprite_RenderThread(const FSpriteDrawCallRecord& NewDynamicData, int32 SplitIndex);
@@ -22,7 +21,6 @@ public:
 protected:
 
 	// FPaperRenderSceneProxy interface
-	virtual void DrawDynamicElements_RichMesh(FPrimitiveDrawInterface* PDI, const FSceneView* View, bool bUseOverrideColor, const FLinearColor& OverrideColor) override;
 	virtual void GetDynamicMeshElementsForView(const FSceneView* View, int32 ViewIndex, bool bUseOverrideColor, const FLinearColor& OverrideColor, FMeshElementCollector& Collector) const override;
 	// End of FPaperRenderSceneProxy interface
 
diff --git a/Engine/Plugins/2D/Paper2D/Source/Paper2D/Private/PaperTileMapRenderSceneProxy.cpp b/Engine/Plugins/2D/Paper2D/Source/Paper2D/Private/PaperTileMapRenderSceneProxy.cpp
--- a/Engine/Plugins/2D/Paper2D/Source/Paper2D/Private/PaperTileMapRenderSceneProxy.cpp
+++ b/Engine/Plugins/2D/Paper2D/Source/Paper2D/Private/PaperTileMapRenderSceneProxy.cpp
@@ -160,127 +160,3 @@ void FPaperTileMapRenderSceneProxy::GetDynamicMeshElements(const TArray<const FS
 	FPaperRenderSceneProxy::GetDynamicMeshElements(Views, ViewFamily, VisibilityMap, Collector);
 }
 
-
-void FPaperTileMapRenderSceneProxy::DrawDynamicElements(FPrimitiveDrawInterface* PDI, const FSceneView* View)
-{
-	// Draw the tile maps
-	//@TODO: RenderThread race condition
-	if (TileMap != NULL)
-	{
-		FColor WireframeColor = FColor(0, 255, 255, 255);
-
-		if ((View->Family->EngineShowFlags.Collision /*@TODO: && bIsCollisionEnabled*/) && AllowDebugViewmodes())
-		{
-			if (UBodySetup2D* BodySetup2D = Cast<UBodySetup2D>(TileMap->BodySetup))
-			{
-				//@TODO: Draw 2D debugging geometry
-			}
-			else if (UBodySetup* BodySetup = TileMap->BodySetup)
-			{
-				if (FMath::Abs(GetLocalToWorld().Determinant()) < SMALL_NUMBER)
-				{
-					// Catch this here or otherwise GeomTransform below will assert
-					// This spams so commented out
-					//UE_LOG(LogStaticMesh, Log, TEXT("Zero scaling not supported (%s)"), *StaticMesh->GetPathName());
-				}
-				else
-				{
-					// Make a material for drawing solid collision stuff
-					const UMaterial* LevelColorationMaterial = View->Family->EngineShowFlags.Lighting
-						? GEngine->ShadedLevelColorationLitMaterial : GEngine->ShadedLevelColorationUnlitMaterial;
-
-					const FColoredMaterialRenderProxy CollisionMaterialInstance(
-						LevelColorationMaterial->GetRenderProxy(IsSelected(), IsHovered()),
-						WireframeColor
-						);
-
-					// Draw the static mesh's body setup.
-
-					// Get transform without scaling.
-					FTransform GeomTransform(GetLocalToWorld());
-
-					// In old wireframe collision mode, always draw the wireframe highlighted (selected or not).
-					bool bDrawWireSelected = IsSelected();
-					if (View->Family->EngineShowFlags.Collision)
-					{
-						bDrawWireSelected = true;
-					}
-
-					// Differentiate the color based on bBlockNonZeroExtent.  Helps greatly with skimming a level for optimization opportunities.
-					FColor CollisionColor = FColor(157, 149, 223, 255);
-
-					const bool bPerHullColor = false;
-					const bool bDrawSimpleSolid = false;
-					BodySetup->AggGeom.DrawAggGeom(PDI, GeomTransform, GetSelectionColor(CollisionColor, bDrawWireSelected, IsHovered()), &CollisionMaterialInstance, bPerHullColor, bDrawSimpleSolid, UseEditorDepthTest());
-				}
-			}
-		}
-
-		// Draw the bounds
-		RenderBounds(PDI, View->Family->EngineShowFlags, GetBounds(), IsSelected());
-
-#if WITH_EDITOR
-		// Draw the debug outline
-		if (View->Family->EngineShowFlags.Grid)
-		{
-			const uint8 DPG = SDPG_Foreground;//GetDepthPriorityGroup(View);
-
-			// Draw separation wires if selected
-			FLinearColor OverrideColor;
-			bool bUseOverrideColor = false;
-
-			const bool bShowAsSelected = !(GIsEditor && View->Family->EngineShowFlags.Selection) || IsSelected();
-			if (bShowAsSelected || IsHovered())
-			{
-				bUseOverrideColor = true;
-				OverrideColor = GetSelectionColor(FLinearColor::White, bShowAsSelected, IsHovered());
-			}
-
-			FTransform LocalToWorld(GetLocalToWorld());
-
-			const float TW = TileMap->TileWidth;
-			const float TH = TileMap->TileHeight;
-			//@TODO: PAPER: Tiles probably shouldn't be drawn with their pivot at the center - RE: all the 0.5f in here
-
-			if (bUseOverrideColor)
-			{
-				// Draw horizontal lines
-				for (int32 Y = 0; Y <= TileMap->MapHeight; ++Y)
-				{
-					int32 X = 0;
-					const FVector Start((X - 0.5f) * TW, 0.0f, -(Y - 0.5f) * TH);
-
-					X = TileMap->MapWidth;
-					const FVector End((X - 0.5f) * TW, 0.0f, -(Y - 0.5f) * TH);
-
-					PDI->DrawLine(LocalToWorld.TransformPosition(Start), LocalToWorld.TransformPosition(End), OverrideColor, DPG);
-				}
-
-				// Draw vertical lines
-				for (int32 X = 0; X <= TileMap->MapWidth; ++X)
-				{
-					int32 Y = 0;
-					const FVector Start((X - 0.5f) * TW, 0.0f, -(Y - 0.5f) * TH);
-
-					Y = TileMap->MapHeight;
-					const FVector End((X - 0.5f) * TW, 0.0f, -(Y - 0.5f) * TH);
-
-					PDI->DrawLine(LocalToWorld.TransformPosition(Start), LocalToWorld.TransformPosition(End), OverrideColor, DPG);
-				}
-			}
-
-
-			// Create a local space bounding box
-			const FVector TopLeft(-TW*0.5f, -2.0f, TH*0.5f);
-			const FVector Dimensions(TileMap->MapWidth * TW, 4.0f, -TileMap->MapHeight * TH);
-			const FBox OutlineBox(TopLeft, TopLeft + Dimensions);
-
-			// Draw it		
-			DrawWireBox(PDI, OutlineBox.TransformBy(LocalToWorld), FLinearColor::White, DPG); //@TODO: Paper: Doesn't handle rotation well - need to use DrawOrientedWireBox
-		}
-#endif
-	}
-
-	// Draw all of the queued up sprites
-	FPaperRenderSceneProxy::DrawDynamicElements(PDI, View);
-}
diff --git a/Engine/Plugins/2D/Paper2D/Source/Paper2D/Private/PaperTileMapRenderSceneProxy.h b/Engine/Plugins/2D/Paper2D/Source/Paper2D/Private/PaperTileMapRenderSceneProxy.h
--- a/Engine/Plugins/2D/Paper2D/Source/Paper2D/Private/PaperTileMapRenderSceneProxy.h
+++ b/Engine/Plugins/2D/Paper2D/Source/Paper2D/Private/PaperTileMapRenderSceneProxy.h
@@ -14,7 +14,6 @@ public:
 
 	// FPrimitiveSceneProxy interface.
 	virtual void GetDynamicMeshElements(const TArray<const FSceneView*>& Views, const FSceneViewFamily& ViewFamily, uint32 VisibilityMap, FMeshElementCollector& Collector) const override;
-	virtual void DrawDynamicElements(FPrimitiveDrawInterface* PDI, const FSceneView* View) override;
 	// End of FPrimitiveSceneProxy interface.
 
 	void SetBatchesHack(TArray<FSpriteDrawCallRecord>& InBatchedSprites)
diff --git a/Engine/Plugins/2D/Paper2D/Source/Paper2D/Private/Rendering/PaperBatchSceneProxy.cpp b/Engine/Plugins/2D/Paper2D/Source/Paper2D/Private/Rendering/PaperBatchSceneProxy.cpp
--- a/Engine/Plugins/2D/Paper2D/Source/Paper2D/Private/Rendering/PaperBatchSceneProxy.cpp
+++ b/Engine/Plugins/2D/Paper2D/Source/Paper2D/Private/Rendering/PaperBatchSceneProxy.cpp
@@ -34,115 +34,6 @@ FPrimitiveViewRelevance FPaperBatchSceneProxy::GetViewRelevance(const FSceneView
 	return Result;
 }
 
-
-struct FProxyBatchPair
-{
-	void* Texture;
-	int32 ProxyIndex;
-	int32 BatchIndex;
-};
-
-struct FSortSpriteProxiesByMaterial
-{
-	FORCEINLINE bool operator()(const FPaperRenderSceneProxy& A, const FPaperRenderSceneProxy& B) const
-	{
-		return A.GetMaterial() < B.GetMaterial();
-	}
-};
-
-
-struct FSortSpriteBatchesByTexture
-{
-	FORCEINLINE bool operator()(const FProxyBatchPair& A, const FProxyBatchPair& B) const
-	{
-		return (A.Texture == B.Texture) ? ((A.ProxyIndex == B.ProxyIndex) ? (A.BatchIndex < B.BatchIndex) : (A.ProxyIndex < B.ProxyIndex)) : (A.Texture < B.Texture);
-	}
-};
-
-void FPaperBatchSceneProxy::DrawDynamicElements(FPrimitiveDrawInterface* PDI, const FSceneView* View)
-{
-#if 0
-	// Not thread safe right now (not to mention use of GWorld...)
-
-	// Draw all the debug rendering for the 2D physics scene
-	if ((View->Family->EngineShowFlags.Collision /*@TODO: && bIsCollisionEnabled*/) && AllowDebugViewmodes())
-	{
-		FPhysicsIntegration2D::DrawDebugPhysics(GWorld, PDI, View); //@TODO: GWorld is worse than the disease
-	}
-#endif
-
-
-#if TEST_BATCHING
-	// First sort by material
-	ManagedProxies.Sort(FSortSpriteProxiesByMaterial());
-
-	// Figure out how many total sprites and verts we need to deal with
-	int32 NumSprites = 0;
-	int32 NumVerts = 0;
-	for (FPaperRenderSceneProxy* Proxy : ManagedProxies)
-	{
-		NumSprites += Proxy->BatchedSprites.Num();
-		for (FSpriteDrawCallRecord& Sprite : Proxy->BatchedSprites)
-		{
-			NumVerts += Sprite.RenderVerts.Num();
-		}
-	}
-
-
-	// Aggregate the individual draw calls
-	TArray<FProxyBatchPair> SpriteDrawCalls;
-	SpriteDrawCalls.Empty(NumSprites);
-
-	for (int32 ProxyIndex = 0; ProxyIndex < ManagedProxies.Num(); ++ProxyIndex)
-	{
-		FPaperRenderSceneProxy* Proxy = ManagedProxies[ProxyIndex];
-		for (int32 BatchIndex = 0; BatchIndex < Proxy->BatchedSprites.Num(); ++BatchIndex)
-		{
-			FProxyBatchPair& Pair = *new (SpriteDrawCalls) FProxyBatchPair();
-			Pair.Texture = Proxy->BatchedSprites[BatchIndex].Texture;//@TODO: Don't need to stick this here, but all of this will change when the batcher handles per-item batches too
-			Pair.ProxyIndex = ProxyIndex;
-			Pair.BatchIndex = BatchIndex;
-		}
-	}
-
-	// Sort by texture
-	SpriteDrawCalls.Sort(FSortSpriteBatchesByTexture());
-
-	// Ready to render in a nice batched fashion
-	for (FProxyBatchPair& Pair : SpriteDrawCalls)
-	{
-		ManagedProxies[Pair.ProxyIndex]->DrawDynamicElements_RichMesh(PDI, View, false, FLinearColor::White);
-	}
-
-
-// 	for (int32 PrimitiveIndex = 0, Num = View.VisibleDynamicPrimitives.Num(); PrimitiveIndex < Num; PrimitiveIndex++)
-// 	{
-// 		const FPrimitiveSceneInfo* PrimitiveSceneInfo = View.VisibleDynamicPrimitives[PrimitiveIndex];
-// 		int32 PrimitiveId = PrimitiveSceneInfo->GetIndex();
-// 		const FPrimitiveViewRelevance& PrimitiveViewRelevance = View.PrimitiveViewRelevanceMap[PrimitiveId];
-// 
-// 		const bool bVisible = View.PrimitiveVisibilityMap[PrimitiveId];
-// 
-// 		// Only draw the primitive if it's visible
-// 		if (bVisible &&
-// 			// only draw opaque and masked primitives if wireframe is disabled
-// 			(PrimitiveViewRelevance.bOpaqueRelevance || ViewFamily.EngineShowFlags.Wireframe) &&
-// 			PrimitiveViewRelevance.bRenderInMainPass
-// 			)
-// 
-// 		{
-// 			FScopeCycleCounter Context(PrimitiveSceneInfo->Proxy->GetStatId());
-// 			Drawer.SetPrimitive(PrimitiveSceneInfo->Proxy);
-// 			PrimitiveSceneInfo->Proxy->DrawDynamicElements(
-// 				&Drawer,
-// 				&View
-// 				);
-// 		}
-// 	}
-#endif
-	//for 
-}
-
 void FPaperBatchSceneProxy::RegisterManagedProxy(FPaperRenderSceneProxy* Proxy)
 {
 	check(IsInRenderingThread());
diff --git a/Engine/Plugins/2D/Paper2D/Source/Paper2D/Private/Rendering/PaperBatchSceneProxy.h b/Engine/Plugins/2D/Paper2D/Source/Paper2D/Private/Rendering/PaperBatchSceneProxy.h
--- a/Engine/Plugins/2D/Paper2D/Source/Paper2D/Private/Rendering/PaperBatchSceneProxy.h
+++ b/Engine/Plugins/2D/Paper2D/Source/Paper2D/Private/Rendering/PaperBatchSceneProxy.h
@@ -11,7 +11,6 @@ public:
 	FPaperBatchSceneProxy(const UPrimitiveComponent* InComponent);
 
  	// FPrimitiveSceneProxy interface.
- 	virtual void DrawDynamicElements(FPrimitiveDrawInterface* PDI, const FSceneView* View) override;
  	virtual FPrimitiveViewRelevance GetViewRelevance(const FSceneView* View) override;
  	virtual uint32 GetMemoryFootprint() const override;
  	virtual bool CanBeOccluded() const override;
diff --git a/Engine/Plugins/2D/Paper2D/Source/Paper2D/Private/Terrain/PaperTerrainComponent.cpp b/Engine/Plugins/2D/Paper2D/Source/Paper2D/Private/Terrain/PaperTerrainComponent.cpp
--- a/Engine/Plugins/2D/Paper2D/Source/Paper2D/Private/Terrain/PaperTerrainComponent.cpp
+++ b/Engine/Plugins/2D/Paper2D/Source/Paper2D/Private/Terrain/PaperTerrainComponent.cpp
@@ -28,15 +28,10 @@ class FPaperTerrainSceneProxy : public FPaperRenderSceneProxy
 public:
 	FPaperTerrainSceneProxy(const UPaperTerrainComponent* InComponent, const TArray<FPaperTerrainMaterialPair>& InDrawingData);
 
-	// FPrimitiveSceneProxy interface.
-	//virtual void DrawDynamicElements(FPrimitiveDrawInterface* PDI, const FSceneView* View) override;
-	// End of FPrimitiveSceneProxy interface.
-
 protected:
 	TArray<FPaperTerrainMaterialPair> DrawingData;
 protected:
 	// FPaperRenderSceneProxy interface
-	virtual void DrawDynamicElements_RichMesh(FPrimitiveDrawInterface* PDI, const FSceneView* View, bool bUseOverrideColor, const FLinearColor& OverrideColor) override;
 	virtual void GetDynamicMeshElementsForView(const FSceneView* View, int32 ViewIndex, bool bUseOverrideColor, const FLinearColor& OverrideColor, FMeshElementCollector& Collector) const override;
 	// End of FPaperRenderSceneProxy interface
 };
@@ -54,17 +49,6 @@ FPaperTerrainSceneProxy::FPaperTerrainSceneProxy(const UPaperTerrainComponent* I
 	}
 }
 
-void FPaperTerrainSceneProxy::DrawDynamicElements_RichMesh(FPrimitiveDrawInterface* PDI, const FSceneView* View, bool bUseOverrideColor, const FLinearColor& OverrideColor)
-{
-	for (const FPaperTerrainMaterialPair& Batch : DrawingData)
-	{
-		if (Batch.Material != nullptr)
-		{
-			DrawBatch(PDI, View, bUseOverrideColor, OverrideColor, Batch.Material, Batch.Records);
-		}
-	}
-}
-
 void FPaperTerrainSceneProxy::GetDynamicMeshElementsForView(const FSceneView* View, int32 ViewIndex, bool bUseOverrideColor, const FLinearColor& OverrideColor, FMeshElementCollector& Collector) const
 {
 	for (const FPaperTerrainMaterialPair& Batch : DrawingData)
diff --git a/Engine/Plugins/Runtime/CableComponent/Source/CableComponent/Private/CableComponent.cpp b/Engine/Plugins/Runtime/CableComponent/Source/CableComponent/Private/CableComponent.cpp
--- a/Engine/Plugins/Runtime/CableComponent/Source/CableComponent/Private/CableComponent.cpp
+++ b/Engine/Plugins/Runtime/CableComponent/Source/CableComponent/Private/CableComponent.cpp
@@ -296,50 +296,6 @@ public:
 		}
 	}
 
-	virtual void DrawDynamicElements(FPrimitiveDrawInterface* PDI,const FSceneView* View)
-	{
-		QUICK_SCOPE_CYCLE_COUNTER( STAT_CableSceneProxy_DrawDynamicElements );
-
-		const bool bWireframe = AllowDebugViewmodes() && View->Family->EngineShowFlags.Wireframe;
-
-		FColoredMaterialRenderProxy WireframeMaterialInstance(
-			GEngine->WireframeMaterial ? GEngine->WireframeMaterial->GetRenderProxy(IsSelected()) : NULL,
-			FLinearColor(0, 0.5f, 1.f)
-			);
-
-		FMaterialRenderProxy* MaterialProxy = NULL;
-		if(bWireframe)
-		{
-			MaterialProxy = &WireframeMaterialInstance;
-		}
-		else
-		{
-			MaterialProxy = Material->GetRenderProxy(IsSelected());
-		}
-
-		// Draw the mesh.
-		FMeshBatch Mesh;
-		FMeshBatchElement& BatchElement = Mesh.Elements[0];
-		BatchElement.IndexBuffer = &IndexBuffer;
-		Mesh.bWireframe = bWireframe;
-		Mesh.VertexFactory = &VertexFactory;
-		Mesh.MaterialRenderProxy = MaterialProxy;
-		BatchElement.PrimitiveUniformBuffer = CreatePrimitiveUniformBufferImmediate(GetLocalToWorld(), GetBounds(), GetLocalBounds(), true, UseEditorDepthTest());
-		BatchElement.FirstIndex = 0;
-		BatchElement.NumPrimitives = GetRequiredIndexCount()/3;
-		BatchElement.MinVertexIndex = 0;
-		BatchElement.MaxVertexIndex = GetRequiredVertexCount();
-		Mesh.ReverseCulling = IsLocalToWorldDeterminantNegative();
-		Mesh.Type = PT_TriangleList;
-		Mesh.DepthPriorityGroup = SDPG_World;
-		PDI->DrawMesh(Mesh);
-
-#if !(UE_BUILD_SHIPPING || UE_BUILD_TEST)
-		// Render bounds
-		RenderBounds(PDI, View->Family->EngineShowFlags, GetBounds(), IsSelected());
-#endif
-	}
-
 	virtual FPrimitiveViewRelevance GetViewRelevance(const FSceneView* View)
 	{
 		FPrimitiveViewRelevance Result;
diff --git a/Engine/Plugins/Runtime/CustomMeshComponent/Source/CustomMeshComponent/Private/CustomMeshComponent.cpp b/Engine/Plugins/Runtime/CustomMeshComponent/Source/CustomMeshComponent/Private/CustomMeshComponent.cpp
--- a/Engine/Plugins/Runtime/CustomMeshComponent/Source/CustomMeshComponent/Private/CustomMeshComponent.cpp
+++ b/Engine/Plugins/Runtime/CustomMeshComponent/Source/CustomMeshComponent/Private/CustomMeshComponent.cpp
@@ -191,45 +191,6 @@ public:
 		}
 	}
 
-	virtual void DrawDynamicElements(FPrimitiveDrawInterface* PDI,const FSceneView* View)
-	{
-		QUICK_SCOPE_CYCLE_COUNTER( STAT_CustomMeshSceneProxy_DrawDynamicElements );
-
-		const bool bWireframe = AllowDebugViewmodes() && View->Family->EngineShowFlags.Wireframe;
-
-		FColoredMaterialRenderProxy WireframeMaterialInstance(
-			GEngine->WireframeMaterial ? GEngine->WireframeMaterial->GetRenderProxy(IsSelected()) : NULL,
-			FLinearColor(0, 0.5f, 1.f)
-			);
-
-		FMaterialRenderProxy* MaterialProxy = NULL;
-		if(bWireframe)
-		{
-			MaterialProxy = &WireframeMaterialInstance;
-		}
-		else
-		{
-			MaterialProxy = Material->GetRenderProxy(IsSelected());
-		}
-
-		// Draw the mesh.
-		FMeshBatch Mesh;
-		FMeshBatchElement& BatchElement = Mesh.Elements[0];
-		BatchElement.IndexBuffer = &IndexBuffer;
-		Mesh.bWireframe = bWireframe;
-		Mesh.VertexFactory = &VertexFactory;
-		Mesh.MaterialRenderProxy = MaterialProxy;
-		BatchElement.PrimitiveUniformBuffer = CreatePrimitiveUniformBufferImmediate(GetLocalToWorld(), GetBounds(), GetLocalBounds(), true, UseEditorDepthTest());
-		BatchElement.FirstIndex = 0;
-		BatchElement.NumPrimitives = IndexBuffer.Indices.Num() / 3;
-		BatchElement.MinVertexIndex = 0;
-		BatchElement.MaxVertexIndex = VertexBuffer.Vertices.Num() - 1;
-		Mesh.ReverseCulling = IsLocalToWorldDeterminantNegative();
-		Mesh.Type = PT_TriangleList;
-		Mesh.DepthPriorityGroup = SDPG_World;
-		PDI->DrawMesh(Mesh);
-	}
-
 	virtual FPrimitiveViewRelevance GetViewRelevance(const FSceneView* View)
 	{
 		FPrimitiveViewRelevance Result;
diff --git a/Engine/Source/Developer/FunctionalTesting/Private/FuncTestRenderingComponent.cpp b/Engine/Source/Developer/FunctionalTesting/Private/FuncTestRenderingComponent.cpp
--- a/Engine/Source/Developer/FunctionalTesting/Private/FuncTestRenderingComponent.cpp
+++ b/Engine/Source/Developer/FunctionalTesting/Private/FuncTestRenderingComponent.cpp
@@ -16,10 +16,6 @@ public:
 	{
 	}
 
-	virtual void DrawDynamicElements(FPrimitiveDrawInterface* PDI,const FSceneView* View)
-	{
-	}
-
 	virtual FPrimitiveViewRelevance GetViewRelevance(const FSceneView* View)
 	{
 		FPrimitiveViewRelevance Result;
diff --git a/Engine/Source/Developer/GameplayDebugger/Private/GameplayDebuggingComponent.cpp b/Engine/Source/Developer/GameplayDebugger/Private/GameplayDebuggingComponent.cpp
--- a/Engine/Source/Developer/GameplayDebugger/Private/GameplayDebuggingComponent.cpp
+++ b/Engine/Source/Developer/GameplayDebugger/Private/GameplayDebuggingComponent.cpp
@@ -62,14 +62,6 @@ public:
 		}
 	}
 
-	virtual void DrawDynamicElements(FPrimitiveDrawInterface* PDI, const FSceneView* View) override
-	{
-		for (int32 Index = 0; Index < ChildProxies.Num(); ++Index)
-		{
-			ChildProxies[Index]->DrawDynamicElements(PDI, View);
-		}
-	}
-
 	virtual void RegisterDebugDrawDelgate() override
 	{
 		for (int32 Index = 0; Index < ChildProxies.Num(); ++Index)
diff --git a/Engine/Source/Editor/LandscapeEditor/Private/LandscapeEdModeRampTool.cpp b/Engine/Source/Editor/LandscapeEditor/Private/LandscapeEdModeRampTool.cpp
--- a/Engine/Source/Editor/LandscapeEditor/Private/LandscapeEdModeRampTool.cpp
+++ b/Engine/Source/Editor/LandscapeEditor/Private/LandscapeEdModeRampTool.cpp
@@ -310,7 +310,7 @@ public:
 			const ALandscapeProxy* LandscapeProxy = EdMode->CurrentToolTarget.LandscapeInfo->GetLandscapeProxy();
 			const FTransform LandscapeToWorld = LandscapeProxy->LandscapeActorToWorld();
 
-			const FLinearColor SelectedSpriteColor = FLinearColor::White + (GEngine->GetSelectedMaterialColor() * GEngine->SelectionHighlightIntensity * 10); // copied from FSpriteSceneProxy::DrawDynamicElements()
+			const FLinearColor SelectedSpriteColor = FLinearColor::White + (GEngine->GetSelectedMaterialColor() * GEngine->SelectionHighlightIntensity * 10);
 
 			FVector WorldPoints[2];
 			for (int32 i = 0; i < NumPoints; i++)
diff --git a/Engine/Source/Editor/UnrealEd/Private/Animation/DebugSkelMeshComponent.cpp b/Engine/Source/Editor/UnrealEd/Private/Animation/DebugSkelMeshComponent.cpp
--- a/Engine/Source/Editor/UnrealEd/Private/Animation/DebugSkelMeshComponent.cpp
+++ b/Engine/Source/Editor/UnrealEd/Private/Animation/DebugSkelMeshComponent.cpp
@@ -60,33 +60,6 @@ public:
 		}
 	}
 
-	/** 
-	* Draw the scene proxy as a dynamic element
-	*
-	* @param	PDI - draw interface to render to
-	* @param	View - current view
-	*/
-	virtual void DrawDynamicElements(FPrimitiveDrawInterface* PDI,const FSceneView* View)
-	{
-		// We don't want to draw the mesh geometry to the hit testing render target
-		// so that we can get to triangle strips that are partially obscured by other
-		// triangle strips easier.
-		if (!PDI->IsHitTesting())
-		{
-			if (SkeletalMeshComponent->bDrawMesh)
-			{
-				// Draw the mesh
-				FSkeletalMeshSceneProxy::DrawDynamicElements(PDI, View);
-			}
-			
-		}
-
-		if( SkeletalMeshComponent->MeshObject && (SkeletalMeshComponent->bDrawNormals || SkeletalMeshComponent->bDrawTangents || SkeletalMeshComponent->bDrawBinormals) )
-		{
-			SkeletalMeshComponent->MeshObject->DrawVertexElements(PDI, SkeletalMeshComponent->ComponentToWorld, SkeletalMeshComponent->bDrawNormals, SkeletalMeshComponent->bDrawTangents, SkeletalMeshComponent->bDrawBinormals);
-		}
-	}
-
 	uint32 GetAllocatedSize() const
 	{
 		return FSkeletalMeshSceneProxy::GetAllocatedSize();
diff --git a/Engine/Source/Runtime/Engine/Classes/Components/NiagaraComponent.h b/Engine/Source/Runtime/Engine/Classes/Components/NiagaraComponent.h
--- a/Engine/Source/Runtime/Engine/Classes/Components/NiagaraComponent.h
+++ b/Engine/Source/Runtime/Engine/Classes/Components/NiagaraComponent.h
@@ -153,8 +153,6 @@ private:
 	virtual void OnActorPositionChanged() override;
 	virtual void OnTransformChanged() override;
 
-	virtual void PreRenderView(const FSceneViewFamily* ViewFamily, const uint32 VisibilityMap, int32 FrameNumber) override;
-	virtual void DrawDynamicElements(FPrimitiveDrawInterface* PDI, const FSceneView* View) override;
 	virtual void GetDynamicMeshElements(const TArray<const FSceneView*>& Views, const FSceneViewFamily& ViewFamily, uint32 VisibilityMap, FMeshElementCollector& Collector) const override;
 
 	virtual FPrimitiveViewRelevance GetViewRelevance(const FSceneView* View)  override;
diff --git a/Engine/Source/Runtime/Engine/Classes/Engine/Engine.h b/Engine/Source/Runtime/Engine/Classes/Engine/Engine.h
--- a/Engine/Source/Runtime/Engine/Classes/Engine/Engine.h
+++ b/Engine/Source/Runtime/Engine/Classes/Engine/Engine.h
@@ -1693,8 +1693,6 @@ public:
 	bool HandleTestslateGameUICommand( const TCHAR* Cmd, FOutputDevice& Ar );
 	bool HandleDirCommand( const TCHAR* Cmd, FOutputDevice& Ar );
 	bool HandleTrackParticleRenderingStatsCommand( const TCHAR* Cmd, FOutputDevice& Ar );
-	bool HandleDumpParticleRenderingStatsCommand( const TCHAR* Cmd, FOutputDevice& Ar );
-	bool HandleDumpParticleFrameRenderingStatsCommand( const TCHAR* Cmd, FOutputDevice& Ar );
 	bool HandleDumpAllocatorStats( const TCHAR* Cmd, FOutputDevice& Ar );
 	bool HandleHeapCheckCommand( const TCHAR* Cmd, FOutputDevice& Ar );
 	bool HandleToggleOnscreenDebugMessageDisplayCommand( const TCHAR* Cmd, FOutputDevice& Ar );
diff --git a/Engine/Source/Runtime/Engine/Classes/Engine/NiagaraEffectRenderer.h b/Engine/Source/Runtime/Engine/Classes/Engine/NiagaraEffectRenderer.h
--- a/Engine/Source/Runtime/Engine/Classes/Engine/NiagaraEffectRenderer.h
+++ b/Engine/Source/Runtime/Engine/Classes/Engine/NiagaraEffectRenderer.h
@@ -33,9 +33,7 @@ class NiagaraEffectRenderer
 public:
 	virtual void GetDynamicMeshElements(const TArray<const FSceneView*>& Views, const FSceneViewFamily& ViewFamily, uint32 VisibilityMap, FMeshElementCollector& Collector, const FNiagaraSceneProxy *SceneProxy) const = 0;
 
-	virtual void PreRenderView(const FSceneViewFamily* ViewFamily, const uint32 VisibilityMap, int32 FrameNumber, const FNiagaraSceneProxy *SceneProxy) = 0;
 	virtual void SetDynamicData_RenderThread(FNiagaraDynamicDataBase* NewDynamicData) = 0;
-	virtual void DrawDynamicElements(FPrimitiveDrawInterface* PDI, const FSceneView* View, const FNiagaraSceneProxy *SceneProxy) = 0;
 	virtual void CreateRenderThreadResources() = 0;
 	virtual void ReleaseRenderThreadResources() = 0;
 	virtual FNiagaraDynamicDataBase *GenerateVertexData(const FNiagaraEmitterParticleData &Data) = 0;
@@ -51,7 +49,6 @@ public:
 		Result.bDrawRelevance = bHasDynamicData && SceneProxy->IsShown(View) && View->Family->EngineShowFlags.Particles;
 		Result.bShadowRelevance = bHasDynamicData && SceneProxy->IsShadowCast(View);
 		Result.bDynamicRelevance = bHasDynamicData;
-		Result.bNeedsPreRenderView = Result.bDrawRelevance || Result.bShadowRelevance;
 		if (bHasDynamicData && View->Family->EngineShowFlags.Bounds)
 		{
 			Result.bOpaqueRelevance = true;
@@ -108,8 +105,6 @@ public:
 	// FPrimitiveSceneProxy interface.
 	virtual void CreateRenderThreadResources() override;
 
-	virtual void PreRenderView(const FSceneViewFamily* ViewFamily, const uint32 VisibilityMap, int32 FrameNumber, const FNiagaraSceneProxy *SceneProxy) override;
-	virtual void DrawDynamicElements(FPrimitiveDrawInterface* PDI, const FSceneView* View, const FNiagaraSceneProxy *SceneProxy) override;
 	virtual void GetDynamicMeshElements(const TArray<const FSceneView*>& Views, const FSceneViewFamily& ViewFamily, uint32 VisibilityMap, FMeshElementCollector& Collector, const FNiagaraSceneProxy *SceneProxy) const override;
 
 
@@ -125,12 +120,9 @@ public:
 	bool HasDynamicData();
 
 private:
-	struct FNiagaraDynamicDataSprites *DynamicDataRender;
+	struct FNiagaraDynamicDataSprites* DynamicDataRender;
 	mutable TUniformBuffer<FPrimitiveUniformShaderParameters> WorldSpacePrimitiveUniformBuffer;
-	class FParticleSpriteVertexFactory *VertexFactory;
-	//FParticleSpriteUniformParameters UniformParameters;
-	TArray<FParticleSpriteUniformBufferRef, TInlineAllocator<2> > PerViewUniformBuffers;
-	FGlobalDynamicVertexBuffer::FAllocation DynamicVertexAllocation;
+	class FParticleSpriteVertexFactory* VertexFactory;
 };
 
 
@@ -161,9 +153,6 @@ public:
 
 	virtual void GetDynamicMeshElements(const TArray<const FSceneView*>& Views, const FSceneViewFamily& ViewFamily, uint32 VisibilityMap, FMeshElementCollector& Collector, const FNiagaraSceneProxy *SceneProxy) const override;
 
-	virtual void PreRenderView(const FSceneViewFamily* ViewFamily, const uint32 VisibilityMap, int32 FrameNumber, const FNiagaraSceneProxy *SceneProxy) override;
-	virtual void DrawDynamicElements(FPrimitiveDrawInterface* PDI, const FSceneView* View, const FNiagaraSceneProxy *SceneProxy) override;
-
 	/** Update render data buffer from attributes */
 	FNiagaraDynamicDataBase *GenerateVertexData(const FNiagaraEmitterParticleData &Data);
 
@@ -193,12 +182,9 @@ public:
 
 
 private:
-	struct FNiagaraDynamicDataRibbon *DynamicDataRender;
+	struct FNiagaraDynamicDataRibbon* DynamicDataRender;
 	mutable TUniformBuffer<FPrimitiveUniformShaderParameters> WorldSpacePrimitiveUniformBuffer;
-	class FParticleBeamTrailVertexFactory *VertexFactory;
-	//FParticleBeamTrailUniformParameters UniformParameters;
-	TArray<	FParticleBeamTrailUniformBufferRef, TInlineAllocator<2> > PerViewUniformBuffers;
-	FGlobalDynamicVertexBuffer::FAllocation DynamicVertexAllocation;
+	class FParticleBeamTrailVertexFactory* VertexFactory;
 };
 
 
diff --git a/Engine/Source/Runtime/Engine/Classes/PhysicsEngine/AggregateGeometry2D.h b/Engine/Source/Runtime/Engine/Classes/PhysicsEngine/AggregateGeometry2D.h
--- a/Engine/Source/Runtime/Engine/Classes/PhysicsEngine/AggregateGeometry2D.h
+++ b/Engine/Source/Runtime/Engine/Classes/PhysicsEngine/AggregateGeometry2D.h
@@ -129,8 +129,6 @@ public:
 public:
 	void EmptyElements();
 
-	void DrawAggGeom(class FPrimitiveDrawInterface* PDI, const FTransform& Transform, const FColor Color, const FMaterialRenderProxy* MatInst, bool bPerHullColor, bool bDrawSolid, bool bUseEditorDepthTest) const;
-
 	// Returns the area of all contained elements
  	float GetArea(const FVector& Scale) const;
 };
diff --git a/Engine/Source/Runtime/Engine/Classes/PhysicsEngine/BodySetup.h b/Engine/Source/Runtime/Engine/Classes/PhysicsEngine/BodySetup.h
--- a/Engine/Source/Runtime/Engine/Classes/PhysicsEngine/BodySetup.h
+++ b/Engine/Source/Runtime/Engine/Classes/PhysicsEngine/BodySetup.h
@@ -376,8 +376,6 @@ struct ENGINE_API FKAggregateGeom
 
 	void Serialize( const FArchive& Ar );
 
-	void DrawAggGeom(class FPrimitiveDrawInterface* PDI, const FTransform& Transform, const FColor Color, const FMaterialRenderProxy* MatInst, bool bPerHullColor, bool bDrawSolid, bool bUseEditorDepthTest);
-
 	void GetAggGeom(const FTransform& Transform, const FColor Color, const FMaterialRenderProxy* MatInst, bool bPerHullColor, bool bDrawSolid, bool bUseEditorDepthTest, int32 ViewIndex, class FMeshElementCollector& Collector) const;
 
 	/** Release the RenderInfo (if its there) and safely clean up any resources. Call on the game thread. */
diff --git a/Engine/Source/Runtime/Engine/Classes/PhysicsEngine/PhysicsAsset.h b/Engine/Source/Runtime/Engine/Classes/PhysicsEngine/PhysicsAsset.h
--- a/Engine/Source/Runtime/Engine/Classes/PhysicsEngine/PhysicsAsset.h
+++ b/Engine/Source/Runtime/Engine/Classes/PhysicsEngine/PhysicsAsset.h
@@ -148,7 +148,7 @@ public:
 #endif
 
 	// @todo document
-	void DrawCollision(class FPrimitiveDrawInterface* PDI, const USkeletalMesh* SkelMesh, const TArray<FTransform>& SpaceBases, const FTransform& LocalToWorld, float Scale);
+	void GetCollisionMesh(int32 ViewIndex, FMeshElementCollector& Collector, const USkeletalMesh* SkelMesh, const TArray<FTransform>& SpaceBases, const FTransform& LocalToWorld, float Scale);
 
 	// @todo document
 	void DrawConstraints(class FPrimitiveDrawInterface* PDI, const USkeletalMesh* SkelMesh, const TArray<FTransform>& SpaceBases, const FTransform& LocalToWorld, float Scale);
diff --git a/Engine/Source/Runtime/Engine/Private/AI/Navigation/NavLinkRenderingComponent.cpp b/Engine/Source/Runtime/Engine/Private/AI/Navigation/NavLinkRenderingComponent.cpp
--- a/Engine/Source/Runtime/Engine/Private/AI/Navigation/NavLinkRenderingComponent.cpp
+++ b/Engine/Source/Runtime/Engine/Private/AI/Navigation/NavLinkRenderingComponent.cpp
@@ -193,35 +193,6 @@ void FNavLinkRenderingProxy::GetDynamicMeshElements(const TArray<const FSceneVie
 	}
 }
 
-void FNavLinkRenderingProxy::DrawDynamicElements(FPrimitiveDrawInterface* PDI,const FSceneView* View)
-{
-	QUICK_SCOPE_CYCLE_COUNTER( STAT_NavLinkRenderingProxy_DrawDynamicElements );
-
-	if (LinkOwnerActor && LinkOwnerActor->GetWorld())
-	{
-		const UNavigationSystem* NavSys = LinkOwnerActor->GetWorld()->GetNavigationSystem();
-		TArray<float> StepHeights;
-		uint32 AgentMask = 0;
-		if (NavSys != NULL)
-		{
-			StepHeights.Reserve(NavSys->NavDataSet.Num());
-			for(int32 DataIndex = 0; DataIndex < NavSys->NavDataSet.Num(); ++DataIndex)
-			{
-				const ARecastNavMesh* NavMesh = Cast<const ARecastNavMesh>(NavSys->NavDataSet[DataIndex]);
-				AgentMask = NavMesh->bEnableDrawing ? AgentMask | (1 << DataIndex) : AgentMask;
-				if (NavMesh != NULL && NavMesh->AgentMaxStepHeight > 0 && NavMesh->bEnableDrawing)
-				{
-					StepHeights.Add(NavMesh->AgentMaxStepHeight);
-				}
-			}
-		}
-
-		static const FColor RadiusColor(150, 160, 150, 48);
-		FMaterialRenderProxy* const MeshColorInstance = new(FMemStack::Get()) FColoredMaterialRenderProxy(GEngine->DebugMeshMaterial->GetRenderProxy(false), RadiusColor);
-		FNavLinkRenderingProxy::DrawLinks(PDI, OffMeshPointLinks, OffMeshSegmentLinks, StepHeights, MeshColorInstance, AgentMask);
-	}
-}
-
 void FNavLinkRenderingProxy::GetLinkMeshes(const TArray<FNavLinkDrawing>& OffMeshPointLinks, const TArray<FNavLinkSegmentDrawing>& OffMeshSegmentLinks, TArray<float>& StepHeights, FMaterialRenderProxy* const MeshColorInstance, int32 ViewIndex, FMeshElementCollector& Collector, uint32 AgentMask)
 {
 	static const FColor LinkColor(0,0,166);
@@ -313,97 +284,6 @@ void FNavLinkRenderingProxy::GetLinkMeshes(const TArray<FNavLinkDrawing>& OffMes
 	}
 }
 
-void FNavLinkRenderingProxy::DrawLinks(FPrimitiveDrawInterface* PDI, TArray<FNavLinkDrawing>& OffMeshPointLinks, TArray<FNavLinkSegmentDrawing>& OffMeshSegmentLinks, TArray<float>& StepHeights, FMaterialRenderProxy* const MeshColorInstance, uint32 AgentMask)
-{
-	static const FColor LinkColor(0,0,166);
-	static const float LinkArcThickness = 3.f;
-	static const float LinkArcHeight = 0.4f;
-	
-	if (StepHeights.Num() == 0)
-	{
-		StepHeights.Add(FNavigationSystem::FallbackAgentHeight / 2);
-	}
-
-	for (int32 LinkIndex = 0; LinkIndex < OffMeshPointLinks.Num(); ++LinkIndex)
-	{
-		const FNavLinkDrawing& Link = OffMeshPointLinks[LinkIndex];
-		if ((Link.SupportedAgentsBits & AgentMask) == 0)
-		{
-			continue;
-		}
-
-		const uint32 Segments = FPlatformMath::Max<uint32>(LinkArcHeight*(Link.Right-Link.Left).Size()/10, 8);
-		DrawArc(PDI, Link.Left, Link.Right, LinkArcHeight, Segments, Link.Color, SDPG_World, 3.5f);
-		const FVector VOffset(0,0,FVector::Dist(Link.Left, Link.Right)*1.333f);
-
-		switch (Link.Direction)
-		{
-		case ENavLinkDirection::LeftToRight:
-			DrawArrowHead(PDI, Link.Right, Link.Left+VOffset, 30.f, Link.Color, SDPG_World, 3.5f);
-			break;
-		case ENavLinkDirection::RightToLeft:
-			DrawArrowHead(PDI, Link.Left, Link.Right+VOffset, 30.f, Link.Color, SDPG_World, 3.5f);
-			break;
-		case ENavLinkDirection::BothWays:
-		default:
-			DrawArrowHead(PDI, Link.Right, Link.Left+VOffset, 30.f, Link.Color, SDPG_World, 3.5f);
-			DrawArrowHead(PDI, Link.Left, Link.Right+VOffset, 30.f, Link.Color, SDPG_World, 3.5f);
-			break;
-		}
-
-		// draw snap-spheres on both ends
-		for (int32 StepHeightIndex = 0; StepHeightIndex < StepHeights.Num(); ++StepHeightIndex)
-		{
-			DrawCylinder(PDI, Link.Right, FVector(1,0,0), FVector(0,1,0), FVector(0,0,1), Link.SnapRadius, StepHeights[StepHeightIndex], 10, MeshColorInstance, SDPG_World);
-			DrawCylinder(PDI, Link.Left, FVector(1,0,0), FVector(0,1,0), FVector(0,0,1), Link.SnapRadius, StepHeights[StepHeightIndex], 10, MeshColorInstance, SDPG_World);
-		}
-	}
-
-	static const float SegmentArcHeight = 0.25f;
-	for (int32 LinkIndex = 0; LinkIndex < OffMeshSegmentLinks.Num(); ++LinkIndex)
-	{
-		const FNavLinkSegmentDrawing& Link = OffMeshSegmentLinks[LinkIndex];
-		if ((Link.SupportedAgentsBits & AgentMask) == 0)
-		{
-			continue;
-		}
-
-		const uint32 SegmentsStart = FPlatformMath::Max<uint32>(SegmentArcHeight*(Link.RightStart-Link.LeftStart).Size()/10, 8);
-		const uint32 SegmentsEnd = FPlatformMath::Max<uint32>(SegmentArcHeight*(Link.RightEnd-Link.LeftEnd).Size()/10, 8);
-		DrawArc(PDI, Link.LeftStart, Link.RightStart, SegmentArcHeight, SegmentsStart, Link.Color, SDPG_World, 3.5f);
-		DrawArc(PDI, Link.LeftEnd, Link.RightEnd, SegmentArcHeight, SegmentsEnd, Link.Color, SDPG_World, 3.5f);
-		const FVector VOffset(0,0,FVector::Dist(Link.LeftStart, Link.RightStart)*1.333f);
-
-		switch (Link.Direction)
-		{
-		case ENavLinkDirection::LeftToRight:
-			DrawArrowHead(PDI, Link.RightStart, Link.LeftStart+VOffset, 30.f, Link.Color, SDPG_World, 3.5f);
-			DrawArrowHead(PDI, Link.RightEnd, Link.LeftEnd+VOffset, 30.f, Link.Color, SDPG_World, 3.5f);
-			break;
-		case ENavLinkDirection::RightToLeft:
-			DrawArrowHead(PDI, Link.LeftStart, Link.RightStart+VOffset, 30.f, Link.Color, SDPG_World, 3.5f);
-			DrawArrowHead(PDI, Link.LeftEnd, Link.RightEnd+VOffset, 30.f, Link.Color, SDPG_World, 3.5f);
-			break;
-		case ENavLinkDirection::BothWays:
-		default:
-			DrawArrowHead(PDI, Link.RightStart, Link.LeftStart+VOffset, 30.f, Link.Color, SDPG_World, 3.5f);
-			DrawArrowHead(PDI, Link.RightEnd, Link.LeftEnd+VOffset, 30.f, Link.Color, SDPG_World, 3.5f);
-			DrawArrowHead(PDI, Link.LeftStart, Link.RightStart+VOffset, 30.f, Link.Color, SDPG_World, 3.5f);
-			DrawArrowHead(PDI, Link.LeftEnd, Link.RightEnd+VOffset, 30.f, Link.Color, SDPG_World, 3.5f);
-			break;
-		}
-
-		// draw snap-spheres on both ends
-		for (int32 StepHeightIndex = 0; StepHeightIndex < StepHeights.Num(); ++StepHeightIndex)
-		{
-			DrawCylinder(PDI, Link.RightStart, FVector(1,0,0), FVector(0,1,0), FVector(0,0,1), Link.SnapRadius, StepHeights[StepHeightIndex], 10, MeshColorInstance, SDPG_World);
-			DrawCylinder(PDI, Link.RightEnd, FVector(1,0,0), FVector(0,1,0), FVector(0,0,1), Link.SnapRadius, StepHeights[StepHeightIndex], 10, MeshColorInstance, SDPG_World);
-			DrawCylinder(PDI, Link.LeftStart, FVector(1,0,0), FVector(0,1,0), FVector(0,0,1), Link.SnapRadius, StepHeights[StepHeightIndex], 10, MeshColorInstance, SDPG_World);
-			DrawCylinder(PDI, Link.LeftEnd, FVector(1,0,0), FVector(0,1,0), FVector(0,0,1), Link.SnapRadius, StepHeights[StepHeightIndex], 10, MeshColorInstance, SDPG_World);
-		}
-	}
-}
-
 FPrimitiveViewRelevance FNavLinkRenderingProxy::GetViewRelevance(const FSceneView* View)
 {
 	FPrimitiveViewRelevance Result;
diff --git a/Engine/Source/Runtime/Engine/Private/AI/Navigation/NavTestRenderingComponent.cpp b/Engine/Source/Runtime/Engine/Private/AI/Navigation/NavTestRenderingComponent.cpp
--- a/Engine/Source/Runtime/Engine/Private/AI/Navigation/NavTestRenderingComponent.cpp
+++ b/Engine/Source/Runtime/Engine/Private/AI/Navigation/NavTestRenderingComponent.cpp
@@ -186,102 +186,6 @@ public:
 		}
 	}
 
-	virtual void DrawDynamicElements(FPrimitiveDrawInterface* PDI,const FSceneView* View) override
-	{
-		if (NavTestActor)
-		{
-			//DrawArc(PDI, Link.Left, Link.Right, 0.4f, NavMeshColors[Link.AreaID], SDPG_World, 3.5f);
-			//const FVector VOffset(0,0,FVector::Dist(Link.Left, Link.Right)*1.333f);
-			//DrawArrowHead(PDI, Link.Right, Link.Left+VOffset, 30.f, NavMeshColors[Link.AreaID], SDPG_World, 3.5f);
-
-			const FVector ActorLocation = NavTestActor->GetActorLocation();
-			const FVector ProjectedLocation = NavTestActor->ProjectedLocation + NavMeshDrawOffset;
-			const FColor ProjectedColor = NavTestActor->bProjectedLocationValid ? FColor(0, 255, 0, 120) : FColor(255, 0, 0, 120);
-			const FVector BoxExtent(20, 20, 20);
-
-			FMaterialRenderProxy* const ColoredMeshInstance = new(FMemStack::Get()) FColoredMaterialRenderProxy(GEngine->DebugMeshMaterial->GetRenderProxy(false), ProjectedColor);
-			//DrawBox(PDI, FTransform(ProjectedLocation).ToMatrixNoScale(),BoxExtent, ColoredMeshInstance, SDPG_World);
-			DrawSphere(PDI, ProjectedLocation, BoxExtent, 10, 7, ColoredMeshInstance, SDPG_World);
-
-			//DrawWireBox(PDI, FBox(ProjectedLocation-BoxExtent, ProjectedLocation+BoxExtent), ProjectedColor, false);
-			DrawWireBox(PDI, FBox(ActorLocation - BoxExtent, ActorLocation + BoxExtent), FColor::White, false);
-			const FVector LineEnd = ProjectedLocation - (ProjectedLocation - ActorLocation).SafeNormal()*BoxExtent.X;
-			PDI->DrawLine(LineEnd, ActorLocation, ProjectedColor, SDPG_World, 2.5);
-			DrawArrowHead(PDI, LineEnd, ActorLocation, 20.f, ProjectedColor, SDPG_World, 2.5f);
-
-			// draw query extent
-			DrawWireBox(PDI, FBox(ActorLocation - NavTestActor->QueryingExtent, ActorLocation + NavTestActor->QueryingExtent), FColor::Blue, false);
-		}
-
-		// draw path
-		if (!bShowBestPath || !NodeDebug.Num())
-		{
-			for (int32 PointIndex = 1; PointIndex < PathPoints.Num(); PointIndex++)
-			{
-				PDI->DrawLine(PathPoints[PointIndex-1], PathPoints[PointIndex], FLinearColor::Red, SDPG_World, 2.0f, 0.0f, true);
-			}
-		}
-
-		// draw path debug data
-		if (bShowNodePool)
-		{
-			if (ClosedSetIndices.Num())
-			{
-				const FColoredMaterialRenderProxy *MeshColorInstance = new(FMemStack::Get()) FColoredMaterialRenderProxy(GEngine->DebugMeshMaterial->GetRenderProxy(false), NavMeshRenderColor_ClosedSet);						
-				FDynamicMeshBuilder	MeshBuilder;
-				MeshBuilder.AddVertices(ClosedSetVerts);
-				MeshBuilder.AddTriangles(ClosedSetIndices);
-				MeshBuilder.Draw(PDI, FMatrix::Identity, MeshColorInstance, GetDepthPriorityGroup(View));
-			}
-
-			if (OpenSetIndices.Num())
-			{
-				const FColoredMaterialRenderProxy *MeshColorInstance = new(FMemStack::Get()) FColoredMaterialRenderProxy(GEngine->DebugMeshMaterial->GetRenderProxy(false), NavMeshRenderColor_OpenSet);						
-				FDynamicMeshBuilder	MeshBuilder;
-				MeshBuilder.AddVertices(OpenSetVerts);
-				MeshBuilder.AddTriangles(OpenSetIndices);
-				MeshBuilder.Draw(PDI, FMatrix::Identity, MeshColorInstance, GetDepthPriorityGroup(View));
-			}
-		}
-
-		for (TSet<FNodeDebugData>::TIterator It(NodeDebug); It; ++It)
-		{
-			const FNodeDebugData& NodeData = *It;
-			
-			FColor LineColor(FColor::Blue);
-			if (bShowBestPath && NodeData.bBestPath)
-			{
-				LineColor = FColor::Red;
-			}
-
-			if (bShowDiff)
-			{
-				LineColor.A = NodeData.bModified ? NavMeshRenderAlpha_Modifed : NavMeshRenderAlpha_NonModified;
-			}
-
-			FVector ParentPos(NodeData.ParentId.IsValidId() ? NodeDebug[NodeData.ParentId].Position : NodeData.Position);
-
-			if (bShowDiff && !NodeData.bModified)
-			{
-				PDI->DrawLine(NodeData.Position, ParentPos, LineColor, SDPG_World);
-			}
-			else
-			{
-				PDI->DrawLine(NodeData.Position, ParentPos, LineColor, SDPG_World, 2.0f, 0.0, true);
-			}
-
-			if (NodeData.bOffMeshLink)
-			{
-				DrawWireBox(PDI, FBox::BuildAABB(NodeData.Position, FVector(10.0f)), LineColor, SDPG_World);
-			}
-
-			if (bShowDiff && NodeData.bModified)
-			{
-				PDI->DrawLine(NodeData.Position + FVector(0,0,10), NodeData.Position + FVector(0,0,100), FColor::Green, SDPG_World);
-			}
-		}
-	}
-
 	void GatherPathPoints()
 	{
 		if (NavTestActor && NavTestActor->LastPath.IsValid())
diff --git a/Engine/Source/Runtime/Engine/Private/ArrowComponent.cpp b/Engine/Source/Runtime/Engine/Private/ArrowComponent.cpp
--- a/Engine/Source/Runtime/Engine/Private/ArrowComponent.cpp
+++ b/Engine/Source/Runtime/Engine/Private/ArrowComponent.cpp
@@ -222,83 +222,6 @@ public:
 		}
 	}
 
-	/** 
-	* Draw the scene proxy as a dynamic element
-	*
-	* @param	PDI - draw interface to render to
-	* @param	View - current view
-	*/
-	virtual void DrawDynamicElements(FPrimitiveDrawInterface* PDI,const FSceneView* View)
-	{
-		QUICK_SCOPE_CYCLE_COUNTER( STAT_ArrowSceneProxy_DrawDynamicElements );
-
-		FMatrix EffectiveLocalToWorld;
-#if WITH_EDITOR
-		if (bLightAttachment)
-		{
-			EffectiveLocalToWorld = GetLocalToWorld().GetMatrixWithoutScale();
-		}
-		else
-#endif	//WITH_EDITOR
-		{
-			EffectiveLocalToWorld = GetLocalToWorld();
-		}
-
-		const FColoredMaterialRenderProxy ArrowMaterialRenderProxy(
-			GEngine->ArrowMaterial->GetRenderProxy(IsSelected(), IsHovered()),
-			ArrowColor,
-			"GizmoColor"
-			);
-
-
-		// Calculate the view-dependent scaling factor.
-		float ViewScale = 1.0f;
-		if (bIsScreenSizeScaled && (View->ViewMatrices.ProjMatrix.M[3][3] != 1.0f))
-		{
-			const float ZoomFactor = FMath::Min<float>(View->ViewMatrices.ProjMatrix.M[0][0], View->ViewMatrices.ProjMatrix.M[1][1]);
-			if (ZoomFactor != 0.0f)
-			{
-				const float Radius = View->WorldToScreen(Origin).W * (ScreenSize / ZoomFactor);
-				if (Radius < 1.0f && Radius > 0)
-				{
-					ViewScale *= Radius;
-				}
-			}
-		}
-
-#if WITH_EDITORONLY_DATA
-		ViewScale *= EditorScale;
-#endif
-
-
-		// Draw the mesh.
-		FMeshBatch Mesh;
-		FMeshBatchElement& BatchElement = Mesh.Elements[0];
-		BatchElement.IndexBuffer = &IndexBuffer;
-		Mesh.bWireframe = false;
-		Mesh.VertexFactory = &VertexFactory;
-		Mesh.MaterialRenderProxy = &ArrowMaterialRenderProxy;
-		BatchElement.PrimitiveUniformBuffer = CreatePrimitiveUniformBufferImmediate(FScaleMatrix(ViewScale) * EffectiveLocalToWorld, GetBounds(), GetLocalBounds(), true, UseEditorDepthTest());
-		BatchElement.FirstIndex = 0;
-		BatchElement.NumPrimitives = IndexBuffer.Indices.Num() / 3;
-		BatchElement.MinVertexIndex = 0;
-		BatchElement.MaxVertexIndex = VertexBuffer.Vertices.Num() - 1;
-		Mesh.ReverseCulling = IsLocalToWorldDeterminantNegative();
-		Mesh.Type = PT_TriangleList;
-		Mesh.DepthPriorityGroup = SDPG_World;
-		PDI->DrawMesh(Mesh);
-
-
-
-
-
-		//const FMatrix CylinderTransform = FScaleMatrix(ViewScale) * FRotationMatrix( FRotator(-90,0.f,0) ) * FTranslationMatrix(FVector(ViewScale, 0.0f, 0.0f)) * EffectiveLocalToWorld;
-		//DrawCylinder(PDI, CylinderTransform, Offset, FVector(1,0,0), FVector(0,1,0), FVector(0,0,1), ShaftRadius, 0.5f * ShaftLength, 16, &ArrowMaterialRenderProxy, SDPG_World);
-		
-		//const FMatrix ConeTransform = FScaleMatrix(ViewScale * -HeadLength) * FTranslationMatrix(FVector(ViewScale * TotalLength,0,0)) * EffectiveLocalToWorld;
-		//DrawCone(PDI, ConeTransform, HeadAngle, HeadAngle, 32, false, FColor::White, &CollisionMaterialInstance, SDPG_World);
-	}
-
 	virtual FPrimitiveViewRelevance GetViewRelevance(const FSceneView* View)
 	{
 		FPrimitiveViewRelevance Result;
diff --git a/Engine/Source/Runtime/Engine/Private/BillboardComponent.cpp b/Engine/Source/Runtime/Engine/Private/BillboardComponent.cpp
--- a/Engine/Source/Runtime/Engine/Private/BillboardComponent.cpp
+++ b/Engine/Source/Runtime/Engine/Private/BillboardComponent.cpp
@@ -182,78 +182,6 @@ public:
 		}
 	}
 
-	/** 
-	 * Draw the scene proxy as a dynamic element
-	 *
-	 * @param	PDI - draw interface to render to
-	 * @param	View - current view
-	 */
-	virtual void DrawDynamicElements(FPrimitiveDrawInterface* PDI, const FSceneView* View) override
-	{
-		QUICK_SCOPE_CYCLE_COUNTER( STAT_SpriteSceneProxy_DrawDynamicElements );
-
-		FTexture* TextureResource = (Texture != NULL) ? Texture->Resource : NULL;
-		if (TextureResource)
-		{
-			// Calculate the view-dependent scaling factor.
-			float ViewedSizeX = SizeX;
-			float ViewedSizeY = SizeY;
-			if (bIsScreenSizeScaled && (View->ViewMatrices.ProjMatrix.M[3][3] != 1.0f))
-			{
-				const float ZoomFactor	= FMath::Min<float>(View->ViewMatrices.ProjMatrix.M[0][0], View->ViewMatrices.ProjMatrix.M[1][1]);
-				if(ZoomFactor != 0.0f)
-				{
-					const float Radius = View->WorldToScreen(Origin).W * (ScreenSize / ZoomFactor);
-					if (Radius < 1.0f)
-					{
-						ViewedSizeX *= Radius;
-						ViewedSizeY *= Radius;
-					}					
-				}
-			}
-
-#if WITH_EDITORONLY_DATA
-			ViewedSizeX *= EditorScale;
-			ViewedSizeY *= EditorScale;
-#endif
-
-			FLinearColor ColorToUse = Color;
-
-			// Set the selection/hover color from the current engine setting.
-			// The color is multiplied by 10 because this value is normally expected to be blended
-			// additively, this is not how the sprites work and therefore need an extra boost
-			// to appear the same color as previously.
-			if (IsSelected())
-			{
-				ColorToUse = FLinearColor::White + (GEngine->GetSelectedMaterialColor() * GEngine->SelectionHighlightIntensityBillboards * 10);
-			}
-			else if (IsHovered())
-			{
-				ColorToUse = FLinearColor::White + (GEngine->GetHoveredMaterialColor() * GEngine->HoverHighlightIntensity * 10);
-			}
-
-			// Sprites of locked actors draw in red.
-			if (bIsActorLocked)
-			{
-				ColorToUse = FColor(255,0,0);
-			}
-			FLinearColor LevelColorToUse = IsSelected() ? ColorToUse : LevelColor;
-			FLinearColor PropertyColorToUse = PropertyColor;
-
-			const FLinearColor& SpriteColor = View->Family->EngineShowFlags.LevelColoration ? LevelColorToUse :
-											( (View->Family->EngineShowFlags.PropertyColoration) ? PropertyColorToUse : ColorToUse );
-
-			PDI->DrawSprite(
-				Origin,
-				ViewedSizeX,
-				ViewedSizeY,
-				TextureResource,
-				SpriteColor,
-				GetDepthPriorityGroup(View),
-				U,UL,V,VL
-				);
-		}
-	}
 	virtual FPrimitiveViewRelevance GetViewRelevance(const FSceneView* View) override
 	{
 		bool bVisible = View->Family->EngineShowFlags.BillboardSprites;
diff --git a/Engine/Source/Runtime/Engine/Private/BoxComponent.cpp b/Engine/Source/Runtime/Engine/Private/BoxComponent.cpp
--- a/Engine/Source/Runtime/Engine/Private/BoxComponent.cpp
+++ b/Engine/Source/Runtime/Engine/Private/BoxComponent.cpp
@@ -113,15 +113,6 @@ FPrimitiveSceneProxy* UBoxComponent::CreateSceneProxy()
 			}
 		}
 
-		virtual void DrawDynamicElements(FPrimitiveDrawInterface* PDI,const FSceneView* View) override
-		{
-			QUICK_SCOPE_CYCLE_COUNTER( STAT_BoxSceneProxy_DrawDynamicElements );
-
-			const FMatrix& LocalToWorld = GetLocalToWorld();
-			const FColor DrawColor = GetSelectionColor(BoxColor, IsSelected(), IsHovered(), /*bUseOverlayIntensity=*/false);
-			DrawOrientedWireBox(PDI, LocalToWorld.GetOrigin(), LocalToWorld.GetScaledAxis( EAxis::X ), LocalToWorld.GetScaledAxis( EAxis::Y ), LocalToWorld.GetScaledAxis( EAxis::Z ), BoxExtents, DrawColor, SDPG_World);
-		}
-
 		virtual FPrimitiveViewRelevance GetViewRelevance(const FSceneView* View) override
 		{
 			const bool bVisible = !bDrawOnlyIfSelected || IsSelected();
diff --git a/Engine/Source/Runtime/Engine/Private/BrushComponent.cpp b/Engine/Source/Runtime/Engine/Private/BrushComponent.cpp
--- a/Engine/Source/Runtime/Engine/Private/BrushComponent.cpp
+++ b/Engine/Source/Runtime/Engine/Private/BrushComponent.cpp
@@ -329,99 +329,6 @@ public:
 		}
 	}
 
-	/** 
-	* Draw the scene proxy as a dynamic element
-	*
-	* @param	PDI - draw interface to render to
-	* @param	View - current view
-	*/
-	virtual void DrawDynamicElements(FPrimitiveDrawInterface* PDI,const FSceneView* View)
-	{
-		QUICK_SCOPE_CYCLE_COUNTER( STAT_BrushSceneProxy_DrawDynamicElements );
-
-
-		if( AllowDebugViewmodes() )
-		{
-			bool bDrawCollision = false;
-			const bool bInCollisionView = IsCollisionView(View->Family->EngineShowFlags, bDrawCollision);
-
-			// Draw solid if 'solid when selected' and selected, or we are in a 'collision view'
-			const bool bDrawSolid = ((bSolidWhenSelected && IsSelected()) || (bInCollisionView && bDrawCollision));
-			// Don't draw wireframe if in a collision view mode and not drawing solid
-			const bool bDrawWireframe = !bInCollisionView;
-
-			// Choose color to draw it
-			FLinearColor DrawColor = BrushColor;
-			// In a collision view mode
-			if(bInCollisionView)
-			{
-				DrawColor = BrushColor;
-			}
-			else if(View->Family->EngineShowFlags.PropertyColoration)
-			{
-				DrawColor = PropertyColor;
-			}
-			else if(View->Family->EngineShowFlags.LevelColoration)
-			{
-				DrawColor = LevelColor;
-			}
-
-
-			// SOLID
-			if(bDrawSolid)
-			{
-				if(BodySetup != NULL)
-				{
-					const FColoredMaterialRenderProxy SolidMaterialInstance(
-						GEngine->ShadedLevelColorationUnlitMaterial->GetRenderProxy(IsSelected(), IsHovered()),
-						DrawColor
-						);
-
-					FTransform GeomTransform(GetLocalToWorld());
-					BodySetup->AggGeom.DrawAggGeom(PDI, GeomTransform, DrawColor, /*Material=*/&SolidMaterialInstance, false, /*bSolid=*/ true, UseEditorDepthTest() );
-				}
-			}
-			// WIREFRAME
-			else if(bDrawWireframe)
-			{
-				// If we have the editor data (Wire Buffers) use those for wireframe
-#if WITH_EDITOR
-				if(WireIndexBuffer.GetNumEdges() && WireVertexBuffer.GetNumVertices())
-				{
-					FColoredMaterialRenderProxy WireframeMaterial(
-						GEngine->LevelColorationUnlitMaterial->GetRenderProxy(IsSelected(), IsHovered()),
-						GetSelectionColor(DrawColor,!(GIsEditor && (View->Family->EngineShowFlags.Selection)) || IsSelected(), IsHovered(), /*bUseOverlayIntensity=*/false)
-						);
-
-					FMeshBatch Mesh;
-					FMeshBatchElement& BatchElement = Mesh.Elements[0];
-					BatchElement.IndexBuffer = &WireIndexBuffer;
-					Mesh.VertexFactory = &VertexFactory;
-					Mesh.MaterialRenderProxy = &WireframeMaterial;
-					BatchElement.PrimitiveUniformBufferResource = &GetUniformBuffer();
-					BatchElement.FirstIndex = 0;
-					BatchElement.NumPrimitives = WireIndexBuffer.GetNumEdges();
-					BatchElement.MinVertexIndex = 0;
-					BatchElement.MaxVertexIndex = WireVertexBuffer.GetNumVertices() - 1;
-					Mesh.CastShadow = false;
-					Mesh.Type = PT_LineList;
-					Mesh.DepthPriorityGroup = IsSelected() ? SDPG_Foreground : SDPG_World;
-					PDI->DrawMesh(Mesh);
-				}
-				else 
-#endif
-				if(BodySetup != NULL)
-				// If not, use the body setup for wireframe
-				{
-					FTransform GeomTransform(GetLocalToWorld());
-					BodySetup->AggGeom.DrawAggGeom(PDI, GeomTransform, GetSelectionColor(DrawColor, IsSelected(), IsHovered()), /* Material=*/ NULL, false, /* bSolid=*/ false, UseEditorDepthTest() );
-				}
-
-			}
-		}
-
-	}
-
 	virtual FPrimitiveViewRelevance GetViewRelevance(const FSceneView* View)
 	{
 		bool bVisible = false;
diff --git a/Engine/Source/Runtime/Engine/Private/CapsuleComponent.cpp b/Engine/Source/Runtime/Engine/Private/CapsuleComponent.cpp
--- a/Engine/Source/Runtime/Engine/Private/CapsuleComponent.cpp
+++ b/Engine/Source/Runtime/Engine/Private/CapsuleComponent.cpp
@@ -51,22 +51,6 @@ FPrimitiveSceneProxy* UCapsuleComponent::CreateSceneProxy()
 			}
 		}
 
-		/** 
-		* Draw the scene proxy as a dynamic element
-		*
-		* @param	PDI - draw interface to render to
-		* @param	View - current view
-		*/
-		virtual void DrawDynamicElements(FPrimitiveDrawInterface* PDI,const FSceneView* View) override
-		{
-			QUICK_SCOPE_CYCLE_COUNTER( STAT_DrawCylinderSceneProxy_DrawDynamicElements );
-
-			const FLinearColor DrawCapsuleColor = GetSelectionColor(ShapeColor, IsSelected(), IsHovered(), /*bUseOverlayIntensity=*/false);
-			const FMatrix& LocalToWorld = GetLocalToWorld();
-			const int32 CapsuleSides =  FMath::Clamp<int32>(CapsuleRadius/4.f, 16, 64);
-			DrawWireCapsule( PDI, LocalToWorld.GetOrigin(), LocalToWorld.GetScaledAxis( EAxis::X ), LocalToWorld.GetScaledAxis( EAxis::Y ), LocalToWorld.GetScaledAxis( EAxis::Z ), DrawCapsuleColor, CapsuleRadius, CapsuleHalfHeight, CapsuleSides, SDPG_World );
-		}
-
 		virtual FPrimitiveViewRelevance GetViewRelevance(const FSceneView* View) override
 		{
 			const bool bVisible = !bDrawOnlyIfSelected || IsSelected();
diff --git a/Engine/Source/Runtime/Engine/Private/DrawFrustumComponent.cpp b/Engine/Source/Runtime/Engine/Private/DrawFrustumComponent.cpp
--- a/Engine/Source/Runtime/Engine/Private/DrawFrustumComponent.cpp
+++ b/Engine/Source/Runtime/Engine/Private/DrawFrustumComponent.cpp
@@ -92,64 +92,6 @@ public:
 		}
 	}
 
-	/** 
-	* Draw the scene proxy as a dynamic element
-	*
-	* @param	PDI - draw interface to render to
-	* @param	View - current view
-	*/
-	virtual void DrawDynamicElements(FPrimitiveDrawInterface* PDI,const FSceneView* View)
-	{
-		QUICK_SCOPE_CYCLE_COUNTER( STAT_DrawFrustumSceneProxy_DrawDynamicElements );
-
-		FVector Direction(1,0,0);
-		FVector LeftVector(0,1,0);
-		FVector UpVector(0,0,1);
-
-		FVector Verts[8];
-
-		// FOVAngle controls the horizontal angle.
-		float HozHalfAngle = (FrustumAngle) * ((float)PI/360.f);
-		float HozLength = FrustumStartDist * FMath::Tan(HozHalfAngle);
-		float VertLength = HozLength/FrustumAspectRatio;
-
-		// near plane verts
-		Verts[0] = (Direction * FrustumStartDist) + (UpVector * VertLength) + (LeftVector * HozLength);
-		Verts[1] = (Direction * FrustumStartDist) + (UpVector * VertLength) - (LeftVector * HozLength);
-		Verts[2] = (Direction * FrustumStartDist) - (UpVector * VertLength) - (LeftVector * HozLength);
-		Verts[3] = (Direction * FrustumStartDist) - (UpVector * VertLength) + (LeftVector * HozLength);
-
-		HozLength = FrustumEndDist * FMath::Tan(HozHalfAngle);
-		VertLength = HozLength/FrustumAspectRatio;
-
-		// far plane verts
-		Verts[4] = (Direction * FrustumEndDist) + (UpVector * VertLength) + (LeftVector * HozLength);
-		Verts[5] = (Direction * FrustumEndDist) + (UpVector * VertLength) - (LeftVector * HozLength);
-		Verts[6] = (Direction * FrustumEndDist) - (UpVector * VertLength) - (LeftVector * HozLength);
-		Verts[7] = (Direction * FrustumEndDist) - (UpVector * VertLength) + (LeftVector * HozLength);
-
-		for( int32 x = 0 ; x < 8 ; ++x )
-		{
-			Verts[x] = GetLocalToWorld().TransformPosition( Verts[x] );
-		}
-
-		const uint8 DepthPriorityGroup = GetDepthPriorityGroup(View);
-		PDI->DrawLine( Verts[0], Verts[1], FrustumColor, DepthPriorityGroup );
-		PDI->DrawLine( Verts[1], Verts[2], FrustumColor, DepthPriorityGroup );
-		PDI->DrawLine( Verts[2], Verts[3], FrustumColor, DepthPriorityGroup );
-		PDI->DrawLine( Verts[3], Verts[0], FrustumColor, DepthPriorityGroup );
-
-		PDI->DrawLine( Verts[4], Verts[5], FrustumColor, DepthPriorityGroup );
-		PDI->DrawLine( Verts[5], Verts[6], FrustumColor, DepthPriorityGroup );
-		PDI->DrawLine( Verts[6], Verts[7], FrustumColor, DepthPriorityGroup );
-		PDI->DrawLine( Verts[7], Verts[4], FrustumColor, DepthPriorityGroup );
-
-		PDI->DrawLine( Verts[0], Verts[4], FrustumColor, DepthPriorityGroup );
-		PDI->DrawLine( Verts[1], Verts[5], FrustumColor, DepthPriorityGroup );
-		PDI->DrawLine( Verts[2], Verts[6], FrustumColor, DepthPriorityGroup );
-		PDI->DrawLine( Verts[3], Verts[7], FrustumColor, DepthPriorityGroup );
-	}
-
 	virtual FPrimitiveViewRelevance GetViewRelevance(const FSceneView* View)
 	{
 		FPrimitiveViewRelevance Result;
diff --git a/Engine/Source/Runtime/Engine/Private/InstancedStaticMesh.cpp b/Engine/Source/Runtime/Engine/Private/InstancedStaticMesh.cpp
--- a/Engine/Source/Runtime/Engine/Private/InstancedStaticMesh.cpp
+++ b/Engine/Source/Runtime/Engine/Private/InstancedStaticMesh.cpp
@@ -942,9 +942,6 @@ public:
 
 	virtual void GetDynamicMeshElements(const TArray<const FSceneView*>& Views, const FSceneViewFamily& ViewFamily, uint32 VisibilityMap, FMeshElementCollector& Collector) const override;
 
-	/** Draw the scene proxy as a dynamic element */
-	virtual void DrawDynamicElements(FPrimitiveDrawInterface* PDI,const FSceneView* View) override;
-
 	virtual int32 GetNumMeshBatches() const override;
 
 	/** Sets up a shadow FMeshBatch for a specific LOD. */
@@ -1068,72 +1065,6 @@ void FInstancedStaticMeshSceneProxy::GetDynamicMeshElements(const TArray<const F
 #endif
 }
 
-void FInstancedStaticMeshSceneProxy::DrawDynamicElements(FPrimitiveDrawInterface* PDI,const FSceneView* View)
-{
-	QUICK_SCOPE_CYCLE_COUNTER( STAT_InstancedStaticMeshSceneProxy_DrawDynamicElements );
-
-#if !(UE_BUILD_SHIPPING || UE_BUILD_TEST)
-	const bool bSelectionRenderEnabled = GIsEditor && View->Family->EngineShowFlags.Selection;
-
-	// If the first pass rendered selected instances only, we need to render the deselected instances in a second pass
-	const int32 NumPasses = (bSelectionRenderEnabled && bHasSelectedInstances && !PDI->IsRenderingSelectionOutline()) ? 2 : 1;
-
-	FInstancingUserData* PassUserData[2] =
-	{
-		bHasSelectedInstances && bSelectionRenderEnabled ? &UserData_SelectedInstances : &UserData_AllInstances,
-		&UserData_DeselectedInstances
-	};
-
-	bool PassRenderSelection[2] = 
-	{
-		bSelectionRenderEnabled && IsSelected(),
-		false
-	};
-
-	const FLinearColor UtilColor( LevelColor );
-	const int32 LODsToDraw[] = { GetLOD(View) };
-	const bool bIsWireframe = View->Family->EngineShowFlags.Wireframe;
-	int32 NumLODs = StaticMesh->GetNumLODs();
-
-	for( int32 Pass=0;Pass < NumPasses; Pass++ )
-	{
-		for (int32 LODLoopIndex = 0; LODLoopIndex < ARRAY_COUNT(LODsToDraw) && LODsToDraw[LODLoopIndex] != INDEX_NONE && LODsToDraw[LODLoopIndex] < NumLODs; LODLoopIndex++)
-		{
-			const int32 LODIndex = LODsToDraw[LODLoopIndex];
-
-			const FStaticMeshLODResources& LODModel = StaticMesh->RenderData->LODResources[LODIndex];
-
-			for(int32 SectionIndex = 0;SectionIndex < LODModel.Sections.Num();SectionIndex++)
-			{
-				const int32 NumBatches = GetNumMeshBatches();
-
-				for (int32 BatchIndex = 0; BatchIndex < NumBatches; BatchIndex++)
-				{
-					FMeshBatch MeshBatch;
-
-					if (GetMeshElement(LODIndex, BatchIndex, SectionIndex, GetDepthPriorityGroup(View), PassRenderSelection[Pass], IsHovered(), MeshBatch))
-					{
-						MeshBatch.Elements[0].UserData = PassUserData[Pass];
-
-						const int32 NumCalls = DrawRichMesh(
-							PDI,
-							MeshBatch,
-							WireframeColor,
-							UtilColor,
-							PropertyColor,
-							this,
-							PassRenderSelection[Pass],
-							bIsWireframe
-							);
-						INC_DWORD_STAT_BY(STAT_StaticMeshTriangles,MeshBatch.GetNumPrimitives() * NumCalls);
-					}
-				}
-			}
-		}
-	}
-#endif
-}
-
 int32 FInstancedStaticMeshSceneProxy::GetNumMeshBatches() const
 {
 	const bool bInstanced = RHISupportsInstancing(GetFeatureLevelShaderPlatform(InstancedRenderData.FeatureLevel));
diff --git a/Engine/Source/Runtime/Engine/Private/LineBatchComponent.cpp b/Engine/Source/Runtime/Engine/Private/LineBatchComponent.cpp
--- a/Engine/Source/Runtime/Engine/Private/LineBatchComponent.cpp
+++ b/Engine/Source/Runtime/Engine/Private/LineBatchComponent.cpp
@@ -74,53 +74,6 @@ public:
 		}
 	}
 
-	/** 
-	* Draw the scene proxy as a dynamic element
-	*
-	* @param	PDI - draw interface to render to
-	* @param	View - current view
-	*/
-	virtual void DrawDynamicElements(FPrimitiveDrawInterface* PDI,const FSceneView* View) override
-	{
-		QUICK_SCOPE_CYCLE_COUNTER( STAT_LineBatcherSceneProxy_DrawDynamicElements );
-
-		for (int32 i = 0; i < Lines.Num(); i++)
-		{
-			PDI->DrawLine(Lines[i].Start, Lines[i].End, Lines[i].Color, Lines[i].DepthPriority, Lines[i].Thickness);
-		}
-
-		for (int32 i = 0; i < Points.Num(); i++)
-		{
-			PDI->DrawPoint(Points[i].Position, Points[i].Color, Points[i].PointSize, Points[i].DepthPriority);
-		}
-
-		for (int32 i = 0; i < Meshes.Num(); i++)
-		{
-			static FVector const PosX(1.f,0,0);
-			static FVector const PosY(0,1.f,0);
-			static FVector const PosZ(0,0,1.f);
-
-			FBatchedMesh const& M = Meshes[i];
-
-			// this seems far from optimal in terms of perf, but it's for debugging
-			FDynamicMeshBuilder MeshBuilder;
-
-			// set up geometry
-			for (int32 VertIdx=0; VertIdx < M.MeshVerts.Num(); ++VertIdx)
-			{
-				MeshBuilder.AddVertex( M.MeshVerts[VertIdx], FVector2D::ZeroVector, PosX, PosY, PosZ, FColor::White );
-			}
-			//MeshBuilder.AddTriangles(M.MeshIndices);
-			for (int32 Idx=0; Idx < M.MeshIndices.Num(); Idx+=3)
-			{
-				MeshBuilder.AddTriangle( M.MeshIndices[Idx], M.MeshIndices[Idx+1], M.MeshIndices[Idx+2] );
-			}
-
-			FMaterialRenderProxy* const MaterialRenderProxy = new(FMemStack::Get()) FColoredMaterialRenderProxy(GEngine->DebugMeshMaterial->GetRenderProxy(false), M.Color);
-			MeshBuilder.Draw(PDI, FMatrix::Identity, MaterialRenderProxy, M.DepthPriority);
-		}
-	}
-
 	/**
 	*  Returns a struct that describes to the renderer when to draw this proxy.
 	*	@param		Scene view to use to determine our relevence.
diff --git a/Engine/Source/Runtime/Engine/Private/MaterialBillboardComponent.cpp b/Engine/Source/Runtime/Engine/Private/MaterialBillboardComponent.cpp
--- a/Engine/Source/Runtime/Engine/Private/MaterialBillboardComponent.cpp
+++ b/Engine/Source/Runtime/Engine/Private/MaterialBillboardComponent.cpp
@@ -230,114 +230,6 @@ public:
 		}
 	}
 
-	// FPrimitiveSceneProxy interface.
-	virtual void DrawDynamicElements(FPrimitiveDrawInterface* PDI,const FSceneView* View)
-	{
-		QUICK_SCOPE_CYCLE_COUNTER( STAT_MaterialSpriteSceneProxy_DrawDynamicElements );
-
-		const bool bIsWireframe = View->Family->EngineShowFlags.Wireframe;
-		// Determine the position of the source
-		const FVector SourcePosition = GetLocalToWorld().GetOrigin();
-		const FVector CameraToSource = View->ViewMatrices.ViewOrigin - SourcePosition;
-		const float DistanceToSource = CameraToSource.Size();
-
-		const FVector CameraUp      = -View->InvViewProjectionMatrix.TransformVector(FVector(1.0f,0.0f,0.0f));
-		const FVector CameraRight   = -View->InvViewProjectionMatrix.TransformVector(FVector(0.0f,1.0f,0.0f));
-		const FVector CameraForward = -View->InvViewProjectionMatrix.TransformVector(FVector(0.0f,0.0f,1.0f));
-		const FMatrix WorldToLocal = GetLocalToWorld().InverseFast();
-		const FVector LocalCameraUp = WorldToLocal.TransformVector(CameraUp);
-		const FVector LocalCameraRight = WorldToLocal.TransformVector(CameraRight);
-		const FVector LocalCameraForward = WorldToLocal.TransformVector(CameraForward);
-
-		// Draw the elements ordered so the last is on top of the first.
-		for(int32 ElementIndex = 0;ElementIndex < Elements.Num();++ElementIndex)
-		{
-			const FMaterialSpriteElement& Element = Elements[ElementIndex];
-			if(Element.Material)
-			{
-				// Evaluate the size of the sprite.
-				float SizeX = Element.BaseSizeX;
-				float SizeY = Element.BaseSizeY;
-				if(Element.DistanceToSizeCurve)
-				{
-					const float SizeFactor = Element.DistanceToSizeCurve->GetFloatValue(DistanceToSource);
-					SizeX *= SizeFactor;
-					SizeY *= SizeFactor;
-				}
-						
-				// Convert the size into world-space.
-				const float W = View->ViewProjectionMatrix.TransformPosition(SourcePosition).W;
-				const float AspectRatio = CameraRight.Size() / CameraUp.Size();
-				const float WorldSizeX = Element.bSizeIsInScreenSpace ? (SizeX               * W) : (SizeX / CameraRight.Size());
-				const float WorldSizeY = Element.bSizeIsInScreenSpace ? (SizeY * AspectRatio * W) : (SizeY / CameraUp.Size());
-			
-				// Evaluate the color/opacity of the sprite.
-				FLinearColor Color = BaseColor;
-				if(Element.DistanceToOpacityCurve)
-				{
-					Color.A *= Element.DistanceToOpacityCurve->GetFloatValue(DistanceToSource);
-				}
-
-				// Set up the sprite vertex attributes that are constant across the sprite.
-				FMaterialSpriteVertex Vertices[4];
-				for(uint32 VertexIndex = 0;VertexIndex < 4;++VertexIndex)
-				{
-					Vertices[VertexIndex].Color = Color;
-					Vertices[VertexIndex].TangentX = FPackedNormal(LocalCameraRight.SafeNormal());
-					Vertices[VertexIndex].TangentZ = FPackedNormal(-LocalCameraForward.SafeNormal());
-				}
-
-				// Set up the sprite vertex positions and texture coordinates.
-				Vertices[0].Position  = -WorldSizeX * LocalCameraRight + +WorldSizeY * LocalCameraUp;
-				Vertices[0].TexCoords = FVector2D(0,0);
-				Vertices[1].Position  = +WorldSizeX * LocalCameraRight + +WorldSizeY * LocalCameraUp;
-				Vertices[1].TexCoords = FVector2D(1,0);
-				Vertices[2].Position  = -WorldSizeX * LocalCameraRight + -WorldSizeY * LocalCameraUp;
-				Vertices[2].TexCoords = FVector2D(0,1);
-				Vertices[3].Position  = +WorldSizeX * LocalCameraRight + -WorldSizeY * LocalCameraUp;
-				Vertices[3].TexCoords = FVector2D(1,1);
-			
-				// Set up the FMeshElement.
-				FMeshBatch Mesh;
-				Mesh.UseDynamicData      = true;
-				Mesh.DynamicVertexData   = Vertices;
-				Mesh.DynamicVertexStride = sizeof(FMaterialSpriteVertex);
-
-				Mesh.VertexFactory           = &GMaterialSpriteVertexFactory;
-				Mesh.MaterialRenderProxy     = Element.Material->GetRenderProxy((View->Family->EngineShowFlags.Selection) && IsSelected(),IsHovered());
-				Mesh.LCI                     = NULL;
-				Mesh.ReverseCulling          = IsLocalToWorldDeterminantNegative() ? true : false;
-				Mesh.CastShadow              = false;
-				Mesh.DepthPriorityGroup      = (ESceneDepthPriorityGroup)GetDepthPriorityGroup(View);
-				Mesh.Type                    = PT_TriangleStrip;
-				Mesh.bDisableBackfaceCulling = true;
-
-				// Set up the FMeshBatchElement.
-				FMeshBatchElement& BatchElement = Mesh.Elements[0];
-				BatchElement.IndexBuffer            = NULL;
-				BatchElement.DynamicIndexData       = NULL;
-				BatchElement.DynamicIndexStride     = 0;
-				BatchElement.FirstIndex             = 0;
-				BatchElement.MinVertexIndex         = 0;
-				BatchElement.MaxVertexIndex         = 3;
-				BatchElement.PrimitiveUniformBuffer = GetUniformBuffer();
-				BatchElement.NumPrimitives          = 2;
-
-				// Draw the sprite.
-				DrawRichMesh(
-					PDI, 
-					Mesh, 
-					FLinearColor(1.0f, 0.0f, 0.0f),
-					LevelColor,
-					PropertyColor,
-					this,
-					IsSelected(),
-					bIsWireframe
-					);
-			}
-		}
-	}
-
 	virtual FPrimitiveViewRelevance GetViewRelevance(const FSceneView* View)
 	{
 		bool bVisible = View->Family->EngineShowFlags.BillboardSprites;
diff --git a/Engine/Source/Runtime/Engine/Private/ModelRender.cpp b/Engine/Source/Runtime/Engine/Private/ModelRender.cpp
--- a/Engine/Source/Runtime/Engine/Private/ModelRender.cpp
+++ b/Engine/Source/Runtime/Engine/Private/ModelRender.cpp
@@ -403,222 +403,6 @@ public:
 		const_cast<FModelSceneProxy*>(this)->SetSelection_RenderThread(bAnySelectedSurfs);
 	}
 
-	virtual void PreRenderView(const FSceneViewFamily* ViewFamily, const uint32 VisibilityMap, int32 FrameNumber) override
-	{
-		// Reset any batches leftover from last frame.
-		PerViewBatches.Reset();
-		DynamicMeshBatches.Reset();
-		bool bAnySelectedSurfs = false;
-
-		int32 ViewBit = 1;
-		int32 NumViews = ViewFamily->Views.Num();
-		for (int32 ViewIndex = 0; ViewIndex < NumViews; ++ViewIndex, ViewBit <<= 1)
-		{
-			// Create a batch entry for each view whether this primitive is visible in it or not!
-			FPerViewBatch* ViewBatch = new(PerViewBatches) FPerViewBatch();
-
-			const FSceneView* View = ViewFamily->Views[ViewIndex];
-			bool bVisibleInView = (VisibilityMap & ViewBit) == ViewBit;
-			bool bShowSelection = GIsEditor && !View->bIsGameView && View->Family->EngineShowFlags.Selection;
-			bool bDynamicBSPTriangles = bShowSelection || IsRichView(*View->Family) ;
-			bool bShowBSPTriangles = View->Family->EngineShowFlags.BSPTriangles;
-			bool bShowBSP = View->Family->EngineShowFlags.BSP;
-
-#if WITH_EDITOR
-			bool bDrawCollision = false;
-			const bool bInCollisionView = IsCollisionView(View, bDrawCollision);
-			// draw bsp as dynamic when in collision view mode
-			if(bInCollisionView)
-			{
-				bDynamicBSPTriangles = true;
-			}
-#endif
-
-			// If in a collision view, only draw if we have collision
-			if (bVisibleInView && bDynamicBSPTriangles && bShowBSPTriangles && bShowBSP 
-#if WITH_EDITOR
-				&& (!bInCollisionView || bDrawCollision)
-#endif
-				)
-			{
-				ViewBatch->FirstBatch = DynamicMeshBatches.Num();
-				ESceneDepthPriorityGroup DepthPriorityGroup = (ESceneDepthPriorityGroup)GetDepthPriorityGroup(View);
-
-				const FMaterialRenderProxy* MatProxyOverride = NULL;
-
-#if WITH_EDITOR
-				if(bInCollisionView && AllowDebugViewmodes())
-				{
-					MatProxyOverride = &CollisionMaterialInstance;
-				}
-#endif
-
-				// If selection is being shown, batch triangles based on whether they are selected or not.
-				if (bShowSelection)
-				{
-					uint32 TotalIndices = 0;
-					for (int32 ElementIndex = 0; ElementIndex < Elements.Num(); ElementIndex++)
-					{
-						const FModelElement& ModelElement = Component->GetElements()[ElementIndex];
-						TotalIndices += ModelElement.NumTriangles * 3;
-					}
-
-					if (TotalIndices > 0)
-					{
-						FGlobalDynamicIndexBuffer::FAllocation IndexAllocation =
-							FGlobalDynamicIndexBuffer::Get().Allocate<uint32>(TotalIndices);
-						if (IndexAllocation.IsValid())
-						{
-							uint32* Indices = (uint32*)IndexAllocation.Buffer;
-							uint32 FirstIndex = IndexAllocation.FirstIndex;
-							for (int32 ElementIndex = 0; ElementIndex < Elements.Num(); ElementIndex++)
-							{
-								const FModelElement& ModelElement = Component->GetElements()[ElementIndex];
-								if (ModelElement.NumTriangles > 0)
-								{
-									const FElementInfo& ProxyElementInfo = Elements[ElementIndex];
-									bool bHasSelectedSurfs = false;
-									bool bHasHoveredSurfs = false;
-									for (uint32 BatchIndex = 0; BatchIndex < 3; BatchIndex++)
-									{
-										// Three batches total:
-										//		Batch 0: Only surfaces that are neither selected, nor hovered
-										//		Batch 1: Only selected surfaces
-										//		Batch 2: Only hovered surfaces
-										const bool bOnlySelectedSurfaces = ( BatchIndex == 1 );
-										const bool bOnlyHoveredSurfaces = ( BatchIndex == 2 );
-
-										if (bOnlySelectedSurfaces && !bHasSelectedSurfs)
-										{
-											continue;
-										}
-
-										if (bOnlyHoveredSurfaces && !bHasHoveredSurfs)
-										{
-											continue;
-										}
-
-										uint32 MinVertexIndex = MAX_uint32;
-										uint32 MaxVertexIndex = 0;
-										uint32 NumIndices = 0;
-
-										for(int32 NodeIndex = 0;NodeIndex < ModelElement.Nodes.Num();NodeIndex++)
-										{
-											FBspNode& Node = Component->GetModel()->Nodes[ModelElement.Nodes[NodeIndex]];
-											FBspSurf& Surf = Component->GetModel()->Surfs[Node.iSurf];
-
-											if (!ShouldDrawSurface(Surf))
-												continue;
-
-											const bool bSurfaceSelected = (Surf.PolyFlags & PF_Selected) == PF_Selected;
-											const bool bSurfaceHovered = !bSurfaceSelected && ((Surf.PolyFlags & PF_Hovered) == PF_Hovered);
-											bHasSelectedSurfs |= bSurfaceSelected;
-											bHasHoveredSurfs |= bSurfaceHovered;
-
-											if (bSurfaceSelected == bOnlySelectedSurfaces && bSurfaceHovered == bOnlyHoveredSurfaces)
-											{
-												for (uint32 BackFace = 0; BackFace < (uint32)((Surf.PolyFlags & PF_TwoSided) ? 2 : 1); BackFace++)
-												{
-													for (int32 VertexIndex = 2; VertexIndex < Node.NumVertices; VertexIndex++)
-													{
-														*Indices++ = Node.iVertexIndex + Node.NumVertices * BackFace;
-														*Indices++ = Node.iVertexIndex + Node.NumVertices * BackFace + VertexIndex;
-														*Indices++ = Node.iVertexIndex + Node.NumVertices * BackFace + VertexIndex - 1;
-														NumIndices += 3;
-													}
-													MinVertexIndex = FMath::Min(Node.iVertexIndex + Node.NumVertices * BackFace,MinVertexIndex);
-													MaxVertexIndex = FMath::Max(Node.iVertexIndex + Node.NumVertices * BackFace + Node.NumVertices - 1,MaxVertexIndex);
-												}
-											}
-										}
-
-										if (NumIndices > 0)
-										{
-											FDynamicModelMeshBatch& MeshElement = *new(DynamicMeshBatches) FDynamicModelMeshBatch(bOnlySelectedSurfaces);
-											FMeshBatchElement& BatchElement = MeshElement.Elements[0];
-											BatchElement.IndexBuffer = IndexAllocation.IndexBuffer;
-											MeshElement.VertexFactory = &Component->GetModel()->VertexFactory;
-											MeshElement.MaterialRenderProxy = (MatProxyOverride != NULL) ? MatProxyOverride : ProxyElementInfo.GetMaterial()->GetRenderProxy(bOnlySelectedSurfaces, bOnlyHoveredSurfaces);
-											MeshElement.LCI = &ProxyElementInfo;
-											BatchElement.PrimitiveUniformBufferResource = &GetUniformBuffer();
-											BatchElement.FirstIndex = FirstIndex;
-											BatchElement.NumPrimitives = NumIndices / 3;
-											BatchElement.MinVertexIndex = MinVertexIndex;
-											BatchElement.MaxVertexIndex = MaxVertexIndex;
-											MeshElement.Type = PT_TriangleList;
-											MeshElement.DepthPriorityGroup = DepthPriorityGroup;
-											MeshElement.ModelElementIndex = ElementIndex;
-											FirstIndex += NumIndices;
-										}
-									}
-									bAnySelectedSurfs |= bHasSelectedSurfs;
-								}
-							}
-						}
-					}
-				}
-				else
-				{
-					for (int32 ElementIndex = 0; ElementIndex < Elements.Num(); ElementIndex++)
-					{
-						const FModelElement& ModelElement = Component->GetElements()[ElementIndex];
-						if (ModelElement.NumTriangles > 0)
-						{
-							FDynamicModelMeshBatch& MeshElement = *new(DynamicMeshBatches) FDynamicModelMeshBatch(false);
-							FMeshBatchElement& BatchElement = MeshElement.Elements[0];
-							BatchElement.IndexBuffer = ModelElement.IndexBuffer;
-							MeshElement.VertexFactory = &Component->GetModel()->VertexFactory;
-							MeshElement.MaterialRenderProxy = (MatProxyOverride != NULL) ? MatProxyOverride : Elements[ElementIndex].GetMaterial()->GetRenderProxy(false);
-							MeshElement.LCI = &Elements[ElementIndex];
-							BatchElement.PrimitiveUniformBufferResource = &GetUniformBuffer();
-							BatchElement.FirstIndex = ModelElement.FirstIndex;
-							BatchElement.NumPrimitives = ModelElement.NumTriangles;
-							BatchElement.MinVertexIndex = ModelElement.MinVertexIndex;
-							BatchElement.MaxVertexIndex = ModelElement.MaxVertexIndex;
-							MeshElement.Type = PT_TriangleList;
-							MeshElement.DepthPriorityGroup = DepthPriorityGroup;
-							MeshElement.ModelElementIndex = ElementIndex;
-						}
-					}
-				}
-
-				ViewBatch->LastBatch = DynamicMeshBatches.Num();
-			}
-		}
-
-		SetSelection_RenderThread(bAnySelectedSurfs);
-	}
-
-	/** 
-	* Draw the scene proxy as a dynamic element
-	*
-	* @param	PDI - draw interface to render to
-	* @param	View - current view
-	*/
-	virtual void DrawDynamicElements(FPrimitiveDrawInterface* PDI,const FSceneView* View)
-	{
-		QUICK_SCOPE_CYCLE_COUNTER( STAT_ModelSceneProxy_DrawDynamicElements );
-
-		int32 ViewIndex = View->Family->Views.Find(View);
-		const bool bRenderingSelectionOutline = PDI->IsRenderingSelectionOutline();
-
-		if (PerViewBatches.IsValidIndex(ViewIndex))
-		{
-			FLinearColor UtilColor = LevelColor;
-
-			const bool bIsWireframe = View->Family->EngineShowFlags.Wireframe;
-			FPerViewBatch& PerViewBatch = PerViewBatches[ViewIndex];
-			for (int32 BatchIndex = PerViewBatch.FirstBatch; BatchIndex < PerViewBatch.LastBatch; ++BatchIndex)
-			{
-				FDynamicModelMeshBatch& Batch = DynamicMeshBatches[BatchIndex];
-				if (!bRenderingSelectionOutline || Batch.bIsSelectedBatch)
-				{
-					DrawRichMesh(PDI,Batch,FLinearColor::White,UtilColor,PropertyColor,this,false, bIsWireframe);
-				}
-			}
-		}
-	}
-
 	virtual void DrawStaticElements(FStaticPrimitiveDrawInterface* PDI)
 	{
 		if(!HasViewDependentDPG())
@@ -662,7 +446,6 @@ public:
 			|| (bShowSelectedTriangles && HasSelectedSurfaces()))
 		{
 			Result.bDynamicRelevance = true;
-			Result.bNeedsPreRenderView = true;
 		}
 		else
 		{
@@ -850,20 +633,6 @@ private:
 
 	FMaterialRelevance MaterialRelevance;
 
-	/** Which dynamic batches should be rendered for a given view. */
-	struct FPerViewBatch
-	{
-		int32 FirstBatch;
-		int32 LastBatch;
-
-		FPerViewBatch()
-			: FirstBatch(0)
-			, LastBatch(0)
-		{
-		}
-	};
-	TArray<FPerViewBatch, TInlineAllocator<2> > PerViewBatches;
-
 	/** Precomputed dynamic mesh batches. */
 	struct FDynamicModelMeshBatch : public FMeshBatch
 	{
@@ -877,7 +646,6 @@ private:
 		{
 		}
 	};
-	TArray<FDynamicModelMeshBatch> DynamicMeshBatches;
 
 	/** Collision Response of this component**/
 	FCollisionResponseContainer CollisionResponse;
