Paper2D: Fix off-by-one rendering for orthogonal tile maps and pulled some math out of the inner loop

diff --git a/Engine/Plugins/2D/Paper2D/Source/Paper2D/Private/PaperTileMapRenderComponent.cpp b/Engine/Plugins/2D/Paper2D/Source/Paper2D/Private/PaperTileMapRenderComponent.cpp
--- a/Engine/Plugins/2D/Paper2D/Source/Paper2D/Private/PaperTileMapRenderComponent.cpp
+++ b/Engine/Plugins/2D/Paper2D/Source/Paper2D/Private/PaperTileMapRenderComponent.cpp
@@ -144,6 +144,36 @@ void UPaperTileMapRenderComponent::RebuildRenderData(FPaperTileMapRenderScenePro
 		return;
 	}
 
+
+	FVector WorldOffsetNoY;
+	FVector WorldOffsetYFactor = FVector::ZeroVector;
+	FVector WorldStepX;
+	FVector WorldStepY;
+	FVector TileSetOffset = FVector::ZeroVector;
+
+	switch (TileMap->ProjectionMode)
+	{
+	case ETileMapProjectionMode::Orthogonal:
+	default:
+		WorldOffsetNoY = -(TileMap->TileWidth * PaperAxisX * 0.5f) - (TileMap->TileHeight * PaperAxisY * 0.5f);
+		WorldStepX = PaperAxisX * TileMap->TileWidth;
+		WorldStepY = -PaperAxisY * TileMap->TileHeight;
+		break;
+	case ETileMapProjectionMode::IsometricDiamond:
+		WorldOffsetNoY = (TileMap->MapWidth * TileMap->TileWidth * 0.5f * PaperAxisX);
+		WorldOffsetYFactor = -(TileMap->TileHeight * 0.5f * PaperAxisY);
+		WorldStepX = (TileMap->TileWidth * PaperAxisX * 0.5f) - (TileMap->TileHeight * PaperAxisY * 0.5f);
+		WorldStepY = (TileMap->TileWidth * PaperAxisX * -0.5f) - (TileMap->TileHeight * PaperAxisY * 0.5f);
+		break;
+	case ETileMapProjectionMode::IsometricStaggered:
+		WorldOffsetNoY = 0.5f * TileMap->TileHeight * PaperAxisY;
+		WorldOffsetYFactor = 0.5f * TileMap->TileWidth * PaperAxisX;
+		WorldStepX = PaperAxisX * TileMap->TileWidth;
+		WorldStepY = -PaperAxisY * TileMap->TileHeight;
+		break;
+	}
+
+
 	for (int32 Z = 0; Z < TileMap->TileLayers.Num(); ++Z)
 	{
 		UPaperTileLayer* Layer = TileMap->TileLayers[Z];
@@ -166,27 +196,20 @@ void UPaperTileMapRenderComponent::RebuildRenderData(FPaperTileMapRenderScenePro
 		for (int32 Y = 0; Y < TileMap->MapHeight; ++Y)
 		{
 			// In pixels
-			FVector WorldOffset = FVector::ZeroVector;
-			FVector WorldStepX;
-			FVector WorldStepY;
 			FVector TileSetOffset = FVector::ZeroVector;
+			FVector WorldOffset;
+
 			switch (TileMap->ProjectionMode)
 			{
 			case ETileMapProjectionMode::Orthogonal:
 			default:
-				WorldOffset = (TileMap->TileWidth * PaperAxisX * 0.5f) + (TileMap->TileHeight * PaperAxisY * 0.5f);
-				WorldStepX = PaperAxisX * TileMap->TileWidth;
-				WorldStepY = -PaperAxisY * TileMap->TileHeight;
+				WorldOffset = WorldOffsetNoY;
 				break;
 			case ETileMapProjectionMode::IsometricDiamond:
-				WorldOffset = (TileMap->MapWidth * TileMap->TileWidth * 0.5f * PaperAxisX) - (Y * TileMap->TileHeight * 0.5f * PaperAxisY);
-				WorldStepX = (TileMap->TileWidth * PaperAxisX * 0.5f) - (TileMap->TileHeight * PaperAxisY * 0.5f);
-				WorldStepY = (TileMap->TileWidth * PaperAxisX * -0.5f) - (TileMap->TileHeight * PaperAxisY * 0.5f);
+				WorldOffset = WorldOffsetNoY + Y * WorldOffsetYFactor;
 				break;
 			case ETileMapProjectionMode::IsometricStaggered:
-				WorldOffset = ((Y & 1) * 0.5f * TileMap->TileWidth * PaperAxisX) + (0.5f * TileMap->TileHeight * PaperAxisY);
-				WorldStepX = PaperAxisX * TileMap->TileWidth;
-				WorldStepY = -PaperAxisY * TileMap->TileHeight;
+				WorldOffset = WorldOffsetNoY + (Y & 1) * WorldOffsetYFactor;
 				break;
 			}
 
