Paper2D: Add FTileMapLayerReregisterContext to UPaperTileLayer::PostEditChangeProperty so tile map edits are pushed to component instances immediately

diff --git a/Engine/Plugins/2D/Paper2D/Source/Paper2D/Private/PaperTileLayer.cpp b/Engine/Plugins/2D/Paper2D/Source/Paper2D/Private/PaperTileLayer.cpp
--- a/Engine/Plugins/2D/Paper2D/Source/Paper2D/Private/PaperTileLayer.cpp
+++ b/Engine/Plugins/2D/Paper2D/Source/Paper2D/Private/PaperTileLayer.cpp
@@ -66,8 +66,50 @@ void UPaperTileLayer::ReallocateAndCopyMap()
 }
 
 #if WITH_EDITOR
+
+/** Removes all components that use the specified tile map layer from their scenes for the lifetime of the class. */
+class FTileMapLayerReregisterContext
+{
+public:
+	/** Initialization constructor. */
+	FTileMapLayerReregisterContext(UPaperTileLayer* TargetAsset)
+	{
+		// Look at tile map components
+		for (TObjectIterator<UPaperTileMapRenderComponent> MapIt; MapIt; ++MapIt)
+		{
+			if (UPaperTileMap* TestMap = (*MapIt)->TileMap)
+			{
+				if (TestMap->TileLayers.Contains(TargetAsset))
+				{
+					AddComponentToRefresh(*MapIt);
+				}
+			}
+		}
+	}
+
+protected:
+	void AddComponentToRefresh(UActorComponent* Component)
+	{
+		if (ComponentContexts.Num() == 0)
+		{
+			// wait until resources are released
+			FlushRenderingCommands();
+		}
+
+		new (ComponentContexts) FComponentReregisterContext(Component);
+	}
+
+private:
+	/** The recreate contexts for the individual components. */
+	TIndirectArray<FComponentReregisterContext> ComponentContexts;
+};
+
+
+
 void UPaperTileLayer::PostEditChangeProperty(struct FPropertyChangedEvent& PropertyChangedEvent)
 {
+	FTileMapLayerReregisterContext ReregisterExistingComponents(this);
+
 	FName PropertyName = (PropertyChangedEvent.Property != NULL) ? PropertyChangedEvent.Property->GetFName() : NAME_None;
 
 	if ((PropertyName == GET_MEMBER_NAME_CHECKED(UPaperTileLayer, LayerWidth)) || (PropertyName == GET_MEMBER_NAME_CHECKED(UPaperTileLayer, LayerHeight)))
