Paper2D: Fix visibility toggle not taking effect on tile map layers

diff --git a/Engine/Plugins/2D/Paper2D/Source/Paper2DEditor/Private/TileMapEditing/STileLayerItem.cpp b/Engine/Plugins/2D/Paper2D/Source/Paper2DEditor/Private/TileMapEditing/STileLayerItem.cpp
--- a/Engine/Plugins/2D/Paper2D/Source/Paper2DEditor/Private/TileMapEditing/STileLayerItem.cpp
+++ b/Engine/Plugins/2D/Paper2D/Source/Paper2DEditor/Private/TileMapEditing/STileLayerItem.cpp
@@ -189,12 +189,7 @@ TSharedRef<SWidget> STileLayerItem::GenerateWidgetForColumn(const FName& ColumnN
 	}
 	else
 	{
-		return SNew(SContentReference)
-			.AssetReference(this, &STileLayerItem::GetLayerTileset)
-			.OnSetReference(this, &STileLayerItem::OnChangeLayerTileSet)
-			.AllowedClass(UPaperTileSet::StaticClass())
-			.AllowSelectingNewAsset(true)
-			.AllowClearingReference(false);
+		return SNew(SSpacer);
 	}
 }
 
@@ -213,26 +208,13 @@ void STileLayerItem::OnLayerNameCommitted(const FText& NewText, ETextCommit::Typ
 	MyLayer->LayerName = NewText;
 }
 
-void STileLayerItem::OnChangeLayerTileSet(UObject* NewAsset)
-{
-	//@TODO: PAPER2D: Dead code now?
-// 	if (UPaperTileSet* TileSet = Cast<UPaperTileSet>(NewAsset))
-// 	{
-// 		MyLayer->TileSet = TileSet;
-// 	}
-}
-
-UObject* STileLayerItem::GetLayerTileset() const
-{
-	return nullptr;// MyLayer->TileSet;
-}
-
 FReply STileLayerItem::OnToggleVisibility()
 {
 	const FScopedTransaction Transaction( LOCTEXT("ToggleVisibility", "Toggle Layer Visibility") );
 	MyLayer->SetFlags(RF_Transactional);
 	MyLayer->Modify();
 	MyLayer->bHiddenInEditor = !MyLayer->bHiddenInEditor;
+	MyLayer->PostEditChange();
 	return FReply::Handled();
 }
 
diff --git a/Engine/Plugins/2D/Paper2D/Source/Paper2DEditor/Private/TileMapEditing/STileLayerItem.h b/Engine/Plugins/2D/Paper2D/Source/Paper2DEditor/Private/TileMapEditing/STileLayerItem.h
--- a/Engine/Plugins/2D/Paper2D/Source/Paper2DEditor/Private/TileMapEditing/STileLayerItem.h
+++ b/Engine/Plugins/2D/Paper2D/Source/Paper2DEditor/Private/TileMapEditing/STileLayerItem.h
@@ -27,7 +27,4 @@ protected:
 	const FSlateBrush* GetVisibilityBrushForLayer() const;
 	FSlateColor GetForegroundColorForVisibilityButton() const;
 	FReply OnToggleVisibility();
-
-	void OnChangeLayerTileSet(UObject* NewAsset);
-	UObject* GetLayerTileset() const;
 };
diff --git a/Engine/Plugins/2D/Paper2D/Source/Paper2DEditor/Private/TileMapEditing/STileLayerList.cpp b/Engine/Plugins/2D/Paper2D/Source/Paper2DEditor/Private/TileMapEditing/STileLayerList.cpp
--- a/Engine/Plugins/2D/Paper2D/Source/Paper2DEditor/Private/TileMapEditing/STileLayerList.cpp
+++ b/Engine/Plugins/2D/Paper2D/Source/Paper2DEditor/Private/TileMapEditing/STileLayerList.cpp
@@ -10,8 +10,8 @@
 void STileLayerList::Construct(const FArguments& InArgs, UPaperTileMap* TileMap)
 {
 	TSharedRef<SHeaderRow> HeaderRowWidget = SNew(SHeaderRow)
-		+SHeaderRow::Column("Name").DefaultLabel(NSLOCTEXT("TileLayerList", "TileLayerNameHeader", "Layer Name")).FillWidth(0.7)
-		+SHeaderRow::Column("TileSet").DefaultLabel(NSLOCTEXT("TileLayerList", "TileLayerTileSetHeader", "Tile Set")).FillWidth(0.7);
+		+SHeaderRow::Column("Name").DefaultLabel(NSLOCTEXT("TileLayerList", "TileLayerNameHeader", "Layer Name")).FillWidth(0.7);
+		//+SHeaderRow::Column("TileSet").DefaultLabel(NSLOCTEXT("TileLayerList", "TileLayerTileSetHeader", "Tile Set")).FillWidth(0.7);
 
 	ChildSlot
 	[
