Paper2D: Fix non-unity build (missing include in PaperTileLayer.cpp)

diff --git a/Engine/Plugins/2D/Paper2D/Source/Paper2D/Private/PaperTileLayer.cpp b/Engine/Plugins/2D/Paper2D/Source/Paper2D/Private/PaperTileLayer.cpp
--- a/Engine/Plugins/2D/Paper2D/Source/Paper2D/Private/PaperTileLayer.cpp
+++ b/Engine/Plugins/2D/Paper2D/Source/Paper2D/Private/PaperTileLayer.cpp
@@ -2,6 +2,7 @@
 
 #include "Paper2DPrivatePCH.h"
 #include "PaperSprite.h"
+#include "ComponentReregisterContext.h"
 
 //////////////////////////////////////////////////////////////////////////
 // UPaperTileLayer
