Paper2D: Fix automatic selection logic for tile map layers

diff --git a/Engine/Plugins/2D/Paper2D/Source/Paper2DEditor/Private/TileMapEditing/STileLayerList.cpp b/Engine/Plugins/2D/Paper2D/Source/Paper2DEditor/Private/TileMapEditing/STileLayerList.cpp
--- a/Engine/Plugins/2D/Paper2D/Source/Paper2DEditor/Private/TileMapEditing/STileLayerList.cpp
+++ b/Engine/Plugins/2D/Paper2D/Source/Paper2DEditor/Private/TileMapEditing/STileLayerList.cpp
@@ -85,7 +85,7 @@ void STileLayerList::Construct(const FArguments& InArgs, UPaperTileMap* TileMap)
 		.HeaderRow(HeaderRowWidget);
 
 	// Select the top item by default
-	if (TileMap->TileLayers.Num() == 0)
+	if (TileMap->TileLayers.Num() > 0)
 	{
 		ListViewWidget->SetSelection(TileMap->TileLayers[TileMap->TileLayers.Num() - 1]);
 	}
