Paper2D: Fix a crash when selecting another actor and then clicking in the viewport after selecting a tile map

diff --git a/Engine/Plugins/2D/Paper2D/Source/Paper2DEditor/Private/TileMapEditing/EdModeTileMap.cpp b/Engine/Plugins/2D/Paper2D/Source/Paper2DEditor/Private/TileMapEditing/EdModeTileMap.cpp
--- a/Engine/Plugins/2D/Paper2D/Source/Paper2DEditor/Private/TileMapEditing/EdModeTileMap.cpp
+++ b/Engine/Plugins/2D/Paper2D/Source/Paper2DEditor/Private/TileMapEditing/EdModeTileMap.cpp
@@ -337,46 +337,49 @@ UPaperTileLayer* FEdModeTileMap::GetSelectedLayerUnderCursor(const FViewportCurs
 		}
 	}
 
-	if (UPaperTileMap* TileMap = TileMapComponent->TileMap)
+	if (TileMapComponent != nullptr)
 	{
-		// Find the first visible layer
-		int32 LayerIndex = 0;
-		for (; LayerIndex < TileMap->TileLayers.Num(); ++LayerIndex)
+		if (UPaperTileMap* TileMap = TileMapComponent->TileMap)
 		{
-			UPaperTileLayer* Layer = TileMap->TileLayers[LayerIndex];
-			if (!Layer->bHiddenInEditor && Layer->bCollisionLayer == bCollisionPainting)
+			// Find the first visible layer
+			int32 LayerIndex = 0;
+			for (; LayerIndex < TileMap->TileLayers.Num(); ++LayerIndex)
 			{
-				break;
+				UPaperTileLayer* Layer = TileMap->TileLayers[LayerIndex];
+				if (!Layer->bHiddenInEditor && Layer->bCollisionLayer == bCollisionPainting)
+				{
+					break;
+				}
 			}
-		}
 
-		// If there was a visible layer, pick it
-		if (LayerIndex < TileMap->TileLayers.Num())
-		{
-			UPaperTileLayer* Layer = TileMap->TileLayers[LayerIndex];
+			// If there was a visible layer, pick it
+			if (LayerIndex < TileMap->TileLayers.Num())
+			{
+				UPaperTileLayer* Layer = TileMap->TileLayers[LayerIndex];
 
-			const float WX = TileMap->MapWidth * TileMap->TileWidth;
-			const float WY = TileMap->MapHeight * TileMap->TileHeight;
+				const float WX = TileMap->MapWidth * TileMap->TileWidth;
+				const float WY = TileMap->MapHeight * TileMap->TileHeight;
 
-			ComponentToWorld = (TileMapComponent != nullptr) ? TileMapComponent->ComponentToWorld : FTransform::Identity;
-			FVector LocalStart = ComponentToWorld.InverseTransformPosition(TraceStart);
-			FVector LocalDirection = ComponentToWorld.InverseTransformVector(TraceDir);
+				ComponentToWorld = (TileMapComponent != nullptr) ? TileMapComponent->ComponentToWorld : FTransform::Identity;
+				FVector LocalStart = ComponentToWorld.InverseTransformPosition(TraceStart);
+				FVector LocalDirection = ComponentToWorld.InverseTransformVector(TraceDir);
 
-			FVector Intersection;
-			FPlane Plane(FVector(1, 0, 0), FVector::ZeroVector, FVector(0, 0, 1));
+				FVector Intersection;
+				FPlane Plane(FVector(1, 0, 0), FVector::ZeroVector, FVector(0, 0, 1));
 
-			if (FMath::SegmentPlaneIntersection(LocalStart, LocalDirection * HALF_WORLD_MAX, Plane, /*out*/ Intersection))
-			{
-				//@TODO: Ideally tile pivots weren't in the center!
-				const float NormalizedX = (Intersection.X + 0.5f * TileMap->TileWidth) / WX;
-				const float NormalizedY = (-Intersection.Z + 0.5f * TileMap->TileHeight) / WY;
-
-				OutTileX = FMath::FloorToInt(NormalizedX * TileMap->MapWidth);
-				OutTileY = FMath::FloorToInt(NormalizedY * TileMap->MapHeight);
-					
-				if ((OutTileX > -BrushWidth) && (OutTileX < TileMap->MapWidth) && (OutTileY > -BrushHeight) && (OutTileY < TileMap->MapHeight))
+				if (FMath::SegmentPlaneIntersection(LocalStart, LocalDirection * HALF_WORLD_MAX, Plane, /*out*/ Intersection))
 				{
-					return Layer;
+					//@TODO: Ideally tile pivots weren't in the center!
+					const float NormalizedX = (Intersection.X + 0.5f * TileMap->TileWidth) / WX;
+					const float NormalizedY = (-Intersection.Z + 0.5f * TileMap->TileHeight) / WY;
+
+					OutTileX = FMath::FloorToInt(NormalizedX * TileMap->MapWidth);
+					OutTileY = FMath::FloorToInt(NormalizedY * TileMap->MapHeight);
+
+					if ((OutTileX > -BrushWidth) && (OutTileX < TileMap->MapWidth) && (OutTileY > -BrushHeight) && (OutTileY < TileMap->MapHeight))
+					{
+						return Layer;
+					}
 				}
 			}
 		}
