Paper2D: Tile map layer browser style tweaks

diff --git a/Engine/Plugins/2D/Paper2D/Source/Paper2DEditor/Private/PaperStyle.cpp b/Engine/Plugins/2D/Paper2D/Source/Paper2DEditor/Private/PaperStyle.cpp
--- a/Engine/Plugins/2D/Paper2D/Source/Paper2DEditor/Private/PaperStyle.cpp
+++ b/Engine/Plugins/2D/Paper2D/Source/Paper2DEditor/Private/PaperStyle.cpp
@@ -20,6 +20,7 @@ TSharedPtr< class ISlateStyle > FPaperStyle::Get() { return StyleSet; }
 void FPaperStyle::Initialize()
 {
 	// Const icon sizes
+	const FVector2D Icon8x8(8.0f, 8.0f);
 	const FVector2D Icon20x20(20.0f, 20.0f);
 	const FVector2D Icon40x40(40.0f, 40.0f);
 
@@ -61,6 +62,22 @@ void FPaperStyle::Initialize()
 
 		StyleSet->Set("TileMapEditor.MoveLayerDown", new IMAGE_PLUGIN_BRUSH(TEXT("Icons/TileMapEditor/icon_TileMapEditor_MoveLayerDown_40x"), Icon40x40));
 		StyleSet->Set("TileMapEditor.MoveLayerDown.Small", new IMAGE_PLUGIN_BRUSH(TEXT("Icons/TileMapEditor/icon_TileMapEditor_MoveLayerDown_40x"), Icon20x20));
+
+
+		const FLinearColor LayerSelectionColor = FLinearColor(0.13f, 0.70f, 1.00f);
+
+		// Selection color for the active row should be ??? to align with the editor viewport.
+		const FTableRowStyle& NormalTableRowStyle = FEditorStyle::Get().GetWidgetStyle<FTableRowStyle>("TableView.Row");
+
+		StyleSet->Set("TileMapEditor.LayerBrowser.TableViewRow",
+			FTableRowStyle(NormalTableRowStyle)
+			.SetActiveBrush(IMAGE_BRUSH("Common/Selection", Icon8x8, LayerSelectionColor))
+			.SetActiveHoveredBrush(IMAGE_BRUSH("Common/Selection", Icon8x8, LayerSelectionColor))
+			.SetInactiveBrush(IMAGE_BRUSH("Common/Selection", Icon8x8, LayerSelectionColor))
+			.SetInactiveHoveredBrush(IMAGE_BRUSH("Common/Selection", Icon8x8, LayerSelectionColor))
+			);
+
+		StyleSet->Set("TileMapEditor.LayerBrowser.SelectionColor", LayerSelectionColor);
 	}
 
 	// Sprite editor
diff --git a/Engine/Plugins/2D/Paper2D/Source/Paper2DEditor/Private/TileMapEditing/STileLayerItem.cpp b/Engine/Plugins/2D/Paper2D/Source/Paper2DEditor/Private/TileMapEditing/STileLayerItem.cpp
--- a/Engine/Plugins/2D/Paper2D/Source/Paper2DEditor/Private/TileMapEditing/STileLayerItem.cpp
+++ b/Engine/Plugins/2D/Paper2D/Source/Paper2DEditor/Private/TileMapEditing/STileLayerItem.cpp
@@ -4,6 +4,7 @@
 #include "STileLayerItem.h"
 #include "SContentReference.h"
 #include "ScopedTransaction.h"
+#include "PaperStyle.h"
 
 #define LOCTEXT_NAMESPACE "Paper2D"
 
@@ -152,7 +153,13 @@ void STileLayerItem::Construct(const FArguments& InArgs, class UPaperTileLayer*
 {
 	MyLayer = InItem;
 
-	SMultiColumnTableRow<class UPaperTileLayer*>::Construct(FSuperRowType::FArguments(), OwnerTable);
+	auto ParentArgs = FSuperRowType::FArguments()
+		.Style(&FPaperStyle::Get()->GetWidgetStyle<FTableRowStyle>("TileMapEditor.LayerBrowser.TableViewRow"));
+
+	SMultiColumnTableRow<class UPaperTileLayer*>::Construct(ParentArgs, OwnerTable);
+
+	EyeClosed = FEditorStyle::GetBrush("Layer.NotVisibleIcon16x");
+	EyeOpened = FEditorStyle::GetBrush("Layer.VisibleIcon16x");
 }
 
 TSharedRef<SWidget> STileLayerItem::GenerateWidgetForColumn(const FName& ColumnName) 
@@ -165,8 +172,8 @@ TSharedRef<SWidget> STileLayerItem::GenerateWidgetForColumn(const FName& ColumnN
 			.VAlign(VAlign_Center)
 			[
 				SAssignNew( VisibilityButton, SButton )
-				.ContentPadding( 0 )
-				.ButtonStyle( FEditorStyle::Get(), "ToolBarButton" ) //@TODO: Bogus?
+				.ContentPadding(FMargin(4.0f, 4.0f, 4.0f, 4.0f))
+				.ButtonStyle( FEditorStyle::Get(), "NoBorder" )
 				.OnClicked( this, &STileLayerItem::OnToggleVisibility )
 				.ToolTipText( LOCTEXT("LayerVisibilityButtonToolTip", "Toggle Layer Visibility") )
 				.ForegroundColor( FSlateColor::UseForeground() )
@@ -181,6 +188,7 @@ TSharedRef<SWidget> STileLayerItem::GenerateWidgetForColumn(const FName& ColumnN
 			]
 			+SHorizontalBox::Slot()
 			.VAlign(VAlign_Center)
+			.Padding(FMargin(4.0f, 4.0f, 4.0f, 4.0f))
 			[
 				SNew(SRenamableEntry)
 				.Text(this, &STileLayerItem::GetLayerDisplayName)
@@ -220,13 +228,13 @@ FReply STileLayerItem::OnToggleVisibility()
 
 const FSlateBrush* STileLayerItem::GetVisibilityBrushForLayer() const
 {
-	return MyLayer->bHiddenInEditor ? FEditorStyle::GetBrush("Layer.NotVisibleIcon16x") : FEditorStyle::GetBrush("Layer.VisibleIcon16x");
+	return MyLayer->bHiddenInEditor ? EyeClosed : EyeOpened;
 }
 
 FSlateColor STileLayerItem::GetForegroundColorForVisibilityButton() const
 {
 	static const FName InvertedForeground("InvertedForeground");
-	return (VisibilityButton.IsValid() && (VisibilityButton->IsHovered() || VisibilityButton->IsPressed())) ? FEditorStyle::GetSlateColor(InvertedForeground) : FSlateColor::UseForeground();
+	return FEditorStyle::GetSlateColor(InvertedForeground);
 }
 
 //////////////////////////////////////////////////////////////////////////
diff --git a/Engine/Plugins/2D/Paper2D/Source/Paper2DEditor/Private/TileMapEditing/STileLayerItem.h b/Engine/Plugins/2D/Paper2D/Source/Paper2DEditor/Private/TileMapEditing/STileLayerItem.h
--- a/Engine/Plugins/2D/Paper2D/Source/Paper2DEditor/Private/TileMapEditing/STileLayerItem.h
+++ b/Engine/Plugins/2D/Paper2D/Source/Paper2DEditor/Private/TileMapEditing/STileLayerItem.h
@@ -20,6 +20,10 @@ protected:
 	class UPaperTileLayer* MyLayer;
 
 	TSharedPtr<SButton> VisibilityButton;
+
+	const FSlateBrush* EyeClosed;
+	const FSlateBrush* EyeOpened;
+
 protected:
 	FText GetLayerDisplayName() const;
 	void OnLayerNameCommitted(const FText& NewText, ETextCommit::Type CommitInfo);
