Paper2D: Improve layer naming when adding a new layer

diff --git a/Engine/Plugins/2D/Paper2D/Source/Paper2D/Classes/PaperTileMap.h b/Engine/Plugins/2D/Paper2D/Source/Paper2D/Classes/PaperTileMap.h
--- a/Engine/Plugins/2D/Paper2D/Source/Paper2D/Classes/PaperTileMap.h
+++ b/Engine/Plugins/2D/Paper2D/Source/Paper2D/Classes/PaperTileMap.h
@@ -89,6 +89,10 @@ public:
 	/** Importing data and options used for this tile map */
 	UPROPERTY(Category=ImportSettings, VisibleAnywhere, Instanced)
 	class UAssetImportData* AssetImportData;
+
+	/** The naming index to start at when trying to create a new layer */
+	UPROPERTY()
+	int32 LayerNameIndex;
 #endif
 
 public:
diff --git a/Engine/Plugins/2D/Paper2D/Source/Paper2D/Private/PaperTileMap.cpp b/Engine/Plugins/2D/Paper2D/Source/Paper2D/Private/PaperTileMap.cpp
--- a/Engine/Plugins/2D/Paper2D/Source/Paper2D/Private/PaperTileMap.cpp
+++ b/Engine/Plugins/2D/Paper2D/Source/Paper2D/Private/PaperTileMap.cpp
@@ -20,6 +20,10 @@ UPaperTileMap::UPaperTileMap(const FObjectInitializer& ObjectInitializer)
 	SeparationPerLayer = 64.0f;
 	SpriteCollisionDomain = ESpriteCollisionMode::None;
 
+#if WITH_EDITORONLY_DATA
+	LayerNameIndex = 1;
+#endif
+
 	static ConstructorHelpers::FObjectFinder<UMaterial> DefaultMaterial(TEXT("/Paper2D/DefaultSpriteMaterial"));
 	Material = DefaultMaterial.Object;
 }
diff --git a/Engine/Plugins/2D/Paper2D/Source/Paper2DEditor/Private/TileMapEditing/STileLayerList.cpp b/Engine/Plugins/2D/Paper2D/Source/Paper2DEditor/Private/TileMapEditing/STileLayerList.cpp
--- a/Engine/Plugins/2D/Paper2D/Source/Paper2DEditor/Private/TileMapEditing/STileLayerList.cpp
+++ b/Engine/Plugins/2D/Paper2D/Source/Paper2DEditor/Private/TileMapEditing/STileLayerList.cpp
@@ -126,13 +126,33 @@ UPaperTileLayer* STileLayerList::AddLayer(bool bCollisionLayer)
 		TileMap->SetFlags(RF_Transactional);
 		TileMap->Modify();
 
+		// Create a set of existing names
+		TSet<FString> ExistingNames;
+		for (UPaperTileLayer* ExistingLayer : TileMap->TileLayers)
+		{
+			ExistingNames.Add(ExistingLayer->LayerName.ToString());
+		}
+
+		// Find a good name
+		FText TestLayerName;
+		do
+		{
+			TileMap->LayerNameIndex++;
+
+			FNumberFormattingOptions NoGroupingFormat;
+			NoGroupingFormat.SetUseGrouping(false);
+
+			TestLayerName = FText::Format(LOCTEXT("NewLayerNameFormatString", "Layer {0}"), FText::AsNumber(TileMap->LayerNameIndex, &NoGroupingFormat));
+		} while (ExistingNames.Contains(TestLayerName.ToString()));
+
+		// Create the new layer
 		NewLayer = NewObject<UPaperTileLayer>(TileMap);
 		NewLayer->SetFlags(RF_Transactional);
 
 		NewLayer->LayerWidth = TileMap->MapWidth;
 		NewLayer->LayerHeight = TileMap->MapHeight;
 		NewLayer->DestructiveAllocateMap(NewLayer->LayerWidth, NewLayer->LayerHeight);
-		NewLayer->LayerName = LOCTEXT("DefaultNewLayerName", "New Layer");
+		NewLayer->LayerName = TestLayerName;
 		NewLayer->bCollisionLayer = bCollisionLayer;
 
 		TileMap->TileLayers.Add(NewLayer);
