Paper2D: Make add new layer actions insert the layer above/below the selection, instead of always at the end of the layer list

diff --git a/Engine/Plugins/2D/Paper2D/Source/Paper2DEditor/Private/TileMapEditing/STileLayerList.cpp b/Engine/Plugins/2D/Paper2D/Source/Paper2DEditor/Private/TileMapEditing/STileLayerList.cpp
--- a/Engine/Plugins/2D/Paper2D/Source/Paper2DEditor/Private/TileMapEditing/STileLayerList.cpp
+++ b/Engine/Plugins/2D/Paper2D/Source/Paper2DEditor/Private/TileMapEditing/STileLayerList.cpp
@@ -116,7 +116,7 @@ UPaperTileLayer* STileLayerList::GetSelectedLayer() const
 	return (ListViewWidget->GetNumItemsSelected() > 0) ? ListViewWidget->GetSelectedItems()[0] : nullptr;
 }
 
-UPaperTileLayer* STileLayerList::AddLayer(bool bCollisionLayer)
+UPaperTileLayer* STileLayerList::AddLayer(bool bCollisionLayer, int32 InsertionIndex)
 {
 	UPaperTileLayer* NewLayer = NULL;
 
@@ -155,7 +155,18 @@ UPaperTileLayer* STileLayerList::AddLayer(bool bCollisionLayer)
 		NewLayer->LayerName = TestLayerName;
 		NewLayer->bCollisionLayer = bCollisionLayer;
 
-		TileMap->TileLayers.Add(NewLayer);
+		// Insert the new layer
+		if (TileMap->TileLayers.IsValidIndex(InsertionIndex))
+		{
+			TileMap->TileLayers.Insert(NewLayer, InsertionIndex);
+		}
+		else
+		{
+			TileMap->TileLayers.Add(NewLayer);
+		}
+
+		// Change the selection set to select it
+		ListViewWidget->SetSelection(NewLayer);
 	}
 
 	return NewLayer;
@@ -163,14 +174,12 @@ UPaperTileLayer* STileLayerList::AddLayer(bool bCollisionLayer)
 
 void STileLayerList::AddNewLayerAbove()
 {
-	//@TODO: TILEMAPS: Support adding layers above or below the selection, instead of always at the top
-	AddLayer(/*bCollisionLayer=*/ false);
+	AddLayer(/*bCollisionLayer=*/ false, GetSelectionIndex());
 }
 
 void STileLayerList::AddNewLayerBelow()
 {
-	//@TODO: TILEMAPS: Support adding layers above or below the selection, instead of always at the top
-	AddLayer(/*bCollisionLayer=*/ false);
+	AddLayer(/*bCollisionLayer=*/ false, GetSelectionIndex() + 1);
 }
 
 int32 STileLayerList::GetSelectionIndex() const
diff --git a/Engine/Plugins/2D/Paper2D/Source/Paper2DEditor/Private/TileMapEditing/STileLayerList.h b/Engine/Plugins/2D/Paper2D/Source/Paper2DEditor/Private/TileMapEditing/STileLayerList.h
--- a/Engine/Plugins/2D/Paper2D/Source/Paper2DEditor/Private/TileMapEditing/STileLayerList.h
+++ b/Engine/Plugins/2D/Paper2D/Source/Paper2DEditor/Private/TileMapEditing/STileLayerList.h
@@ -29,7 +29,7 @@ protected:
 	// Returns the selected index if anything is selected, or the top item otherwise (only returns INDEX_NONE if there are no layers)
 	int32 GetSelectionIndex() const;
 
-	class UPaperTileLayer* AddLayer(bool bCollisionLayer);
+	class UPaperTileLayer* AddLayer(bool bCollisionLayer, int32 InsertionIndex = INDEX_NONE);
 
 	void AddNewLayerAbove();
 	void AddNewLayerBelow();
