Paper2D: Tile map layer browser selection improvements (delete selects the item below the one deleted; duplicate selects the duplicate)

diff --git a/Engine/Plugins/2D/Paper2D/Source/Paper2DEditor/Private/TileMapEditing/STileLayerList.cpp b/Engine/Plugins/2D/Paper2D/Source/Paper2DEditor/Private/TileMapEditing/STileLayerList.cpp
--- a/Engine/Plugins/2D/Paper2D/Source/Paper2DEditor/Private/TileMapEditing/STileLayerList.cpp
+++ b/Engine/Plugins/2D/Paper2D/Source/Paper2DEditor/Private/TileMapEditing/STileLayerList.cpp
@@ -87,7 +87,7 @@ void STileLayerList::Construct(const FArguments& InArgs, UPaperTileMap* TileMap)
 	// Select the top item by default
 	if (TileMap->TileLayers.Num() > 0)
 	{
-		ListViewWidget->SetSelection(TileMap->TileLayers[TileMap->TileLayers.Num() - 1]);
+		SetSelectedLayer(TileMap->TileLayers[TileMap->TileLayers.Num() - 1]);
 	}
 
 	ChildSlot
@@ -166,7 +166,7 @@ UPaperTileLayer* STileLayerList::AddLayer(bool bCollisionLayer, int32 InsertionI
 		}
 
 		// Change the selection set to select it
-		ListViewWidget->SetSelection(NewLayer);
+		SetSelectedLayer(NewLayer);
 	}
 
 	return NewLayer;
@@ -215,6 +215,13 @@ void STileLayerList::DeleteLayer()
 			TileMap->TileLayers.RemoveAt(DeleteIndex);
 
 			TileMap->PostEditChange();
+
+			// Select the item below the one that just got deleted
+			const int32 NewSelectionIndex = FMath::Min<int32>(DeleteIndex, TileMap->TileLayers.Num() - 1);
+			if (TileMap->TileLayers.IsValidIndex(NewSelectionIndex))
+			{
+				SetSelectedLayer(TileMap->TileLayers[NewSelectionIndex]);
+			}
 		}
 	}
 }
@@ -232,8 +239,12 @@ void STileLayerList::DuplicateLayer()
 
 			UPaperTileLayer* NewLayer = DuplicateObject<UPaperTileLayer>(TileMap->TileLayers[DuplicateIndex], TileMap);
 			TileMap->TileLayers.Insert(NewLayer, DuplicateIndex);
+			//@TODO: Try renaming it to a better name (Photoshop uses "[OldName] copy", but with a prompt to let you rename it immediately)
 
 			TileMap->PostEditChange();
+
+			// Select the duplicated layer
+			SetSelectedLayer(NewLayer);
 		}
 	}
 }
@@ -253,6 +264,11 @@ void STileLayerList::MoveLayerDown()
 	//@TODO: TILEMAPS: Support moving layers down
 }
 
+void STileLayerList::SetSelectedLayer(UPaperTileLayer* SelectedLayer)
+{
+	ListViewWidget->SetSelection(SelectedLayer);
+}
+
 //////////////////////////////////////////////////////////////////////////
 
 #undef LOCTEXT_NAMESPACE
\ No newline at end of file
diff --git a/Engine/Plugins/2D/Paper2D/Source/Paper2DEditor/Private/TileMapEditing/STileLayerList.h b/Engine/Plugins/2D/Paper2D/Source/Paper2DEditor/Private/TileMapEditing/STileLayerList.h
--- a/Engine/Plugins/2D/Paper2D/Source/Paper2DEditor/Private/TileMapEditing/STileLayerList.h
+++ b/Engine/Plugins/2D/Paper2D/Source/Paper2DEditor/Private/TileMapEditing/STileLayerList.h
@@ -38,4 +38,6 @@ protected:
 	void MergeLayerDown();
 	void MoveLayerUp();
 	void MoveLayerDown();
+
+	void SetSelectedLayer(UPaperTileLayer* SelectedLayer);
 };
