Paper2D: Add support for reordering tile map layers and fixed some list refresh issues

diff --git a/Engine/Plugins/2D/Paper2D/Source/Paper2DEditor/Private/TileMapEditing/STileLayerList.cpp b/Engine/Plugins/2D/Paper2D/Source/Paper2DEditor/Private/TileMapEditing/STileLayerList.cpp
--- a/Engine/Plugins/2D/Paper2D/Source/Paper2DEditor/Private/TileMapEditing/STileLayerList.cpp
+++ b/Engine/Plugins/2D/Paper2D/Source/Paper2DEditor/Private/TileMapEditing/STileLayerList.cpp
@@ -165,6 +165,9 @@ UPaperTileLayer* STileLayerList::AddLayer(bool bCollisionLayer, int32 InsertionI
 			TileMap->TileLayers.Add(NewLayer);
 		}
 
+		ListViewWidget->RequestListRefresh();
+		TileMap->PostEditChange();
+
 		// Change the selection set to select it
 		SetSelectedLayer(NewLayer);
 	}
@@ -172,6 +175,26 @@ UPaperTileLayer* STileLayerList::AddLayer(bool bCollisionLayer, int32 InsertionI
 	return NewLayer;
 }
 
+void STileLayerList::ChangeLayerOrdering(int32 OldIndex, int32 NewIndex)
+{
+	if (UPaperTileMap* TileMap = TileMapPtr.Get())
+	{
+		if (TileMap->TileLayers.IsValidIndex(OldIndex) && TileMap->TileLayers.IsValidIndex(NewIndex))
+		{
+			const FScopedTransaction Transaction(LOCTEXT("TileMapReorderLayer", "Reorder Layer"));
+			TileMap->SetFlags(RF_Transactional);
+			TileMap->Modify();
+
+			UPaperTileLayer* LayerToMove = TileMap->TileLayers[OldIndex];
+			TileMap->TileLayers.RemoveAt(OldIndex);
+			TileMap->TileLayers.Insert(LayerToMove, NewIndex);
+
+			ListViewWidget->RequestListRefresh();
+			TileMap->PostEditChange();
+		}
+	}
+}
+
 void STileLayerList::AddNewLayerAbove()
 {
 	AddLayer(/*bCollisionLayer=*/ false, GetSelectionIndex());
@@ -215,6 +238,7 @@ void STileLayerList::DeleteLayer()
 			TileMap->TileLayers.RemoveAt(DeleteIndex);
 
 			TileMap->PostEditChange();
+			ListViewWidget->RequestListRefresh();
 
 			// Select the item below the one that just got deleted
 			const int32 NewSelectionIndex = FMath::Min<int32>(DeleteIndex, TileMap->TileLayers.Num() - 1);
@@ -242,6 +266,7 @@ void STileLayerList::DuplicateLayer()
 			//@TODO: Try renaming it to a better name (Photoshop uses "[OldName] copy", but with a prompt to let you rename it immediately)
 
 			TileMap->PostEditChange();
+			ListViewWidget->RequestListRefresh();
 
 			// Select the duplicated layer
 			SetSelectedLayer(NewLayer);
@@ -256,12 +281,16 @@ void STileLayerList::MergeLayerDown()
 
 void STileLayerList::MoveLayerUp()
 {
-	//@TODO: TILEMAPS: Support moving layers up
+	const int32 SelectedIndex = GetSelectionIndex();
+	const int32 NewIndex = SelectedIndex - 1;
+	ChangeLayerOrdering(SelectedIndex, NewIndex);
 }
 
 void STileLayerList::MoveLayerDown()
 {
-	//@TODO: TILEMAPS: Support moving layers down
+	const int32 SelectedIndex = GetSelectionIndex();
+	const int32 NewIndex = SelectedIndex + 1;
+	ChangeLayerOrdering(SelectedIndex, NewIndex);
 }
 
 void STileLayerList::SetSelectedLayer(UPaperTileLayer* SelectedLayer)
diff --git a/Engine/Plugins/2D/Paper2D/Source/Paper2DEditor/Private/TileMapEditing/STileLayerList.h b/Engine/Plugins/2D/Paper2D/Source/Paper2DEditor/Private/TileMapEditing/STileLayerList.h
--- a/Engine/Plugins/2D/Paper2D/Source/Paper2DEditor/Private/TileMapEditing/STileLayerList.h
+++ b/Engine/Plugins/2D/Paper2D/Source/Paper2DEditor/Private/TileMapEditing/STileLayerList.h
@@ -31,6 +31,9 @@ protected:
 
 	class UPaperTileLayer* AddLayer(bool bCollisionLayer, int32 InsertionIndex = INDEX_NONE);
 
+	// Moves a layer from OldIndex to NewIndex if both are valid, otherwise it does nothing silently
+	void ChangeLayerOrdering(int32 OldIndex, int32 NewIndex);
+
 	void AddNewLayerAbove();
 	void AddNewLayerBelow();
 	void DeleteLayer();
