Paper2D: Add tile map layer duplication renaming logic similar to how Photoshop behaves (copy, copy 2, etc...)

diff --git a/Engine/Plugins/2D/Paper2D/Source/Paper2DEditor/Private/TileMapEditing/STileLayerList.cpp b/Engine/Plugins/2D/Paper2D/Source/Paper2DEditor/Private/TileMapEditing/STileLayerList.cpp
--- a/Engine/Plugins/2D/Paper2D/Source/Paper2DEditor/Private/TileMapEditing/STileLayerList.cpp
+++ b/Engine/Plugins/2D/Paper2D/Source/Paper2DEditor/Private/TileMapEditing/STileLayerList.cpp
@@ -116,34 +116,93 @@ UPaperTileLayer* STileLayerList::GetSelectedLayer() const
 	return (ListViewWidget->GetNumItemsSelected() > 0) ? ListViewWidget->GetSelectedItems()[0] : nullptr;
 }
 
-UPaperTileLayer* STileLayerList::AddLayer(bool bCollisionLayer, int32 InsertionIndex)
+FText STileLayerList::GenerateNewLayerName(UPaperTileMap* TileMap)
 {
-	UPaperTileLayer* NewLayer = NULL;
+	// Create a set of existing names
+	TSet<FString> ExistingNames;
+	for (UPaperTileLayer* ExistingLayer : TileMap->TileLayers)
+	{
+		ExistingNames.Add(ExistingLayer->LayerName.ToString());
+	}
 
-	if (UPaperTileMap* TileMap = TileMapPtr.Get())
+	// Find a good name
+	FText TestLayerName;
+	do
 	{
-		const FScopedTransaction Transaction(LOCTEXT("TileMapAddLayer", "Add New Layer"));
-		TileMap->SetFlags(RF_Transactional);
-		TileMap->Modify();
+		TileMap->LayerNameIndex++;
+
+		FNumberFormattingOptions NoGroupingFormat;
+		NoGroupingFormat.SetUseGrouping(false);
+
+		TestLayerName = FText::Format(LOCTEXT("NewLayerNameFormatString", "Layer {0}"), FText::AsNumber(TileMap->LayerNameIndex, &NoGroupingFormat));
+	} while (ExistingNames.Contains(TestLayerName.ToString()));
 
-		// Create a set of existing names
-		TSet<FString> ExistingNames;
-		for (UPaperTileLayer* ExistingLayer : TileMap->TileLayers)
+	return TestLayerName;
+}
+
+FText STileLayerList::GenerateDuplicatedLayerName(const FString& InputNameRaw, UPaperTileMap* TileMap)
+{
+	// Create a set of existing names
+	TSet<FString> ExistingNames;
+	for (UPaperTileLayer* ExistingLayer : TileMap->TileLayers)
+	{
+		ExistingNames.Add(ExistingLayer->LayerName.ToString());
+	}
+
+	FString BaseName = InputNameRaw;
+	int32 TestIndex = 0;
+	bool bAddNumber = false;
+
+	// See if this is the result of a previous duplication operation, and change the desired name accordingly
+	int32 SpaceIndex;
+	if (InputNameRaw.FindLastChar(' ', /*out*/ SpaceIndex))
+	{
+		FString PossibleDuplicationSuffix = InputNameRaw.Mid(SpaceIndex + 1);
+
+		if (PossibleDuplicationSuffix == TEXT("copy"))
+		{
+			bAddNumber = true;
+			BaseName = InputNameRaw.Left(SpaceIndex);
+			TestIndex = 2;
+		}
+		else
 		{
-			ExistingNames.Add(ExistingLayer->LayerName.ToString());
+			int32 ExistingIndex = FCString::Atoi(*PossibleDuplicationSuffix);
+
+			const FString TestSuffix = FString::Printf(TEXT(" copy %d"), ExistingIndex);
+
+			if (InputNameRaw.EndsWith(TestSuffix))
+			{
+				bAddNumber = true;
+				BaseName = InputNameRaw.Left(InputNameRaw.Len() - TestSuffix.Len());
+				TestIndex = ExistingIndex + 1;
+			}
 		}
+	}
+
+	// Find a good name
+	FString TestLayerName = BaseName + TEXT(" copy");
 
-		// Find a good name
-		FText TestLayerName;
+	if (bAddNumber || ExistingNames.Contains(TestLayerName))
+	{
 		do
 		{
-			TileMap->LayerNameIndex++;
+			TestLayerName = FString::Printf(TEXT("%s copy %d"), *BaseName, TestIndex++);
+		} while (ExistingNames.Contains(TestLayerName));
+	}
 
-			FNumberFormattingOptions NoGroupingFormat;
-			NoGroupingFormat.SetUseGrouping(false);
+	return FText::FromString(TestLayerName);
+}
+
+UPaperTileLayer* STileLayerList::AddLayer(bool bCollisionLayer, int32 InsertionIndex)
+{
+	UPaperTileLayer* NewLayer = NULL;
 
-			TestLayerName = FText::Format(LOCTEXT("NewLayerNameFormatString", "Layer {0}"), FText::AsNumber(TileMap->LayerNameIndex, &NoGroupingFormat));
-		} while (ExistingNames.Contains(TestLayerName.ToString()));
+	if (UPaperTileMap* TileMap = TileMapPtr.Get())
+	{
+		const FScopedTransaction Transaction(LOCTEXT("TileMapAddLayer", "Add New Layer"));
+		TileMap->SetFlags(RF_Transactional);
+		TileMap->Modify();
 
 		// Create the new layer
 		NewLayer = NewObject<UPaperTileLayer>(TileMap);
@@ -152,7 +211,7 @@ UPaperTileLayer* STileLayerList::AddLayer(bool bCollisionLayer, int32 InsertionI
 		NewLayer->LayerWidth = TileMap->MapWidth;
 		NewLayer->LayerHeight = TileMap->MapHeight;
 		NewLayer->DestructiveAllocateMap(NewLayer->LayerWidth, NewLayer->LayerHeight);
-		NewLayer->LayerName = TestLayerName;
+		NewLayer->LayerName = GenerateNewLayerName(TileMap);
 		NewLayer->bCollisionLayer = bCollisionLayer;
 
 		// Insert the new layer
@@ -263,7 +322,7 @@ void STileLayerList::DuplicateLayer()
 
 			UPaperTileLayer* NewLayer = DuplicateObject<UPaperTileLayer>(TileMap->TileLayers[DuplicateIndex], TileMap);
 			TileMap->TileLayers.Insert(NewLayer, DuplicateIndex);
-			//@TODO: Try renaming it to a better name (Photoshop uses "[OldName] copy", but with a prompt to let you rename it immediately)
+			NewLayer->LayerName = GenerateDuplicatedLayerName(NewLayer->LayerName.ToString(), TileMap);
 
 			TileMap->PostEditChange();
 			ListViewWidget->RequestListRefresh();
diff --git a/Engine/Plugins/2D/Paper2D/Source/Paper2DEditor/Private/TileMapEditing/STileLayerList.h b/Engine/Plugins/2D/Paper2D/Source/Paper2DEditor/Private/TileMapEditing/STileLayerList.h
--- a/Engine/Plugins/2D/Paper2D/Source/Paper2DEditor/Private/TileMapEditing/STileLayerList.h
+++ b/Engine/Plugins/2D/Paper2D/Source/Paper2DEditor/Private/TileMapEditing/STileLayerList.h
@@ -29,6 +29,9 @@ protected:
 	// Returns the selected index if anything is selected, or the top item otherwise (only returns INDEX_NONE if there are no layers)
 	int32 GetSelectionIndex() const;
 
+	static FText GenerateNewLayerName(UPaperTileMap* TileMap);
+	static FText GenerateDuplicatedLayerName(const FString& InputNameRaw, UPaperTileMap* TileMap);
+
 	class UPaperTileLayer* AddLayer(bool bCollisionLayer, int32 InsertionIndex = INDEX_NONE);
 
 	// Moves a layer from OldIndex to NewIndex if both are valid, otherwise it does nothing silently
