Paper2D: Gate layer browser actions based on the selection

diff --git a/Engine/Plugins/2D/Paper2D/Source/Paper2DEditor/Private/TileMapEditing/STileLayerList.cpp b/Engine/Plugins/2D/Paper2D/Source/Paper2DEditor/Private/TileMapEditing/STileLayerList.cpp
--- a/Engine/Plugins/2D/Paper2D/Source/Paper2DEditor/Private/TileMapEditing/STileLayerList.cpp
+++ b/Engine/Plugins/2D/Paper2D/Source/Paper2DEditor/Private/TileMapEditing/STileLayerList.cpp
@@ -36,24 +36,29 @@ void STileLayerList::Construct(const FArguments& InArgs, UPaperTileMap* TileMap)
 
 	CommandList->MapAction(
 		Commands.DeleteLayer,
-		FExecuteAction::CreateSP(this, &STileLayerList::DeleteLayer));
+		FExecuteAction::CreateSP(this, &STileLayerList::DeleteLayer),
+		FCanExecuteAction::CreateSP(this, &STileLayerList::CanExecuteActionNeedingSelectedLayer));
 
 	CommandList->MapAction(
 		Commands.DuplicateLayer,
-		FExecuteAction::CreateSP(this, &STileLayerList::DuplicateLayer));
+		FExecuteAction::CreateSP(this, &STileLayerList::DuplicateLayer),
+		FCanExecuteAction::CreateSP(this, &STileLayerList::CanExecuteActionNeedingSelectedLayer));
 
 	CommandList->MapAction(
 		Commands.MergeLayerDown,
-		FExecuteAction::CreateSP(this, &STileLayerList::MergeLayerDown));
+		FExecuteAction::CreateSP(this, &STileLayerList::MergeLayerDown),
+		FCanExecuteAction::CreateSP(this, &STileLayerList::CanExecuteActionNeedingLayerBelow));
 
 	CommandList->MapAction(
 		Commands.MoveLayerUp,
-		FExecuteAction::CreateSP(this, &STileLayerList::MoveLayerUp));
+		FExecuteAction::CreateSP(this, &STileLayerList::MoveLayerUp),
+		FCanExecuteAction::CreateSP(this, &STileLayerList::CanExecuteActionNeedingLayerAbove));
 
 	CommandList->MapAction(
 		Commands.MoveLayerDown,
-		FExecuteAction::CreateSP(this, &STileLayerList::MoveLayerDown));
-
+		FExecuteAction::CreateSP(this, &STileLayerList::MoveLayerDown),
+		FCanExecuteAction::CreateSP(this, &STileLayerList::CanExecuteActionNeedingLayerBelow));
+	
 	FToolBarBuilder ToolbarBuilder(CommandList, FMultiBoxCustomization("TileLayerBrowserToolbar"), TSharedPtr<FExtender>(), Orient_Horizontal, /*InForceSmallIcons=*/ true);
 	ToolbarBuilder.SetLabelVisibility(EVisibility::Collapsed);
 
@@ -352,6 +357,31 @@ void STileLayerList::MoveLayerDown()
 	ChangeLayerOrdering(SelectedIndex, NewIndex);
 }
 
+int32 STileLayerList::GetNumLayers() const
+{
+	if (UPaperTileMap* TileMap = TileMapPtr.Get())
+	{
+		return TileMap->TileLayers.Num();
+	}
+	return 0;
+}
+
+bool STileLayerList::CanExecuteActionNeedingLayerAbove() const
+{
+	return GetSelectionIndex() > 0;
+}
+
+bool STileLayerList::CanExecuteActionNeedingLayerBelow() const
+{
+	const int32 SelectedLayer = GetSelectionIndex();
+	return (SelectedLayer != INDEX_NONE) && (SelectedLayer + 1 < GetNumLayers());
+}
+
+bool STileLayerList::CanExecuteActionNeedingSelectedLayer() const
+{
+	return GetSelectionIndex() != INDEX_NONE;
+}
+
 void STileLayerList::SetSelectedLayer(UPaperTileLayer* SelectedLayer)
 {
 	ListViewWidget->SetSelection(SelectedLayer);
diff --git a/Engine/Plugins/2D/Paper2D/Source/Paper2DEditor/Private/TileMapEditing/STileLayerList.h b/Engine/Plugins/2D/Paper2D/Source/Paper2DEditor/Private/TileMapEditing/STileLayerList.h
--- a/Engine/Plugins/2D/Paper2D/Source/Paper2DEditor/Private/TileMapEditing/STileLayerList.h
+++ b/Engine/Plugins/2D/Paper2D/Source/Paper2DEditor/Private/TileMapEditing/STileLayerList.h
@@ -45,5 +45,10 @@ protected:
 	void MoveLayerUp();
 	void MoveLayerDown();
 
+	int32 GetNumLayers() const;
+	bool CanExecuteActionNeedingLayerAbove() const;
+	bool CanExecuteActionNeedingLayerBelow() const;
+	bool CanExecuteActionNeedingSelectedLayer() const;
+
 	void SetSelectedLayer(UPaperTileLayer* SelectedLayer);
 };
