Paper2D: Tile map layer browser styling changes
- Remove list header
- Preserve selection when clicking on empty space
- Enforce a fixed size
- Add a context menu

diff --git a/Engine/Plugins/2D/Paper2D/Source/Paper2DEditor/Private/TileMapEditing/STileLayerItem.cpp b/Engine/Plugins/2D/Paper2D/Source/Paper2DEditor/Private/TileMapEditing/STileLayerItem.cpp
--- a/Engine/Plugins/2D/Paper2D/Source/Paper2DEditor/Private/TileMapEditing/STileLayerItem.cpp
+++ b/Engine/Plugins/2D/Paper2D/Source/Paper2DEditor/Private/TileMapEditing/STileLayerItem.cpp
@@ -4,7 +4,6 @@
 #include "STileLayerItem.h"
 #include "SContentReference.h"
 #include "ScopedTransaction.h"
-#include "PaperStyle.h"
 
 #define LOCTEXT_NAMESPACE "Paper2D"
 
@@ -149,56 +148,44 @@ EVisibility	SRenamableEntry::GetNonEditModeVisibility() const
 //////////////////////////////////////////////////////////////////////////
 // STileLayerItem
 
-void STileLayerItem::Construct(const FArguments& InArgs, class UPaperTileLayer* InItem, const TSharedRef<STableViewBase>& OwnerTable)
+void STileLayerItem::Construct(const FArguments& InArgs, class UPaperTileLayer* InItem)
 {
 	MyLayer = InItem;
 
-	auto ParentArgs = FSuperRowType::FArguments()
-		.Style(&FPaperStyle::Get()->GetWidgetStyle<FTableRowStyle>("TileMapEditor.LayerBrowser.TableViewRow"));
-
-	SMultiColumnTableRow<class UPaperTileLayer*>::Construct(ParentArgs, OwnerTable);
-
 	EyeClosed = FEditorStyle::GetBrush("Layer.NotVisibleIcon16x");
 	EyeOpened = FEditorStyle::GetBrush("Layer.VisibleIcon16x");
-}
 
-TSharedRef<SWidget> STileLayerItem::GenerateWidgetForColumn(const FName& ColumnName) 
-{
-	if (ColumnName == TEXT("Name"))
-	{
-		return SNew(SHorizontalBox)
-			+SHorizontalBox::Slot()
-			.AutoWidth()
-			.VAlign(VAlign_Center)
+	ChildSlot
+	[
+		SNew(SHorizontalBox)
+		+SHorizontalBox::Slot()
+		.AutoWidth()
+		.VAlign(VAlign_Center)
+		[
+			SAssignNew( VisibilityButton, SButton )
+			.ContentPadding(FMargin(4.0f, 4.0f, 4.0f, 4.0f))
+			.ButtonStyle( FEditorStyle::Get(), "NoBorder" )
+			.OnClicked( this, &STileLayerItem::OnToggleVisibility )
+			.ToolTipText( LOCTEXT("LayerVisibilityButtonToolTip", "Toggle Layer Visibility") )
+			.ForegroundColor( FSlateColor::UseForeground() )
+			.HAlign( HAlign_Center )
+			.VAlign( VAlign_Center )
+			.Content()
 			[
-				SAssignNew( VisibilityButton, SButton )
-				.ContentPadding(FMargin(4.0f, 4.0f, 4.0f, 4.0f))
-				.ButtonStyle( FEditorStyle::Get(), "NoBorder" )
-				.OnClicked( this, &STileLayerItem::OnToggleVisibility )
-				.ToolTipText( LOCTEXT("LayerVisibilityButtonToolTip", "Toggle Layer Visibility") )
-				.ForegroundColor( FSlateColor::UseForeground() )
-				.HAlign( HAlign_Center )
-				.VAlign( VAlign_Center )
-				.Content()
-				[
-					SNew(SImage)
-					.Image(this, &STileLayerItem::GetVisibilityBrushForLayer)
-					.ColorAndOpacity(this, &STileLayerItem::GetForegroundColorForVisibilityButton)
-				]
+				SNew(SImage)
+				.Image(this, &STileLayerItem::GetVisibilityBrushForLayer)
+				.ColorAndOpacity(this, &STileLayerItem::GetForegroundColorForVisibilityButton)
 			]
-			+SHorizontalBox::Slot()
-			.VAlign(VAlign_Center)
-			.Padding(FMargin(4.0f, 4.0f, 4.0f, 4.0f))
-			[
-				SNew(SRenamableEntry)
-				.Text(this, &STileLayerItem::GetLayerDisplayName)
-				.OnTextCommitted(this, &STileLayerItem::OnLayerNameCommitted)
-			];
-	}
-	else
-	{
-		return SNew(SSpacer);
-	}
+		]
+		+SHorizontalBox::Slot()
+		.VAlign(VAlign_Center)
+		.Padding(FMargin(4.0f, 4.0f, 4.0f, 4.0f))
+		[
+			SNew(SRenamableEntry)
+			.Text(this, &STileLayerItem::GetLayerDisplayName)
+			.OnTextCommitted(this, &STileLayerItem::OnLayerNameCommitted)
+		]
+	];
 }
 
 FText STileLayerItem::GetLayerDisplayName() const
diff --git a/Engine/Plugins/2D/Paper2D/Source/Paper2DEditor/Private/TileMapEditing/STileLayerItem.h b/Engine/Plugins/2D/Paper2D/Source/Paper2DEditor/Private/TileMapEditing/STileLayerItem.h
--- a/Engine/Plugins/2D/Paper2D/Source/Paper2DEditor/Private/TileMapEditing/STileLayerItem.h
+++ b/Engine/Plugins/2D/Paper2D/Source/Paper2DEditor/Private/TileMapEditing/STileLayerItem.h
@@ -5,17 +5,14 @@
 //////////////////////////////////////////////////////////////////////////
 // STileLayerItem
 
-class STileLayerItem : public SMultiColumnTableRow<class UPaperTileLayer*>
+class STileLayerItem : public SCompoundWidget
 {
 public:
 	SLATE_BEGIN_ARGS(STileLayerItem) {}
 	SLATE_END_ARGS()
 
-	void Construct(const FArguments& InArgs, class UPaperTileLayer* InItem, const TSharedRef<STableViewBase>& OwnerTable);
+	void Construct(const FArguments& InArgs, class UPaperTileLayer* InItem);
 
-	// SMultiColumnTableRow<> interface
-	virtual TSharedRef<SWidget> GenerateWidgetForColumn(const FName& ColumnName) override;
-	// End of SMultiColumnTableRow<> interface
 protected:
 	class UPaperTileLayer* MyLayer;
 
diff --git a/Engine/Plugins/2D/Paper2D/Source/Paper2DEditor/Private/TileMapEditing/STileLayerList.cpp b/Engine/Plugins/2D/Paper2D/Source/Paper2DEditor/Private/TileMapEditing/STileLayerList.cpp
--- a/Engine/Plugins/2D/Paper2D/Source/Paper2DEditor/Private/TileMapEditing/STileLayerList.cpp
+++ b/Engine/Plugins/2D/Paper2D/Source/Paper2DEditor/Private/TileMapEditing/STileLayerList.cpp
@@ -3,6 +3,7 @@
 #include "Paper2DEditorPrivatePCH.h"
 #include "STileLayerList.h"
 #include "STileLayerItem.h"
+#include "PaperStyle.h"
 
 #include "ScopedTransaction.h"
 
@@ -15,10 +16,6 @@
 
 void STileLayerList::Construct(const FArguments& InArgs, UPaperTileMap* TileMap)
 {
-	TSharedRef<SHeaderRow> HeaderRowWidget = SNew(SHeaderRow)
-		+SHeaderRow::Column("Name").DefaultLabel(NSLOCTEXT("TileLayerList", "TileLayerNameHeader", "Layer Name")).FillWidth(0.7);
-		//+SHeaderRow::Column("TileSet").DefaultLabel(NSLOCTEXT("TileLayerList", "TileLayerTileSetHeader", "Tile Set")).FillWidth(0.7);
-
 	TileMapPtr = TileMap;
 
 	FTileMapEditorCommands::Register();
@@ -72,22 +69,11 @@ void STileLayerList::Construct(const FArguments& InArgs, UPaperTileMap* TileMap)
 
 	ListViewWidget = SNew(SPaperLayerListView)
 		.SelectionMode(ESelectionMode::Single)
-
-		// Point the tree to our array of root-level items.  Whenever this changes, we'll call RequestTreeRefresh()
+		.ClearSelectionOnClick(false)
 		.ListItemsSource(&(TileMap->TileLayers))
-
-		// Find out when the user selects something in the tree
-		//.OnSelectionChanged( this, &SLevelsView::OnSelectionChanged )
-
-		// Called when the user double-clicks with LMB on an item in the list
-		//.OnMouseButtonDoubleClick( this, &SLevelsView::OnListViewMouseButtonDoubleClick )
-
-		// Generates the actual widget for a tree item
-		.OnGenerateRow(this, &STileLayerList::OnGenerateRowDefault)
-
-		// Use the level viewport context menu as the right click menu for list items
-		//.OnContextMenuOpening( InArgs._ConstructContextMenu )
-		.HeaderRow(HeaderRowWidget);
+		.OnSelectionChanged(this, &STileLayerList::OnSelectionChanged)
+		.OnGenerateRow(this, &STileLayerList::OnGenerateLayerListRow)
+		.OnContextMenuOpening(this, &STileLayerList::OnConstructContextMenu);
 
 	// Select the top item by default
 	if (TileMap->TileLayers.Num() > 0)
@@ -97,11 +83,14 @@ void STileLayerList::Construct(const FArguments& InArgs, UPaperTileMap* TileMap)
 
 	ChildSlot
 	[
-		SNew( SVerticalBox )
+		SNew(SVerticalBox)
 		+SVerticalBox::Slot()
-		.FillHeight( 1.0f )
 		[
-			ListViewWidget.ToSharedRef()
+			SNew(SBox)
+			.HeightOverride(115.0f)
+			[
+				ListViewWidget.ToSharedRef()
+			]
 		]
 		+SVerticalBox::Slot()
 		.AutoHeight()
@@ -111,9 +100,13 @@ void STileLayerList::Construct(const FArguments& InArgs, UPaperTileMap* TileMap)
 	];
 }
 
-TSharedRef<ITableRow> STileLayerList::OnGenerateRowDefault(class UPaperTileLayer* Item, const TSharedRef<STableViewBase>& OwnerTable)
+TSharedRef<ITableRow> STileLayerList::OnGenerateLayerListRow(class UPaperTileLayer* Item, const TSharedRef<STableViewBase>& OwnerTable)
 {
-	return SNew(STileLayerItem, Item, OwnerTable);		
+	return SNew(STableRow<class UPaperTileLayer*>, OwnerTable)
+		.Style(&FPaperStyle::Get()->GetWidgetStyle<FTableRowStyle>("TileMapEditor.LayerBrowser.TableViewRow"))
+		[
+			SNew(STileLayerItem, Item)
+		];
 }
 
 UPaperTileLayer* STileLayerList::GetSelectedLayer() const
@@ -387,6 +380,28 @@ void STileLayerList::SetSelectedLayer(UPaperTileLayer* SelectedLayer)
 	ListViewWidget->SetSelection(SelectedLayer);
 }
 
+void STileLayerList::OnSelectionChanged(UPaperTileLayer* ItemChangingState, ESelectInfo::Type SelectInfo)
+{
+}
+
+TSharedPtr<SWidget> STileLayerList::OnConstructContextMenu()
+{
+	const bool bShouldCloseWindowAfterMenuSelection = true;
+	FMenuBuilder MenuBuilder(bShouldCloseWindowAfterMenuSelection, CommandList);
+
+	const FTileMapEditorCommands& Commands = FTileMapEditorCommands::Get();
+
+ 	MenuBuilder.BeginSection("BasicOperations");
+ 	{
+		MenuBuilder.AddMenuEntry(Commands.DuplicateLayer);
+		MenuBuilder.AddMenuEntry(Commands.DeleteLayer);
+		MenuBuilder.AddMenuEntry(Commands.MergeLayerDown);
+ 	}
+	MenuBuilder.EndSection();
+
+	return MenuBuilder.MakeWidget();
+}
+
 //////////////////////////////////////////////////////////////////////////
 
 #undef LOCTEXT_NAMESPACE
\ No newline at end of file
diff --git a/Engine/Plugins/2D/Paper2D/Source/Paper2DEditor/Private/TileMapEditing/STileLayerList.h b/Engine/Plugins/2D/Paper2D/Source/Paper2DEditor/Private/TileMapEditing/STileLayerList.h
--- a/Engine/Plugins/2D/Paper2D/Source/Paper2DEditor/Private/TileMapEditing/STileLayerList.h
+++ b/Engine/Plugins/2D/Paper2D/Source/Paper2DEditor/Private/TileMapEditing/STileLayerList.h
@@ -22,7 +22,7 @@ protected:
 	TWeakObjectPtr<class UPaperTileMap> TileMapPtr;
 
 protected:
-	TSharedRef<ITableRow> OnGenerateRowDefault(class UPaperTileLayer* Item, const TSharedRef<STableViewBase>& OwnerTable);
+	TSharedRef<ITableRow> OnGenerateLayerListRow(class UPaperTileLayer* Item, const TSharedRef<STableViewBase>& OwnerTable);
 
 	class UPaperTileLayer* GetSelectedLayer() const;
 
@@ -51,4 +51,7 @@ protected:
 	bool CanExecuteActionNeedingSelectedLayer() const;
 
 	void SetSelectedLayer(UPaperTileLayer* SelectedLayer);
+
+	void OnSelectionChanged(UPaperTileLayer* ItemChangingState, ESelectInfo::Type SelectInfo);
+	TSharedPtr<SWidget> OnConstructContextMenu();
 };
