Paper2D: Fix bounding volume on tile map components to encompass more than one layer

diff --git a/Engine/Plugins/2D/Paper2D/Source/Paper2D/Private/PaperTileMap.cpp b/Engine/Plugins/2D/Paper2D/Source/Paper2D/Private/PaperTileMap.cpp
--- a/Engine/Plugins/2D/Paper2D/Source/Paper2D/Private/PaperTileMap.cpp
+++ b/Engine/Plugins/2D/Paper2D/Source/Paper2D/Private/PaperTileMap.cpp
@@ -116,10 +116,10 @@ FVector UPaperTileMap::GetTilePositionInLocalSpace(int32 TileX, int32 TileY, int
 FBoxSphereBounds UPaperTileMap::GetRenderBounds() const
 {
 	//@TODO: Tile pivot issue
-	//@TODO: Layer thickness issue
+	const float Depth = SeparationPerLayer * (TileLayers.Num() - 1);
 	const float HalfThickness = 2.0f;
-	const FVector TopLeft((-0.5f)*TileWidth, -HalfThickness, -(MapHeight - 0.5f) * TileHeight);
-	const FVector Dimenisons(MapWidth*TileWidth, 2 * HalfThickness, MapHeight * TileHeight);
+	const FVector TopLeft((-0.5f)*TileWidth, -HalfThickness - Depth, -(MapHeight - 0.5f) * TileHeight);
+	const FVector Dimenisons(MapWidth*TileWidth, Depth + 2 * HalfThickness, MapHeight * TileHeight);
 
 	const FBox Box(TopLeft, TopLeft + Dimenisons);
 	return FBoxSphereBounds(Box);
