Paper2D: Switch tile map layer name editor to use SInlineEditableTextBlock, improving the styling and behavior of editing layer names

diff --git a/Engine/Plugins/2D/Paper2D/Source/Paper2DEditor/Private/TileMapEditing/STileLayerItem.cpp b/Engine/Plugins/2D/Paper2D/Source/Paper2DEditor/Private/TileMapEditing/STileLayerItem.cpp
--- a/Engine/Plugins/2D/Paper2D/Source/Paper2DEditor/Private/TileMapEditing/STileLayerItem.cpp
+++ b/Engine/Plugins/2D/Paper2D/Source/Paper2DEditor/Private/TileMapEditing/STileLayerItem.cpp
@@ -4,151 +4,14 @@
 #include "STileLayerItem.h"
 #include "SContentReference.h"
 #include "ScopedTransaction.h"
+#include "SInlineEditableTextBlock.h"
 
 #define LOCTEXT_NAMESPACE "Paper2D"
 
-//////////////////////////////////////////////////////////////////////////
-// SRenamableEntry
-
-class SRenamableEntry : public SCompoundWidget
-{
-public:
-	SLATE_BEGIN_ARGS(SRenamableEntry) {}
-		SLATE_ATTRIBUTE(FText, Text)
-		SLATE_EVENT(FOnTextCommitted, OnTextCommitted)
-	SLATE_END_ARGS()
-
-	void Construct(const FArguments& InArgs);
-
-	// SWidget interface
-	virtual FReply OnMouseButtonDoubleClick(const FGeometry& InMyGeometry, const FPointerEvent& InMouseEvent) override;
-	virtual bool SupportsKeyboardFocus() const override;
-	// End of SWidget interface
-protected:
-	bool bInEditMode;
-
-	/** Editable text object used when node is being renamed */
-	TSharedPtr<SEditableText> EditableTextWidget;
-
-	TAttribute<FText> MyText;
-
-	FOnTextCommitted OnTextCommittedEvent;
-protected:
-	/* Called to check the visibility of the editable text */
-	EVisibility	GetEditModeVisibility() const;
-
-	/* Called to check the visibility of the non-editable text */
-	EVisibility	GetNonEditModeVisibility() const;
-
-	FText GetEditableText() const;
-	FString GetNonEditableText() const;
-
-	void OnTextChanged(const FText& InNewText);
-	void OnTextCommited(const FText& InText, ETextCommit::Type CommitInfo);
-};
-
-FText SRenamableEntry::GetEditableText() const
-{
-	return MyText.Get();
-}
-
-FString SRenamableEntry::GetNonEditableText() const
-{
-	return MyText.Get().ToString();
-}
-
-void SRenamableEntry::OnTextChanged(const FText& InNewText)
-{
-	//@TODO:	OnTextChanged.ExecuteIfBound(InNewText);
-
-	//UpdateErrorInfo();
-	//ErrorReporting->SetError(ErrorMsg);
-}
-
-void SRenamableEntry::OnTextCommited(const FText& InText, ETextCommit::Type CommitInfo) 
-{
-	OnTextCommittedEvent.ExecuteIfBound(InText, CommitInfo);
-
-	//UpdateErrorInfo();
-	//ErrorReporting->SetError(ErrorMsg);
-
-	// lost focus so exit rename mode
-	bInEditMode = false;
-}
-
-FReply SRenamableEntry::OnMouseButtonDoubleClick(const FGeometry& InMyGeometry, const FPointerEvent& InMouseEvent)
-{
-	if (!bInEditMode)
-	{
-		bInEditMode = true;
-
-		FWidgetPath TextBoxPath;
-		FSlateApplication::Get().GeneratePathToWidgetUnchecked(EditableTextWidget.ToSharedRef(), TextBoxPath);
-		FSlateApplication::Get().SetKeyboardFocus(EditableTextWidget, EFocusCause::Mouse);
-
-		return FReply::Handled();
-	}
-	else
-	{
-		return FReply::Unhandled();
-	}
-}
-
-bool SRenamableEntry::SupportsKeyboardFocus() const
-{
-	return (bInEditMode && EditableTextWidget.IsValid())
-		? StaticCastSharedPtr<SWidget>(EditableTextWidget)->SupportsKeyboardFocus() 
-		: false;
-}
-
-void SRenamableEntry::Construct(const FArguments& InArgs)
-{
-	bInEditMode = false;
-
-	MyText = InArgs._Text;
-	OnTextCommittedEvent = InArgs._OnTextCommitted;
-
-	//@TODO: Add an error reporting area / text validation
-
-	ChildSlot
-	[
-		SNew(SHorizontalBox)
-		+SHorizontalBox::Slot()
-		.AutoWidth()
-		[
-			SAssignNew(EditableTextWidget, SEditableText)
-			.Text(this, &SRenamableEntry::GetEditableText)
-			.OnTextChanged(this, &SRenamableEntry::OnTextChanged)
-			.OnTextCommitted(this, &SRenamableEntry::OnTextCommited)
-			.SelectAllTextWhenFocused(true)
-			.ClearKeyboardFocusOnCommit(true)
-			.RevertTextOnEscape(true)
-			.Visibility(this, &SRenamableEntry::GetEditModeVisibility)
-		]
-		+SHorizontalBox::Slot()
-		.AutoWidth()
-		[
-			SNew(STextBlock)
-			.Visibility(this, &SRenamableEntry::GetNonEditModeVisibility)
-			.Text(this, &SRenamableEntry::GetNonEditableText)
-		]
-	];
-}
-
-EVisibility	SRenamableEntry::GetEditModeVisibility() const
-{
-	return bInEditMode ? EVisibility::Visible : EVisibility::Collapsed;
-}
-
-EVisibility	SRenamableEntry::GetNonEditModeVisibility() const
-{
-	return bInEditMode ? EVisibility::Collapsed : EVisibility::Visible;
-}
-
 //////////////////////////////////////////////////////////////////////////
 // STileLayerItem
 
-void STileLayerItem::Construct(const FArguments& InArgs, class UPaperTileLayer* InItem)
+void STileLayerItem::Construct(const FArguments& InArgs, class UPaperTileLayer* InItem, FIsSelected InIsSelectedDelegate)
 {
 	MyLayer = InItem;
 
@@ -181,9 +44,11 @@ void STileLayerItem::Construct(const FArguments& InArgs, class UPaperTileLayer*
 		.VAlign(VAlign_Center)
 		.Padding(FMargin(4.0f, 4.0f, 4.0f, 4.0f))
 		[
-			SNew(SRenamableEntry)
+			SNew(SInlineEditableTextBlock)
 			.Text(this, &STileLayerItem::GetLayerDisplayName)
+			.ToolTipText(this, &STileLayerItem::GetLayerDisplayName)
 			.OnTextCommitted(this, &STileLayerItem::OnLayerNameCommitted)
+			.IsSelected(InIsSelectedDelegate)
 		]
 	];
 }
diff --git a/Engine/Plugins/2D/Paper2D/Source/Paper2DEditor/Private/TileMapEditing/STileLayerItem.h b/Engine/Plugins/2D/Paper2D/Source/Paper2DEditor/Private/TileMapEditing/STileLayerItem.h
--- a/Engine/Plugins/2D/Paper2D/Source/Paper2DEditor/Private/TileMapEditing/STileLayerItem.h
+++ b/Engine/Plugins/2D/Paper2D/Source/Paper2DEditor/Private/TileMapEditing/STileLayerItem.h
@@ -11,7 +11,7 @@ public:
 	SLATE_BEGIN_ARGS(STileLayerItem) {}
 	SLATE_END_ARGS()
 
-	void Construct(const FArguments& InArgs, class UPaperTileLayer* InItem);
+	void Construct(const FArguments& InArgs, class UPaperTileLayer* InItem, FIsSelected InIsSelectedDelegate);
 
 protected:
 	class UPaperTileLayer* MyLayer;
diff --git a/Engine/Plugins/2D/Paper2D/Source/Paper2DEditor/Private/TileMapEditing/STileLayerList.cpp b/Engine/Plugins/2D/Paper2D/Source/Paper2DEditor/Private/TileMapEditing/STileLayerList.cpp
--- a/Engine/Plugins/2D/Paper2D/Source/Paper2DEditor/Private/TileMapEditing/STileLayerList.cpp
+++ b/Engine/Plugins/2D/Paper2D/Source/Paper2DEditor/Private/TileMapEditing/STileLayerList.cpp
@@ -102,11 +102,15 @@ void STileLayerList::Construct(const FArguments& InArgs, UPaperTileMap* TileMap)
 
 TSharedRef<ITableRow> STileLayerList::OnGenerateLayerListRow(class UPaperTileLayer* Item, const TSharedRef<STableViewBase>& OwnerTable)
 {
-	return SNew(STableRow<class UPaperTileLayer*>, OwnerTable)
-		.Style(&FPaperStyle::Get()->GetWidgetStyle<FTableRowStyle>("TileMapEditor.LayerBrowser.TableViewRow"))
-		[
-			SNew(STileLayerItem, Item)
-		];
+	typedef STableRow<class UPaperTileLayer*> RowType;
+
+	TSharedRef<RowType> NewRow = SNew(RowType, OwnerTable)
+		.Style(&FPaperStyle::Get()->GetWidgetStyle<FTableRowStyle>("TileMapEditor.LayerBrowser.TableViewRow"));
+
+	FIsSelected IsSelectedDelegate = FIsSelected::CreateSP(NewRow, &RowType::IsSelectedExclusively);
+	NewRow->SetContent(SNew(STileLayerItem, Item, IsSelectedDelegate));
+
+	return NewRow;
 }
 
 UPaperTileLayer* STileLayerList::GetSelectedLayer() const
