Paper2D: Remove some debug display annotations on tile layer properties

diff --git a/Engine/Plugins/2D/Paper2D/Source/Paper2D/Classes/PaperTileLayer.h b/Engine/Plugins/2D/Paper2D/Source/Paper2D/Classes/PaperTileLayer.h
--- a/Engine/Plugins/2D/Paper2D/Source/Paper2D/Classes/PaperTileLayer.h
+++ b/Engine/Plugins/2D/Paper2D/Source/Paper2D/Classes/PaperTileLayer.h
@@ -11,10 +11,10 @@ struct FPaperTileInfo
 {
 	GENERATED_USTRUCT_BODY()
 
-	UPROPERTY(Category=Debug, VisibleAnywhere)
+	UPROPERTY()
 	UPaperTileSet* TileSet;
 
-	UPROPERTY(Category=Debug, VisibleAnywhere)
+	UPROPERTY()
 	int32 PackedTileIndex;
 
 	FPaperTileInfo()
@@ -40,46 +40,46 @@ class PAPER2D_API UPaperTileLayer : public UObject
 	GENERATED_UCLASS_BODY()
 
 	// Name of the layer
-	UPROPERTY(Category=Debug, VisibleAnywhere)
+	UPROPERTY()
 	FText LayerName;
 
 	// Width of the layer (in tiles)
-	UPROPERTY(Category=Debug, VisibleAnywhere)
+	UPROPERTY()
 	int32 LayerWidth;
 
 	// Height of the layer (in tiles)
-	UPROPERTY(Category=Debug, VisibleAnywhere)
+	UPROPERTY()
 	int32 LayerHeight;
 
 #if WITH_EDITORONLY_DATA
 	// Opacity of the layer (editor only)
-	UPROPERTY(Category=Debug, VisibleAnywhere)
+	UPROPERTY()
 	float LayerOpacity;
 
 	// Is this layer currently hidden in the editor?
-	UPROPERTY(Category=Debug, VisibleAnywhere)
+	UPROPERTY()
 	bool bHiddenInEditor;
 #endif
 
 	// Should this layer be hidden in the game?
-	UPROPERTY(Category=Debug, VisibleAnywhere)
+	UPROPERTY()
 	bool bHiddenInGame;
 
 	// Is this layer a collision layer?
-	UPROPERTY(Category=Debug, VisibleAnywhere)
+	UPROPERTY()
 	bool bCollisionLayer;
 
 protected:
 	// The allocated width of the tile data (used to handle resizing without data loss)
-	UPROPERTY(Category=Debug, VisibleAnywhere)
+	UPROPERTY()
 	int32 AllocatedWidth;
 
 	// The allocated height of the tile data (used to handle resizing without data loss)
-	UPROPERTY(Category=Debug, VisibleAnywhere)
+	UPROPERTY()
 	int32 AllocatedHeight;
 
 	// The allocated tile data
-	UPROPERTY(Category=Debug, VisibleAnywhere)
+	UPROPERTY()
 	TArray<FPaperTileInfo> AllocatedCells;
 
 private:
diff --git a/Engine/Plugins/2D/Paper2D/Source/Paper2D/Classes/PaperTileMap.h b/Engine/Plugins/2D/Paper2D/Source/Paper2D/Classes/PaperTileMap.h
--- a/Engine/Plugins/2D/Paper2D/Source/Paper2D/Classes/PaperTileMap.h
+++ b/Engine/Plugins/2D/Paper2D/Source/Paper2D/Classes/PaperTileMap.h
@@ -69,7 +69,7 @@ class PAPER2D_API UPaperTileMap : public UDataAsset //@TODO: Just to make it eas
 	UMaterialInterface* Material;
 
 	// The list of layers
-	UPROPERTY(Category=Debug, VisibleAnywhere, Instanced) //@TODO: Remove this
+	UPROPERTY(Instanced)
 	TArray<class UPaperTileLayer*> TileLayers;
 
 	// Collision domain (no collision, 2D, or 3D)
