Paper2D: Tile Maps: Fix cursor preview not respecting layer depth when editing something other than the top one (still always edits the topmost non-hidden layer)

diff --git a/Engine/Plugins/2D/Paper2D/Source/Paper2D/Private/PaperTileMap.cpp b/Engine/Plugins/2D/Paper2D/Source/Paper2D/Private/PaperTileMap.cpp
--- a/Engine/Plugins/2D/Paper2D/Source/Paper2D/Private/PaperTileMap.cpp
+++ b/Engine/Plugins/2D/Paper2D/Source/Paper2D/Private/PaperTileMap.cpp
@@ -106,9 +106,12 @@ void UPaperTileMap::UpdateBodySetup()
 
 FVector UPaperTileMap::GetTilePositionInLocalSpace(int32 TileX, int32 TileY, int32 LayerIndex) const
 {
-	//@TODO: Layer depth, multiply by layer depth here!
 	//@TODO: Tile pivot issue
-	const FVector LocalPos(PaperAxisX * (TileX - 0.5f) * TileWidth + PaperAxisY * -(TileY - 0.5f) * TileHeight);
+	const FVector PartialX = PaperAxisX * (TileX - 0.5f) * TileWidth;
+	const FVector PartialY = PaperAxisY * -(TileY - 0.5f) * TileHeight;
+	const FVector PartialZ = PaperAxisZ * (LayerIndex * SeparationPerLayer);
+
+	const FVector LocalPos(PartialX + PartialY + PartialZ);
 	
 	return LocalPos;
 }
diff --git a/Engine/Plugins/2D/Paper2D/Source/Paper2DEditor/Private/TileMapEditing/EdModeTileMap.cpp b/Engine/Plugins/2D/Paper2D/Source/Paper2DEditor/Private/TileMapEditing/EdModeTileMap.cpp
--- a/Engine/Plugins/2D/Paper2D/Source/Paper2DEditor/Private/TileMapEditing/EdModeTileMap.cpp
+++ b/Engine/Plugins/2D/Paper2D/Source/Paper2DEditor/Private/TileMapEditing/EdModeTileMap.cpp
@@ -364,10 +364,12 @@ UPaperTileLayer* FEdModeTileMap::GetSelectedLayerUnderCursor(const FViewportCurs
 				FVector LocalStart = ComponentToWorld.InverseTransformPosition(TraceStart);
 				FVector LocalDirection = ComponentToWorld.InverseTransformVector(TraceDir);
 
-				FVector Intersection;
-				FPlane Plane(FVector(1, 0, 0), FVector::ZeroVector, FVector(0, 0, 1));
+				const FVector LSPlaneCorner = PaperAxisZ * TileMap->SeparationPerLayer;
+
+				const FPlane LayerPlane(LSPlaneCorner + PaperAxisX, LSPlaneCorner, LSPlaneCorner + PaperAxisY);
 
-				if (FMath::SegmentPlaneIntersection(LocalStart, LocalDirection * HALF_WORLD_MAX, Plane, /*out*/ Intersection))
+				FVector Intersection;
+				if (FMath::SegmentPlaneIntersection(LocalStart, LocalDirection * HALF_WORLD_MAX, LayerPlane, /*out*/ Intersection))
 				{
 					//@TODO: Ideally tile pivots weren't in the center!
 					const float NormalizedX = (Intersection.X + 0.5f * TileMap->TileWidth) / WX;
@@ -735,8 +737,11 @@ void FEdModeTileMap::UpdatePreviewCursor(const FViewportCursorLocation& Ray)
 			const int32 LocalTileY1 = LocalTileY0 + CursorHeight;
 
 			UPaperTileMap* TileMap = TileLayer->GetTileMap();
-			const FVector WorldPosition = ComponentToWorld.TransformPosition(TileMap->GetTilePositionInLocalSpace(LocalTileX0, LocalTileY0));
-			const FVector WorldPositionBR = ComponentToWorld.TransformPosition(TileMap->GetTilePositionInLocalSpace(LocalTileX1, LocalTileY1));
+			int32 LayerIndex;
+			ensure(TileMap->TileLayers.Find(TileLayer, LayerIndex));
+
+			const FVector WorldPosition = ComponentToWorld.TransformPosition(TileMap->GetTilePositionInLocalSpace(LocalTileX0, LocalTileY0, LayerIndex));
+			const FVector WorldPositionBR = ComponentToWorld.TransformPosition(TileMap->GetTilePositionInLocalSpace(LocalTileX1, LocalTileY1, LayerIndex));
 
 			DrawPreviewSpace = ComponentToWorld;
 			DrawPreviewLocation = (WorldPosition + WorldPositionBR) * 0.5f;
