Paper2D: Tile maps: Add tracking of the 'active' layer for use in editing and debug visualization

diff --git a/Engine/Plugins/2D/Paper2D/Source/Paper2D/Classes/PaperTileMap.h b/Engine/Plugins/2D/Paper2D/Source/Paper2D/Classes/PaperTileMap.h
--- a/Engine/Plugins/2D/Paper2D/Source/Paper2D/Classes/PaperTileMap.h
+++ b/Engine/Plugins/2D/Paper2D/Source/Paper2D/Classes/PaperTileMap.h
@@ -90,6 +90,10 @@ public:
 	UPROPERTY(Category=ImportSettings, VisibleAnywhere, Instanced)
 	class UAssetImportData* AssetImportData;
 
+	/** The currently selected layer index */
+	UPROPERTY()
+	int32 SelectedLayerIndex;
+
 	/** The naming index to start at when trying to create a new layer */
 	UPROPERTY()
 	int32 LayerNameIndex;
@@ -99,6 +103,8 @@ public:
 	// UObject interface
 #if WITH_EDITOR
 	virtual void PostEditChangeProperty(FPropertyChangedEvent& PropertyChangedEvent) override;
+	virtual void PostLoad() override;
+	void ValidateSelectedLayerIndex();
 #endif
 	// End of UObject interface
 
diff --git a/Engine/Plugins/2D/Paper2D/Source/Paper2D/Private/PaperTileMap.cpp b/Engine/Plugins/2D/Paper2D/Source/Paper2D/Private/PaperTileMap.cpp
--- a/Engine/Plugins/2D/Paper2D/Source/Paper2D/Private/PaperTileMap.cpp
+++ b/Engine/Plugins/2D/Paper2D/Source/Paper2D/Private/PaperTileMap.cpp
@@ -21,6 +21,7 @@ UPaperTileMap::UPaperTileMap(const FObjectInitializer& ObjectInitializer)
 	SpriteCollisionDomain = ESpriteCollisionMode::None;
 
 #if WITH_EDITORONLY_DATA
+	SelectedLayerIndex = INDEX_NONE;
 	LayerNameIndex = 1;
 #endif
 
@@ -40,6 +41,8 @@ void UPaperTileMap::PostEditChangeProperty(FPropertyChangedEvent& PropertyChange
 	//@TODO: Determine when these are really needed, as they're seriously expensive!
 	TComponentReregisterContext<UPaperTileMapRenderComponent> ReregisterStaticComponents;
 
+	ValidateSelectedLayerIndex();
+
 	if ((PropertyName == GET_MEMBER_NAME_CHECKED(UPaperTileMap, MapWidth)) || (PropertyName == GET_MEMBER_NAME_CHECKED(UPaperTileMap, MapHeight)))
 	{
 		MapWidth = FMath::Max(MapWidth, 1);
@@ -60,6 +63,34 @@ void UPaperTileMap::PostEditChangeProperty(FPropertyChangedEvent& PropertyChange
 
 	Super::PostEditChangeProperty(PropertyChangedEvent);
 }
+
+void UPaperTileMap::PostLoad()
+{
+	Super::PostLoad();
+	ValidateSelectedLayerIndex();
+}
+
+void UPaperTileMap::ValidateSelectedLayerIndex()
+{
+	if (!TileLayers.IsValidIndex(SelectedLayerIndex))
+	{
+		// Select the top-most visible layer
+		SelectedLayerIndex = INDEX_NONE;
+		for (int32 LayerIndex = 0; (LayerIndex < TileLayers.Num()) && (SelectedLayerIndex == INDEX_NONE); ++LayerIndex)
+		{
+			if (!TileLayers[LayerIndex]->bHiddenInEditor)
+			{
+				SelectedLayerIndex = LayerIndex;
+			}
+		}
+
+		if ((SelectedLayerIndex == INDEX_NONE) && (TileLayers.Num() > 0))
+		{
+			SelectedLayerIndex = 0;
+		}
+	}
+}
+
 #endif
 
 
diff --git a/Engine/Plugins/2D/Paper2D/Source/Paper2DEditor/Private/TileMapEditing/STileLayerList.cpp b/Engine/Plugins/2D/Paper2D/Source/Paper2DEditor/Private/TileMapEditing/STileLayerList.cpp
--- a/Engine/Plugins/2D/Paper2D/Source/Paper2DEditor/Private/TileMapEditing/STileLayerList.cpp
+++ b/Engine/Plugins/2D/Paper2D/Source/Paper2DEditor/Private/TileMapEditing/STileLayerList.cpp
@@ -386,6 +386,10 @@ void STileLayerList::SetSelectedLayer(UPaperTileLayer* SelectedLayer)
 
 void STileLayerList::OnSelectionChanged(UPaperTileLayer* ItemChangingState, ESelectInfo::Type SelectInfo)
 {
+	if (UPaperTileMap* TileMap = TileMapPtr.Get())
+	{
+		TileMap->SelectedLayerIndex = GetSelectionIndex();
+	}
 }
 
 TSharedPtr<SWidget> STileLayerList::OnConstructContextMenu()
