Paper2D: Tile map debug visualization improvements
- Grid is now positioned in Z according to the active selection layer instead of the first visible layer
- Hidden non-selected layers are now visualized as an outline (when the actor owning the component is selected)
- Add a depth bias so grid wires don't z-fight with tiles

diff --git a/Engine/Plugins/2D/Paper2D/Source/Paper2D/Private/PaperTileMapRenderSceneProxy.cpp b/Engine/Plugins/2D/Paper2D/Source/Paper2D/Private/PaperTileMapRenderSceneProxy.cpp
--- a/Engine/Plugins/2D/Paper2D/Source/Paper2D/Private/PaperTileMapRenderSceneProxy.cpp
+++ b/Engine/Plugins/2D/Paper2D/Source/Paper2D/Private/PaperTileMapRenderSceneProxy.cpp
@@ -29,6 +29,9 @@ void FPaperTileMapRenderSceneProxy::GetDynamicMeshElements(const TArray<const FS
 	QUICK_SCOPE_CYCLE_COUNTER(STAT_FPaperTileMapRenderSceneProxy_GetDynamicMeshElements);
 	checkSlow(IsInRenderingThread());
 
+	// Slight depth bias so that the wireframe grid overlay doesn't z-fight with the tiles themselves
+	const float DepthBias = 0.0001f;
+
 	for (int32 ViewIndex = 0; ViewIndex < Views.Num(); ViewIndex++)
 	{
 		if (VisibilityMap & (1 << ViewIndex))
@@ -111,45 +114,49 @@ void FPaperTileMapRenderSceneProxy::GetDynamicMeshElements(const TArray<const FS
 
 					FTransform LocalToWorld(GetLocalToWorld());
 
-					const float TW = TileMap->TileWidth;
-					const float TH = TileMap->TileHeight;
-					//@TODO: PAPER: Tiles probably shouldn't be drawn with their pivot at the center - RE: all the 0.5f in here
-
 					if (bUseOverrideColor)
 					{
-						// Draw horizontal lines
+						// Draw a bound for any invisible layers
+						for (int32 LayerIndex = 0; LayerIndex < TileMap->TileLayers.Num(); ++LayerIndex)
+						{
+							if (LayerIndex != TileMap->SelectedLayerIndex)
+							{
+								const FVector TL(LocalToWorld.TransformPosition(TileMap->GetTilePositionInLocalSpace(0, 0, LayerIndex)));
+								const FVector TR(LocalToWorld.TransformPosition(TileMap->GetTilePositionInLocalSpace(TileMap->MapWidth, 0, LayerIndex)));
+								const FVector BL(LocalToWorld.TransformPosition(TileMap->GetTilePositionInLocalSpace(0, TileMap->MapHeight, LayerIndex)));
+								const FVector BR(LocalToWorld.TransformPosition(TileMap->GetTilePositionInLocalSpace(TileMap->MapWidth, TileMap->MapHeight, LayerIndex)));
+
+								PDI->DrawLine(TL, TR, OverrideColor, DPG, 0.0f, DepthBias);
+								PDI->DrawLine(TR, BR, OverrideColor, DPG, 0.0f, DepthBias);
+								PDI->DrawLine(BR, BL, OverrideColor, DPG, 0.0f, DepthBias);
+								PDI->DrawLine(BL, TL, OverrideColor, DPG, 0.0f, DepthBias);
+							}
+						}
+
+						// Draw horizontal lines on the selection
 						for (int32 Y = 0; Y <= TileMap->MapHeight; ++Y)
 						{
 							int32 X = 0;
-							const FVector Start((X - 0.5f) * TW, 0.0f, -(Y - 0.5f) * TH);
+							const FVector Start(TileMap->GetTilePositionInLocalSpace(X, Y, TileMap->SelectedLayerIndex));
 
 							X = TileMap->MapWidth;
-							const FVector End((X - 0.5f) * TW, 0.0f, -(Y - 0.5f) * TH);
+							const FVector End(TileMap->GetTilePositionInLocalSpace(X, Y, TileMap->SelectedLayerIndex));
 
-							PDI->DrawLine(LocalToWorld.TransformPosition(Start), LocalToWorld.TransformPosition(End), OverrideColor, DPG);
+							PDI->DrawLine(LocalToWorld.TransformPosition(Start), LocalToWorld.TransformPosition(End), OverrideColor, DPG, 0.0f, DepthBias);
 						}
 
 						// Draw vertical lines
 						for (int32 X = 0; X <= TileMap->MapWidth; ++X)
 						{
 							int32 Y = 0;
-							const FVector Start((X - 0.5f) * TW, 0.0f, -(Y - 0.5f) * TH);
+							const FVector Start(TileMap->GetTilePositionInLocalSpace(X, Y, TileMap->SelectedLayerIndex));
 
 							Y = TileMap->MapHeight;
-							const FVector End((X - 0.5f) * TW, 0.0f, -(Y - 0.5f) * TH);
+							const FVector End(TileMap->GetTilePositionInLocalSpace(X, Y, TileMap->SelectedLayerIndex));
 
-							PDI->DrawLine(LocalToWorld.TransformPosition(Start), LocalToWorld.TransformPosition(End), OverrideColor, DPG);
+							PDI->DrawLine(LocalToWorld.TransformPosition(Start), LocalToWorld.TransformPosition(End), OverrideColor, DPG, 0.0f, DepthBias);
 						}
 					}
-
-
-					// Create a local space bounding box
-					const FVector TopLeft(-TW*0.5f, -2.0f, TH*0.5f);
-					const FVector Dimensions(TileMap->MapWidth * TW, 4.0f, -TileMap->MapHeight * TH);
-					const FBox OutlineBox(TopLeft, TopLeft + Dimensions);
-
-					// Draw it		
-					DrawWireBox(PDI, OutlineBox.TransformBy(LocalToWorld), FLinearColor::White, DPG); //@TODO: Paper: Doesn't handle rotation well - need to use DrawOrientedWireBox
 				}
 #endif
 			}
diff --git a/Engine/Plugins/2D/Paper2D/Source/Paper2DEditor/Private/TileMapEditing/EdModeTileMap.cpp b/Engine/Plugins/2D/Paper2D/Source/Paper2DEditor/Private/TileMapEditing/EdModeTileMap.cpp
--- a/Engine/Plugins/2D/Paper2D/Source/Paper2DEditor/Private/TileMapEditing/EdModeTileMap.cpp
+++ b/Engine/Plugins/2D/Paper2D/Source/Paper2DEditor/Private/TileMapEditing/EdModeTileMap.cpp
@@ -341,19 +341,11 @@ UPaperTileLayer* FEdModeTileMap::GetSelectedLayerUnderCursor(const FViewportCurs
 	{
 		if (UPaperTileMap* TileMap = TileMapComponent->TileMap)
 		{
-			// Find the first visible layer
-			int32 LayerIndex = 0;
-			for (; LayerIndex < TileMap->TileLayers.Num(); ++LayerIndex)
-			{
-				UPaperTileLayer* Layer = TileMap->TileLayers[LayerIndex];
-				if (!Layer->bHiddenInEditor && Layer->bCollisionLayer == bCollisionPainting)
-				{
-					break;
-				}
-			}
+			// Find the selected layer
+			int32 LayerIndex = TileMap->SelectedLayerIndex;
 
-			// If there was a visible layer, pick it
-			if (LayerIndex < TileMap->TileLayers.Num())
+			// If there was a selected layer, pick it
+			if (TileMap->TileLayers.IsValidIndex(LayerIndex))
 			{
 				UPaperTileLayer* Layer = TileMap->TileLayers[LayerIndex];
 
