Paper2D: Fix z-fighting on tile map edit mode cursor preview

diff --git a/Engine/Plugins/2D/Paper2D/Source/Paper2DEditor/Private/TileMapEditing/EdModeTileMap.cpp b/Engine/Plugins/2D/Paper2D/Source/Paper2DEditor/Private/TileMapEditing/EdModeTileMap.cpp
--- a/Engine/Plugins/2D/Paper2D/Source/Paper2DEditor/Private/TileMapEditing/EdModeTileMap.cpp
+++ b/Engine/Plugins/2D/Paper2D/Source/Paper2DEditor/Private/TileMapEditing/EdModeTileMap.cpp
@@ -273,7 +273,7 @@ void FEdModeTileMap::Render(const FSceneView* View, FViewport* Viewport, FPrimit
 		FVector Z = DrawPreviewSpace.GetScaledAxis(EAxis::Z);
 		FVector Base = DrawPreviewLocation;
 
-		DrawOrientedWireBox(PDI, Base, X, Y, Z, DrawPreviewDimensionsLS, FLinearColor::White, SDPG_Foreground);
+		DrawOrientedWireBox(PDI, Base, X, Y, Z, DrawPreviewDimensionsLS, FLinearColor::White, SDPG_Foreground, 0.0f, 0.0001f);
 	}
 }
 
@@ -738,7 +738,7 @@ void FEdModeTileMap::UpdatePreviewCursor(const FViewportCursorLocation& Ray)
 			DrawPreviewSpace = ComponentToWorld;
 			DrawPreviewLocation = (WorldPosition + WorldPositionBR) * 0.5f;
 
-			DrawPreviewDimensionsLS = 0.5f*FVector(CursorWidth * TileMap->TileWidth, 0.0f, -CursorHeight * TileMap->TileHeight);
+			DrawPreviewDimensionsLS = 0.5f*((PaperAxisX * CursorWidth * TileMap->TileWidth) + (PaperAxisY * -CursorHeight * TileMap->TileHeight));
 		}
 	}
 }
