Paper2D: Fix smooth panning of tile set preview viewports

diff --git a/Engine/Plugins/2D/Paper2D/Source/Paper2DEditor/Private/SPaperEditorViewport.cpp b/Engine/Plugins/2D/Paper2D/Source/Paper2DEditor/Private/SPaperEditorViewport.cpp
--- a/Engine/Plugins/2D/Paper2D/Source/Paper2DEditor/Private/SPaperEditorViewport.cpp
+++ b/Engine/Plugins/2D/Paper2D/Source/Paper2DEditor/Private/SPaperEditorViewport.cpp
@@ -151,9 +151,16 @@ void SPaperEditorViewport::Tick(const FGeometry& AllottedGeometry, const double
 	ViewportClient->SetZoomPos(ViewOffset, GetZoomAmount());
 	ViewportClient->bNeedsRedraw = true;
 
-	if (Marquee.IsValid() || FSlateThrottleManager::Get().IsAllowingExpensiveTasks())
+	bool bSelectionModified = false;
+	if (Marquee.IsValid())
 	{
+		bSelectionModified = true;
+
 		OnSelectionChanged.ExecuteIfBound(Marquee, true);
+	}
+
+	if (bSelectionModified || bIsPanning || FSlateThrottleManager::Get().IsAllowingExpensiveTasks())
+	{
 		// Setup the selection set for the viewport
 		ViewportClient->SelectionRectangles.Empty();
 
