Paper2D: PaperGeomTools now included in runtime
Paper2D: TerrainComponent uses PaperGeomTools for triangulation
Paper2D: Triangulation works for CCW and CW splines, but not yet for crossing edges

diff --git a/Engine/Plugins/2D/Paper2D/Source/Paper2D/Classes/PaperTerrainComponent.h b/Engine/Plugins/2D/Paper2D/Source/Paper2D/Classes/PaperTerrainComponent.h
--- a/Engine/Plugins/2D/Paper2D/Source/Paper2D/Classes/PaperTerrainComponent.h
+++ b/Engine/Plugins/2D/Paper2D/Source/Paper2D/Classes/PaperTerrainComponent.h
@@ -78,7 +78,7 @@ public:
 
 protected:
 	void SpawnSegment(const class UPaperSprite* NewSprite, float Position, float HorizontalScale, float NominalWidth, const struct FPaperTerrainMaterialRule* Rule);
-	void SpawnFromPoly(const class UPaperSprite* NewSprite, FSpriteDrawCallRecord& FillDrawCall, const FVector2D& TextureSize, FPoly& Poly);
+	void SpawnFromPoly(const class UPaperSprite* NewSprite, FSpriteDrawCallRecord& FillDrawCall, const FVector2D& TextureSize, const TArray<FVector2D>& SplinePolyVertices2D);
 
 	void OnSplineEdited();
 
diff --git a/Engine/Plugins/2D/Paper2D/Source/Paper2D/Private/PaperGeomTools.cpp b/Engine/Plugins/2D/Paper2D/Source/Paper2D/Private/PaperGeomTools.cpp
--- a/Engine/Plugins/2D/Paper2D/Source/Paper2D/Private/PaperGeomTools.cpp
+++ b/Engine/Plugins/2D/Paper2D/Source/Paper2D/Private/PaperGeomTools.cpp
@@ -1,10 +1,7 @@
 // Copyright 1998-2014 Epic Games, Inc. All Rights Reserved.
 
-#if WITH_EDITOR
-
 #include "Paper2DPrivatePCH.h"
 #include "PaperGeomTools.h"
-#include "GeomTools.h"
 
 bool PaperGeomTools::IsPolygonWindingCCW(const TArray<FVector2D>& Points)
 {
@@ -188,9 +185,47 @@ static void JoinMutuallyVisible(TArray<FVector2D>& AdditivePoly, const TArray<FV
 	JoinSubtractiveToAdditive(AdditivePoly, SubtractivePoly, (IndexR == -1) ? IndexP : IndexR, IndexMaxX);
 }
 
-TArray<FSpritePolygon> PaperGeomTools::ReducePolygons(const TArray<FSpritePolygon>& Polygons)
+/** Util that tries to combine two triangles if possible. */
+// Unlike the implementation in GeomTools, we only care about 2D points and not any other vertex attributes
+static bool MergeTriangleIntoPolygon(
+	TArray<FVector2D>& PolygonVertices,
+	const FVector2D& TriangleVertexA, const FVector2D& TriangleVertexB, const FVector2D& TriangleVertexC)
+{
+	// Create indexable copy
+	FVector2D TriangleVertices[3] = { TriangleVertexA, TriangleVertexB, TriangleVertexC };
+
+	for (int32 PolygonEdgeIndex = 0; PolygonEdgeIndex < PolygonVertices.Num(); PolygonEdgeIndex++)
+	{
+		const uint32 PolygonEdgeVertex0 = PolygonEdgeIndex + 0;
+		const uint32 PolygonEdgeVertex1 = (PolygonEdgeIndex + 1) % PolygonVertices.Num();
+
+		for (uint32 TriangleEdgeIndex = 0; TriangleEdgeIndex < 3; TriangleEdgeIndex++)
+		{
+			const uint32 TriangleEdgeVertex0 = TriangleEdgeIndex + 0;
+			const uint32 TriangleEdgeVertex1 = (TriangleEdgeIndex + 1) % 3;
+
+			// If the triangle and polygon share an edge, then the triangle is in the same plane (implied by the above normal check),
+			// and may be merged into the polygon.
+			if (PolygonVertices[PolygonEdgeVertex0].Equals(TriangleVertices[TriangleEdgeVertex1], THRESH_POINTS_ARE_SAME) &&
+				PolygonVertices[PolygonEdgeVertex1].Equals(TriangleVertices[TriangleEdgeVertex0], THRESH_POINTS_ARE_SAME))
+			{
+				// Add the triangle's vertex that isn't in the adjacent edge to the polygon in between the vertices of the adjacent edge.
+				const int32 TriangleOppositeVertexIndex = (TriangleEdgeIndex + 2) % 3;
+				PolygonVertices.Insert(TriangleVertices[TriangleOppositeVertexIndex], PolygonEdgeVertex1);
+
+				return true;
+			}
+		}
+	}
+
+	// Could not merge triangles.
+	return false;
+}
+
+
+TArray<TArray<FVector2D>> PaperGeomTools::ReducePolygons(const TArray<TArray<FVector2D>>& Polygons, const TArray<bool>& PolygonNegativeWinding)
 {
-	TArray<FSpritePolygon> ReturnPolygons;
+	TArray<TArray<FVector2D>> ReturnPolygons;
 
 	int NumPolygons = Polygons.Num();
 	
@@ -198,7 +233,7 @@ TArray<FSpritePolygon> PaperGeomTools::ReducePolygons(const TArray<FSpritePolygo
 	for (int PolyIndex = 0; PolyIndex < NumPolygons; ++PolyIndex)
 	{
 		float MaxX = -BIG_NUMBER;
-		const TArray<FVector2D> &Vertices = Polygons[PolyIndex].Vertices;
+		const TArray<FVector2D>& Vertices = Polygons[PolyIndex];
 		for (int VertexIndex = 0; VertexIndex < Vertices.Num(); ++VertexIndex)
 		{
 			MaxX = FMath::Max(Vertices[VertexIndex].X, MaxX);
@@ -209,15 +244,16 @@ TArray<FSpritePolygon> PaperGeomTools::ReducePolygons(const TArray<FSpritePolygo
 	// Iterate through additive shapes
 	for (int PolyIndex = 0; PolyIndex < NumPolygons; ++PolyIndex)
 	{
-		TArray<FVector2D> Verts = Polygons[PolyIndex].Vertices;
-		if (!Polygons[PolyIndex].bNegativeWinding)
+		if (!PolygonNegativeWinding[PolyIndex])
 		{
+			TArray<FVector2D> Verts(Polygons[PolyIndex]);
+
 			// Store indexes of subtractive shapes that lie inside the additive shape
 			TArray<int> SubtractiveShapeIndices;
 			for (int J = 0; J < NumPolygons; ++J)
 			{
-				const TArray<FVector2D>& CandidateShape = Polygons[J].Vertices;
-				if (Polygons[J].bNegativeWinding)
+				const TArray<FVector2D>& CandidateShape = Polygons[J];
+				if (PolygonNegativeWinding[J])
 				{
 					if (IsPointInPolygon(CandidateShape[0], Verts))
 					{
@@ -229,7 +265,7 @@ TArray<FSpritePolygon> PaperGeomTools::ReducePolygons(const TArray<FSpritePolygo
 			// Remove subtractive shapes that lie inside the subtractive shapes we've found
 			for (int J = 0; J < SubtractiveShapeIndices.Num();)
 			{
-				const TArray<FVector2D>& ourShape = Polygons[SubtractiveShapeIndices[J]].Vertices;
+				const TArray<FVector2D>& ourShape = Polygons[SubtractiveShapeIndices[J]];
 				bool bRemoveOurShape = false;
 				for (int K = 0; K < SubtractiveShapeIndices.Num(); ++K)
 				{
@@ -237,7 +273,7 @@ TArray<FSpritePolygon> PaperGeomTools::ReducePolygons(const TArray<FSpritePolygo
 					{
 						continue;
 					}
-					if (IsPointInPolygon(ourShape[0], Polygons[SubtractiveShapeIndices[K]].Vertices))
+					if (IsPointInPolygon(ourShape[0], Polygons[SubtractiveShapeIndices[K]]))
 					{
 						bRemoveOurShape = true;
 						break;
@@ -270,13 +306,11 @@ TArray<FSpritePolygon> PaperGeomTools::ReducePolygons(const TArray<FSpritePolygo
 
 			for (int SubtractiveIndex = 0; SubtractiveIndex < NumSubtractiveShapes; ++SubtractiveIndex)
 			{
-				JoinMutuallyVisible(Verts, Polygons[SubtractiveShapeIndices[SubtractiveIndex]].Vertices);
+				JoinMutuallyVisible(Verts, Polygons[SubtractiveShapeIndices[SubtractiveIndex]]);
 			}
 
 			// Add new hole-less polygon to our output shapes
-			FSpritePolygon *NewPolygon = new(ReturnPolygons)FSpritePolygon();
-			NewPolygon->bNegativeWinding = false;
-			NewPolygon->Vertices = Verts;
+			ReturnPolygons.Add(Verts);
 		}
 	}
 
@@ -284,52 +318,44 @@ TArray<FSpritePolygon> PaperGeomTools::ReducePolygons(const TArray<FSpritePolygo
 }
 
 
-TArray<FSpritePolygon> PaperGeomTools::CorrectPolygonWinding(const TArray<FSpritePolygon>& Polygons)
+void PaperGeomTools::CorrectPolygonWinding(TArray<FVector2D>& OutVertices, const TArray<FVector2D>& Vertices, const bool bNegativeWinding)
 {
-	TArray<FSpritePolygon> ReturnPolygons;
-	for (int32 PolygonIndex = 0; PolygonIndex < Polygons.Num(); ++PolygonIndex)
+	if (Vertices.Num() >= 3)
 	{
-		const FSpritePolygon& SourcePolygon = Polygons[PolygonIndex];
-		if (SourcePolygon.Vertices.Num() >= 3)
+		// Make sure the polygon winding is correct
+		if ((!bNegativeWinding && !IsPolygonWindingCCW(Vertices)) ||
+			(bNegativeWinding && IsPolygonWindingCCW(Vertices)))
 		{
-			FSpritePolygon* FixedPolygon = new (ReturnPolygons)FSpritePolygon();
-			// Make sure the polygon winding is correct
-			if ((!SourcePolygon.bNegativeWinding && !IsPolygonWindingCCW(SourcePolygon.Vertices)) ||
-				(SourcePolygon.bNegativeWinding && IsPolygonWindingCCW(SourcePolygon.Vertices)))
+			// Reverse vertices
+			for (int32 VertexIndex = Vertices.Num() - 1; VertexIndex >= 0; --VertexIndex)
 			{
-				// Reverse vertices
-				for (int32 VertexIndex = SourcePolygon.Vertices.Num() - 1; VertexIndex >= 0; --VertexIndex)
-				{
-					new (FixedPolygon->Vertices)FVector2D(SourcePolygon.Vertices[VertexIndex]);
-				}
+				new (OutVertices)FVector2D(Vertices[VertexIndex]);
 			}
-			else
+		}
+		else
+		{
+			// Copy vertices
+			for (int32 VertexIndex = 0; VertexIndex < Vertices.Num(); ++VertexIndex)
 			{
-				// Copy vertices
-				for (int32 VertexIndex = 0; VertexIndex < SourcePolygon.Vertices.Num(); ++VertexIndex)
-				{
-					new (FixedPolygon->Vertices)FVector2D(SourcePolygon.Vertices[VertexIndex]);
-				}
+				new (OutVertices)FVector2D(Vertices[VertexIndex]);
 			}
-			FixedPolygon->bNegativeWinding = SourcePolygon.bNegativeWinding;
 		}
 	}
-	return ReturnPolygons;
 }
 
 // Assumes polygons are valid, closed and the winding is correct and matches bWindingCCW
-bool PaperGeomTools::ArePolygonsValid(const TArray<FSpritePolygon>& Polygons)
+bool PaperGeomTools::ArePolygonsValid(const TArray<TArray<FVector2D>>& Polygons)
 {
 	const int PolygonCount = Polygons.Num();
 
 	// Find every pair of shapes that overlap
 	for (int PolygonIndexA = 0; PolygonIndexA < PolygonCount; ++PolygonIndexA)
 	{
-		const TArray<FVector2D>& PolygonA = Polygons[PolygonIndexA].Vertices;
+		const TArray<FVector2D>& PolygonA = Polygons[PolygonIndexA];
 		const bool bIsWindingA_CCW = IsPolygonWindingCCW(PolygonA);
 		for (int PolygonIndexB = PolygonIndexA + 1; PolygonIndexB < PolygonCount; ++PolygonIndexB)
 		{
-			const TArray<FVector2D>& PolygonB = Polygons[PolygonIndexB].Vertices;
+			const TArray<FVector2D>& PolygonB = Polygons[PolygonIndexB];
 			const bool bIsWindingB_CCW = IsPolygonWindingCCW(PolygonB);
 
 			// Additive polygons are allowed to intersect
@@ -529,34 +555,49 @@ bool PaperGeomTools::TriangulatePoly(TArray<FVector2D>& OutTris, const TArray<FV
 	return true;
 }
 
-// Wrap around GeomTools RemoveRedundantTriangles
-void PaperGeomTools::RemoveRedundantTriangles(TArray<FVector2D>& OutTriangles, const TArray<FVector2D>& InTriangles)
+// 2D version of GeomTools RemoveRedundantTriangles
+void PaperGeomTools::RemoveRedundantTriangles(TArray<FVector2D>& OutTriangles, const TArray<FVector2D>& InTriangleVertices)
 {
-	// Convert to ClipSMTriangles
-	TArray<FClipSMTriangle> ClipSMTriangles;
-	for (int TriangleVertex = 0; TriangleVertex < InTriangles.Num(); TriangleVertex += 3)
+	struct FLocalTriangle
 	{
-		FClipSMTriangle* Tri = new(ClipSMTriangles)FClipSMTriangle(0); // Zeroed
-		for (int TriVertIndex = 0; TriVertIndex < 3; ++TriVertIndex)
-		{
-			Tri->Vertices[TriVertIndex].Pos.X = InTriangles[TriangleVertex + TriVertIndex].X;
-			Tri->Vertices[TriVertIndex].Pos.Z = InTriangles[TriangleVertex + TriVertIndex].Y;
-		}
-		Tri->FaceNormal = FVector(0.0f, -1.0f, 0.0f);
-	}
+		int VertexA, VertexB, VertexC;
+	};
 
-	// GeomTools.cpp
-	::RemoveRedundantTriangles(ClipSMTriangles);
+	TArray<FLocalTriangle> Triangles;
+	for (int32 TriangleVertexIndex = 0; TriangleVertexIndex < InTriangleVertices.Num(); TriangleVertexIndex += 3)
+	{
+		FLocalTriangle* NewTriangle = new(Triangles)FLocalTriangle();
+		NewTriangle->VertexA = TriangleVertexIndex + 0;
+		NewTriangle->VertexB = TriangleVertexIndex + 1;
+		NewTriangle->VertexC = TriangleVertexIndex + 2;
+	}
 
-	// Convert back to FVector2Ds
-	OutTriangles.Empty();
-	for (int TriangleIndex = 0; TriangleIndex < ClipSMTriangles.Num(); ++TriangleIndex)
+	while (Triangles.Num() > 0)
 	{
-		FClipSMTriangle& Triangle = ClipSMTriangles[TriangleIndex];
-		new (OutTriangles)FVector2D(Triangle.Vertices[0].Pos.X, Triangle.Vertices[0].Pos.Z);
-		new (OutTriangles)FVector2D(Triangle.Vertices[1].Pos.X, Triangle.Vertices[1].Pos.Z);
-		new (OutTriangles)FVector2D(Triangle.Vertices[2].Pos.X, Triangle.Vertices[2].Pos.Z);
+		TArray<FVector2D> PolygonVertices;
+
+		const FLocalTriangle InitialTriangle = Triangles.Pop();
+		PolygonVertices.Add(InTriangleVertices[InitialTriangle.VertexA]);
+		PolygonVertices.Add(InTriangleVertices[InitialTriangle.VertexB]);
+		PolygonVertices.Add(InTriangleVertices[InitialTriangle.VertexC]);
+
+		// Find triangles that can be merged into the polygon.
+		for(int32 CandidateTriangleIndex = 0;CandidateTriangleIndex < Triangles.Num(); ++CandidateTriangleIndex)
+		{
+			const FLocalTriangle& MergeCandidateTriangle = Triangles[CandidateTriangleIndex];
+			if (MergeTriangleIntoPolygon(PolygonVertices, InTriangleVertices[MergeCandidateTriangle.VertexA], InTriangleVertices[MergeCandidateTriangle.VertexB], InTriangleVertices[MergeCandidateTriangle.VertexC]))
+			{
+				// Remove the merged triangle from the array.
+				Triangles.RemoveAtSwap(CandidateTriangleIndex);
+
+				// Restart the search for mergeable triangles from the start of the array.
+				CandidateTriangleIndex = -1;
+			}
+		}
+
+		// Triangulate merged polygon and append to triangle array
+		TArray<FVector2D> TriangulatedPoly;
+		TriangulatePoly(/*out*/TriangulatedPoly, PolygonVertices);
+		OutTriangles.Append(TriangulatedPoly);
 	}
 }
-
-#endif // WITH_EDITOR
diff --git a/Engine/Plugins/2D/Paper2D/Source/Paper2D/Private/PaperGeomTools.h b/Engine/Plugins/2D/Paper2D/Source/Paper2D/Private/PaperGeomTools.h
--- a/Engine/Plugins/2D/Paper2D/Source/Paper2D/Private/PaperGeomTools.h
+++ b/Engine/Plugins/2D/Paper2D/Source/Paper2D/Private/PaperGeomTools.h
@@ -2,35 +2,29 @@
 
 #pragma once
 
-#include "SpriteEditorOnlyTypes.h"
-
-#if WITH_EDITOR
-
-// PaperGeomTools is editor only, used only by PaperSprite.cpp
 class PaperGeomTools
 {
 public:
-	// Returns a list of polygons with the points in the corresponding winding to the winding flag on FSpritePolygon
+	// Corrects the polygon winding to match bNegativeWinding
 	// Ie. If the vertex order doesn't match the winding, it is reversed
-	// Entries in the array that aren't polygons (< 3 points) are removed
-	static TArray<FSpritePolygon> CorrectPolygonWinding(const TArray<FSpritePolygon>& Polygons);
+	// Returns empty array of points if num points < 3
+	static void CorrectPolygonWinding(TArray<FVector2D>& OutVertices, const TArray<FVector2D>& Vertices, const bool bNegativeWinding);
 
 	// Returns true if the points forming a polygon have CCW winding
 	// Returns true if the polygon isn't valid
 	static bool IsPolygonWindingCCW(const TArray<FVector2D>& Points);
 
 	// Checks that these polygons can be successfully triangulated	
-	static bool ArePolygonsValid(const TArray<FSpritePolygon>& Polygons);
+	static bool ArePolygonsValid(const TArray<TArray<FVector2D>>& Polygons);
 
 	// Merge additive and subtractive polygons, split them up into additive polygons
 	// Assumes all polygons and overlapping polygons are valid, and the windings match the setting on the polygon
-	static TArray<FSpritePolygon> ReducePolygons(const TArray<FSpritePolygon>& Polygons);
+	static TArray<TArray<FVector2D>> PaperGeomTools::ReducePolygons(const TArray<TArray<FVector2D>>& Polygons, const TArray<bool>& PolygonNegativeWinding);
 
 	// Triangulate a polygon. Check notes in implementation
-	static bool TriangulatePoly(TArray<FVector2D>& OutTris, const TArray<FVector2D>& PolygonVertices, bool bKeepColinearVertices);
+	static bool TriangulatePoly(TArray<FVector2D>& OutTris, const TArray<FVector2D>& PolygonVertices, bool bKeepColinearVertices = false);
 
-	// Wrapper around GeomTools.cpp, RemoveRedundantTriangles
+	// 2D version of RemoveRedundantTriangles from GeomTools.cpp
 	static void RemoveRedundantTriangles(TArray<FVector2D>& OutTriangles, const TArray<FVector2D>& InTriangles);
 };
 
-#endif
diff --git a/Engine/Plugins/2D/Paper2D/Source/Paper2D/Private/PaperSprite.cpp b/Engine/Plugins/2D/Paper2D/Source/Paper2D/Private/PaperSprite.cpp
--- a/Engine/Plugins/2D/Paper2D/Source/Paper2D/Private/PaperSprite.cpp
+++ b/Engine/Plugins/2D/Paper2D/Source/Paper2D/Private/PaperSprite.cpp
@@ -1033,41 +1033,52 @@ void UPaperSprite::Triangulate(const FSpritePolygonCollection& Source, TArray<FV
 {
 	Target.Empty();
 	TArray<FVector2D> AllGeneratedTriangles;
+	
+	// AOS -> Validate -> SOA
+	TArray<bool> PolygonsNegativeWinding; // do these polygons have negative winding?
+	TArray<TArray<FVector2D> > ValidPolygonTriangles;
+	PolygonsNegativeWinding.Empty(Source.Polygons.Num());
+	ValidPolygonTriangles.Empty(Source.Polygons.Num());
+	bool bSourcePolygonHasHoles = false;
 
 	// Correct polygon winding for additive and subtractive polygons
 	// Invalid polygons (< 3 verts) removed from this list
-	TArray<FSpritePolygon> ValidPolygons = PaperGeomTools::CorrectPolygonWinding(Source.Polygons);
-	
+	TArray<FSpritePolygon> ValidPolygons;
+	for (int32 PolygonIndex = 0; PolygonIndex < Source.Polygons.Num(); ++PolygonIndex)
+	{
+		const FSpritePolygon& SourcePolygon = Source.Polygons[PolygonIndex];
+		if (SourcePolygon.Vertices.Num() >= 3)
+		{
+			TArray<FVector2D>* FixedVertices = new (ValidPolygonTriangles)TArray<FVector2D>();
+			PaperGeomTools::CorrectPolygonWinding(*FixedVertices, SourcePolygon.Vertices, SourcePolygon.bNegativeWinding);
+			PolygonsNegativeWinding.Add(SourcePolygon.bNegativeWinding);
+		}
+
+		if (Source.Polygons[PolygonIndex].bNegativeWinding)
+		{
+			bSourcePolygonHasHoles = true;
+		}
+	}
+
 	// Check if polygons overlap, or have inconsistent winding, or edges overlap
-	if (!PaperGeomTools::ArePolygonsValid(ValidPolygons))
+	if (!PaperGeomTools::ArePolygonsValid(ValidPolygonTriangles))
 	{
 		return;
 	}
 
 	// Merge each additive and associated subtractive polygons to form a list of polygons in CCW winding
-	ValidPolygons = PaperGeomTools::ReducePolygons(ValidPolygons);
+	ValidPolygonTriangles = PaperGeomTools::ReducePolygons(ValidPolygonTriangles, PolygonsNegativeWinding);
 
 	// Triangulate the polygons
-	for (int32 PolygonIndex = 0; PolygonIndex < ValidPolygons.Num(); ++PolygonIndex)
+	for (int32 PolygonIndex = 0; PolygonIndex < ValidPolygonTriangles.Num(); ++PolygonIndex)
 	{
-		const FSpritePolygon& SourcePoly = ValidPolygons[PolygonIndex];
 		TArray<FVector2D> Generated2DTriangles;
-		if (PaperGeomTools::TriangulatePoly(Generated2DTriangles, SourcePoly.Vertices, Source.bAvoidVertexMerging))
+		if (PaperGeomTools::TriangulatePoly(Generated2DTriangles, ValidPolygonTriangles[PolygonIndex], Source.bAvoidVertexMerging))
 		{
 			AllGeneratedTriangles.Append(Generated2DTriangles);
 		}
 	}
 
-	bool bSourcePolygonHasHoles = false;
-	for (int32 PolygonIndex = 0; PolygonIndex < Source.Polygons.Num(); ++PolygonIndex)
-	{
-		if (Source.Polygons[PolygonIndex].bNegativeWinding)
-		{
-			bSourcePolygonHasHoles = true;
-			break;
-		}
-	}
-
 	// This doesn't work when polys have holes as edges will likely form a loop around the poly
 	if (!bSourcePolygonHasHoles && !Source.bAvoidVertexMerging && (Source.Polygons.Num() > 1 && AllGeneratedTriangles.Num() > 1))
 	{
diff --git a/Engine/Plugins/2D/Paper2D/Source/Paper2D/Private/Terrain/PaperTerrainComponent.cpp b/Engine/Plugins/2D/Paper2D/Source/Paper2D/Private/Terrain/PaperTerrainComponent.cpp
--- a/Engine/Plugins/2D/Paper2D/Source/Paper2D/Private/Terrain/PaperTerrainComponent.cpp
+++ b/Engine/Plugins/2D/Paper2D/Source/Paper2D/Private/Terrain/PaperTerrainComponent.cpp
@@ -4,6 +4,7 @@
 #include "Components/SplineComponent.h"
 
 #include "PaperRenderSceneProxy.h"
+#include "PaperGeomTools.h"
 
 #define PAPER_USE_MATERIAL_SLOPES 1
 
@@ -469,7 +470,7 @@ void UPaperTerrainComponent::OnSplineEdited()
 		{
 			// Create a polygon from the spline
 			FBox2D SplineBounds(ForceInit);
-			FPoly SplineAsPolygon;
+			TArray<FVector2D> SplinePolyVertices2D;
 			{
 				float CurrentTime = 0.0f;
 				while (CurrentTime < SplineLength)
@@ -477,10 +478,9 @@ void UPaperTerrainComponent::OnSplineEdited()
 					const float Param = AssociatedSpline->SplineReparamTable.Eval(CurrentTime, 0.0f);
 					const FVector Position3D = AssociatedSpline->SplineInfo.Eval(Param, FVector::ZeroVector);
 					const FVector2D Position2D = FVector2D(FVector::DotProduct(Position3D, PaperAxisX), FVector::DotProduct(Position3D, PaperAxisY));
-					const FVector CleanPosition3D = (PaperAxisX * Position2D.X) + (PaperAxisY * Position2D.Y);
 
 					SplineBounds += Position2D;
-					SplineAsPolygon.Vertices.Add(CleanPosition3D);
+					SplinePolyVertices2D.Add(Position2D);
 
 					CurrentTime += FillRasterizationTimeRate;
 				}
@@ -499,7 +499,7 @@ void UPaperTerrainComponent::OnSplineEdited()
 			const FVector2D TextureSize = GetSpriteRenderDataBounds2D(FillSprite->BakedRenderData).GetSize();
 			const FVector2D SplineSize = SplineBounds.GetSize();
 			
-			SpawnFromPoly(FillSprite, FillDrawCall, TextureSize, SplineAsPolygon);
+			SpawnFromPoly(FillSprite, FillDrawCall, TextureSize, SplinePolyVertices2D);
 
 			//@TODO: Add support for the fill sprite being smaller than the entire texture
 #if NOT_WORKING
@@ -599,55 +599,27 @@ void UPaperTerrainComponent::SpawnSegment(const UPaperSprite* NewSprite, float P
 	}
 }
 
-#if WITH_EDITOR
-#include "GeomTools.h"
-#endif
-
-void UPaperTerrainComponent::SpawnFromPoly(const class UPaperSprite* NewSprite, FSpriteDrawCallRecord& Batch, const FVector2D& TextureSize, FPoly& Poly)
+void UPaperTerrainComponent::SpawnFromPoly(const class UPaperSprite* NewSprite, FSpriteDrawCallRecord& Batch, const FVector2D& TextureSize, const TArray<FVector2D>& SplinePolyVertices2D)
 {
-	//@TODO: Need to split geom tools out of UnrealEd to support this outside of the editor
-#if WITH_EDITOR
 	const FVector2D TextureSizeInUnits = TextureSize * NewSprite->GetUnrealUnitsPerPixel();
 
-	if (Poly.Vertices.Num() >= 3)
-	{
-		FClipSMPolygon Polygon(0);
-		for (int32 v = 0; v < Poly.Vertices.Num(); ++v)
-		{
-			FClipSMVertex vtx;
-			vtx.Pos = Poly.Vertices[v];
-			Polygon.Vertices.Add(vtx);
-		}
-
-		Polygon.FaceNormal = -PaperAxisZ;
-
-		TArray<FClipSMTriangle> GeneratedTriangles;
-		TriangulatePoly(GeneratedTriangles, Polygon, false);
+	// Correct spline vertices so they are always CCW
+	//@TODO: Handle overlapping and crossing polygon edges during triangulation
+	TArray<FVector2D> CorrectedSpriteVertices;
+	PaperGeomTools::CorrectPolygonWinding(/*out*/CorrectedSpriteVertices, SplinePolyVertices2D, false);
 
-		struct Local
-		{
-			static FVector4 GenerateVert(const FVector& Source3D, const FVector2D& Size2D)
-			{
-				const float X = FVector::DotProduct(Source3D, PaperAxisX);
-				const float Y = FVector::DotProduct(Source3D, PaperAxisY);
-				const float U = (X / (Size2D.X));
-				const float V = (-Y / (Size2D.Y));
-
-				return FVector4(X, Y, U, V);
-			}
-		};
+	if (SplinePolyVertices2D.Num() >= 3)
+	{
+		TArray<FVector2D> TriangulatedPolygonVertices;
+		PaperGeomTools::TriangulatePoly(/*out*/TriangulatedPolygonVertices, CorrectedSpriteVertices, false);
 
-		// Convert the triangles back to our 2D data structure
-		for (int32 TriangleIndex = 0; TriangleIndex < GeneratedTriangles.Num(); ++TriangleIndex)
+		// Pack vertex data
+		for (int32 TriangleVertexIndex = 0; TriangleVertexIndex < TriangulatedPolygonVertices.Num(); ++TriangleVertexIndex)
 		{
-			const FClipSMTriangle& Triangle = GeneratedTriangles[TriangleIndex];
-
-			new (Batch.RenderVerts) FVector4(Local::GenerateVert(Triangle.Vertices[0].Pos, TextureSizeInUnits));
-			new (Batch.RenderVerts) FVector4(Local::GenerateVert(Triangle.Vertices[1].Pos, TextureSizeInUnits));
-			new (Batch.RenderVerts) FVector4(Local::GenerateVert(Triangle.Vertices[2].Pos, TextureSizeInUnits));
+			const FVector2D& TriangleVertex = TriangulatedPolygonVertices[TriangleVertexIndex];
+			new (Batch.RenderVerts) FVector4(TriangleVertex.X, TriangleVertex.Y, TriangleVertex.X / TextureSizeInUnits.X, -TriangleVertex.Y / TextureSizeInUnits.Y);
 		}
 	}
-#endif
 }
 
 void UPaperTerrainComponent::SetTerrainColor(FLinearColor NewColor)
