Paper2D: Save source texture dimensions to handle resize, RescaleSpriteData to perform the rescale
Paper2D: Disabled by default in PostEditChangeProperty but could be useful there

diff --git a/Engine/Plugins/2D/Paper2D/Source/Paper2D/Classes/PaperSprite.h b/Engine/Plugins/2D/Paper2D/Source/Paper2D/Classes/PaperSprite.h
--- a/Engine/Plugins/2D/Paper2D/Source/Paper2D/Classes/PaperSprite.h
+++ b/Engine/Plugins/2D/Paper2D/Source/Paper2D/Classes/PaperSprite.h
@@ -74,6 +74,11 @@ protected:
 	// This texture is rotated in the atlas
 	UPROPERTY(Category=Sprite, EditAnywhere, AdvancedDisplay)
 	bool bRotatedInSourceImage;
+
+	// Dimension of the texture when this sprite was created
+	// Used when the sprite is resized at some point
+	UPROPERTY(Category=Sprite, EditAnywhere, AdvancedDisplay)
+	FVector2D SourceTextureDimension;
 #endif
 
 	// The source texture that the sprite comes from
@@ -181,6 +186,9 @@ public:
 	// Returns the current pivot position in texture space
 	FVector2D GetPivotPosition() const;
 
+	// Rescale properties to handle source texture size change
+	void RescaleSpriteData(const FVector2D& PreviousTextureDimension, const FVector2D& CurrentTextureDimension);
+
 	void RebuildCollisionData();
 	void RebuildRenderData();
 
diff --git a/Engine/Plugins/2D/Paper2D/Source/Paper2D/Private/PaperSprite.cpp b/Engine/Plugins/2D/Paper2D/Source/Paper2D/Private/PaperSprite.cpp
--- a/Engine/Plugins/2D/Paper2D/Source/Paper2D/Private/PaperSprite.cpp
+++ b/Engine/Plugins/2D/Paper2D/Source/Paper2D/Private/PaperSprite.cpp
@@ -269,6 +269,8 @@ UPaperSprite::UPaperSprite(const FObjectInitializer& ObjectInitializer)
 
 	bTrimmedInSourceImage = false;
 	bRotatedInSourceImage = false;
+
+	SourceTextureDimension.Set(0, 0);
 #endif
 
 	static ConstructorHelpers::FObjectFinder<UMaterialInterface> MaskedMaterialRef(TEXT("/Paper2D/MaskedUnlitSpriteMaterial"));
@@ -403,15 +405,27 @@ void UPaperSprite::PostEditChangeProperty(FPropertyChangedEvent& PropertyChanged
 
 	if (PropertyName == GET_MEMBER_NAME_CHECKED(UPaperSprite, SourceTexture))
 	{
-		// If this is a brand new sprite that didn't have a texture set previously, act like we were factoried with the texture
 		if ((SourceTexture != nullptr) && SourceDimension.IsNearlyZero())
 		{
+			// If this is a brand new sprite that didn't have a texture set previously, act like we were factoried with the texture
 			SourceUV = FVector2D::ZeroVector;
 			SourceDimension = FVector2D(SourceTexture->GetSizeX(), SourceTexture->GetSizeY());
+			SourceTextureDimension = FVector2D(SourceTexture->GetSizeX(), SourceTexture->GetSizeY());
 		}
 		bBothModified = true;
 	}
 
+	// The texture dimensions have changed
+	if (SourceTexture != nullptr && (SourceTextureDimension.X != SourceTexture->GetSizeX() || SourceTextureDimension.Y != SourceTexture->GetSizeY()))
+	{
+		FVector2D PreviousTextureDimension = SourceTextureDimension;
+		SourceTextureDimension = FVector2D(SourceTexture->GetSizeX(), SourceTexture->GetSizeY());
+		// Tmp - disabled, not sure if we want it on here or not
+		// RescaleSpriteData(PreviousTextureDimension, SourceTextureDimension);
+		bBothModified = true;
+	}
+
+
 	if (bCollisionDataModified || bBothModified)
 	{
 		RebuildCollisionData();
@@ -425,6 +439,89 @@ void UPaperSprite::PostEditChangeProperty(FPropertyChangedEvent& PropertyChanged
 	Super::PostEditChangeProperty(PropertyChangedEvent);
 }
 
+void UPaperSprite::RescaleSpriteData(const FVector2D& PreviousTextureDimension, const FVector2D& CurrentTextureDimension)
+{
+	// Don't ever divby0 (no previously stored texture dimensions)
+	// or scale to 0
+	if (CurrentTextureDimension.X == 0 || CurrentTextureDimension.Y == 0 ||
+		PreviousTextureDimension.X == 0 || PreviousTextureDimension.Y == 0)
+	{
+		return;
+	}
+
+	const FVector2D& S = CurrentTextureDimension;
+	const FVector2D& D = PreviousTextureDimension;
+	struct Local
+	{
+		static float NoSnap(const float Value, const float Scale, const float Divisor)
+		{
+			return (Value * Scale) / Divisor;
+		}
+
+		static FVector2D RescaleNeverSnap(const FVector2D& Value, const FVector2D& Scale, const FVector2D& Divisor)
+		{
+			return FVector2D(NoSnap(Value.X, Scale.X, Divisor.X), NoSnap(Value.Y, Scale.Y, Divisor.Y));
+		}
+
+		static FVector2D Rescale(const FVector2D& Value, const FVector2D& Scale, const FVector2D& Divisor)
+		{
+			// Never snap, want to be able to return to original values when rescaled back
+			return RescaleNeverSnap(Value, Scale, Divisor);
+			//return FVector2D(FMath::FloorToFloat(NoSnap(Value.X, Scale.X, Divisor.X)), FMath::FloorToFloat(NoSnap(Value.Y, Scale.Y, Divisor.Y)));
+		}
+	};
+
+	// Sockets are in pivot space, convert these to texture space to apply later
+	TArray<FVector2D> RescaledTextureSpaceSocketPositions;
+	for (int32 SocketIndex = 0; SocketIndex < Sockets.Num(); ++SocketIndex)
+	{
+		FPaperSpriteSocket& Socket = Sockets[SocketIndex];
+		FVector Translation = Socket.LocalTransform.GetTranslation();
+		FVector2D TextureSpaceSocketPosition = ConvertPivotSpaceToTextureSpace(FVector2D(Translation.X, Translation.Z));
+		RescaledTextureSpaceSocketPositions.Add(Local::RescaleNeverSnap(TextureSpaceSocketPosition, S, D));
+	}
+
+	SourceUV = Local::Rescale(SourceUV, S, D);
+	SourceDimension = Local::Rescale(SourceDimension, S, D);
+	SourceImageDimensionBeforeTrimming = Local::Rescale(SourceImageDimensionBeforeTrimming, S, D);
+
+	if (bSnapPivotToPixelGrid)
+	{
+		CustomPivotPoint = Local::Rescale(CustomPivotPoint, S, D);
+	}
+	else
+	{
+		CustomPivotPoint = Local::RescaleNeverSnap(CustomPivotPoint, S, D);
+	}
+
+	for (int32 GeomtryIndex = 0; GeomtryIndex < 2; ++GeomtryIndex)
+	{
+		FSpritePolygonCollection& Geometry = (GeomtryIndex == 0) ? CollisionGeometry : RenderGeometry;
+		for (int32 PolygonIndex = 0; PolygonIndex < Geometry.Polygons.Num(); ++PolygonIndex)
+		{
+			FSpritePolygon& Polygon = Geometry.Polygons[PolygonIndex];
+			Polygon.BoxPosition = Local::Rescale(Polygon.BoxPosition, S, D);
+			Polygon.BoxSize = Local::Rescale(Polygon.BoxSize, S, D);
+
+			for (int32 VertexIndex = 0; VertexIndex < Polygon.Vertices.Num(); ++VertexIndex)
+			{
+				Polygon.Vertices[VertexIndex] = Local::Rescale(Polygon.Vertices[VertexIndex], S, D);
+			}
+		}
+	}
+
+	// Apply texture space pivot positions now that pivot space is correctly defined
+	for (int32 SocketIndex = 0; SocketIndex < Sockets.Num(); ++SocketIndex)
+	{
+		FPaperSpriteSocket& Socket = Sockets[SocketIndex];
+		FVector2D PivotSpaceSocketPosition = ConvertTextureSpaceToPivotSpace(RescaledTextureSpaceSocketPositions[SocketIndex]);
+		FVector Translation = Socket.LocalTransform.GetTranslation();
+		Translation.X = PivotSpaceSocketPosition.X;
+		Translation.Z = PivotSpaceSocketPosition.Y;
+		Socket.LocalTransform.SetTranslation(Translation);
+	}
+}
+
 void UPaperSprite::RebuildCollisionData()
 {
 	UBodySetup* OldBodySetup = BodySetup;
@@ -1115,8 +1212,17 @@ void UPaperSprite::InitializeSprite(const FSpriteAssetInitParameters& InitParams
 	}
 
 	SourceTexture = InitParams.Texture;
+	if (SourceTexture != nullptr)
+	{
+		SourceTextureDimension.Set(SourceTexture->GetSizeX(), SourceTexture->GetSizeY());
+	}
+	else
+	{
+		SourceTextureDimension.Set(0, 0);
+	}
 	SourceUV = InitParams.Offset;
 	SourceDimension = InitParams.Dimension;
+
 	RebuildCollisionData();
 	RebuildRenderData();
 }
