Fix CIS failure on Mac.

diff --git a/Engine/Plugins/2D/Paper2D/Source/Paper2D/Private/PaperGeomTools.h b/Engine/Plugins/2D/Paper2D/Source/Paper2D/Private/PaperGeomTools.h
--- a/Engine/Plugins/2D/Paper2D/Source/Paper2D/Private/PaperGeomTools.h
+++ b/Engine/Plugins/2D/Paper2D/Source/Paper2D/Private/PaperGeomTools.h
@@ -19,7 +19,7 @@ public:
 
 	// Merge additive and subtractive polygons, split them up into additive polygons
 	// Assumes all polygons and overlapping polygons are valid, and the windings match the setting on the polygon
-	static TArray<TArray<FVector2D>> PaperGeomTools::ReducePolygons(const TArray<TArray<FVector2D>>& Polygons, const TArray<bool>& PolygonNegativeWinding);
+	static TArray<TArray<FVector2D>> ReducePolygons(const TArray<TArray<FVector2D>>& Polygons, const TArray<bool>& PolygonNegativeWinding);
 
 	// Triangulate a polygon. Check notes in implementation
 	static bool TriangulatePoly(TArray<FVector2D>& OutTris, const TArray<FVector2D>& PolygonVertices, bool bKeepColinearVertices = false);
