Paper2D: Handle changing source texture resolution on sprites by rescaling to fit

diff --git a/Engine/Plugins/2D/Paper2D/Source/Paper2D/Classes/PaperSprite.h b/Engine/Plugins/2D/Paper2D/Source/Paper2D/Classes/PaperSprite.h
--- a/Engine/Plugins/2D/Paper2D/Source/Paper2D/Classes/PaperSprite.h
+++ b/Engine/Plugins/2D/Paper2D/Source/Paper2D/Classes/PaperSprite.h
@@ -187,7 +187,8 @@ public:
 	FVector2D GetPivotPosition() const;
 
 	// Rescale properties to handle source texture size change
-	void RescaleSpriteData(const FVector2D& PreviousTextureDimension, const FVector2D& CurrentTextureDimension);
+	void RescaleSpriteData(const UTexture2D* Texture);
+	bool NeedRescaleSpriteData();
 
 	void RebuildCollisionData();
 	void RebuildRenderData();
@@ -222,6 +223,12 @@ public:
 	UTexture2D* GetSourceTexture() const { return SourceTexture; }
 #endif
 
+#if WITH_EDITOR
+	bool bRegisteredObjectReimport;
+	// Called via delegate when an object is re-imported in the editor
+	void OnObjectReimported(UObject* InObject);
+#endif
+
 	// Return the scaling factor between pixels and Unreal units (cm)
 	float GetPixelsPerUnrealUnit() const { return PixelsPerUnrealUnit; }
 
diff --git a/Engine/Plugins/2D/Paper2D/Source/Paper2D/Private/PaperCustomVersion.h b/Engine/Plugins/2D/Paper2D/Source/Paper2D/Private/PaperCustomVersion.h
--- a/Engine/Plugins/2D/Paper2D/Source/Paper2D/Private/PaperCustomVersion.h
+++ b/Engine/Plugins/2D/Paper2D/Source/Paper2D/Private/PaperCustomVersion.h
@@ -42,6 +42,9 @@ struct FPaperCustomVersion
 
 		// Add default BodyInstance setup in the sprite asset
 		AddDefaultCollisionProfileInSpriteAsset = 11,
+		
+		// Add source texture size
+		AddSourceTextureSize = 12,
 
 		// -----<new versions can be added above this line>-------------------------------------------------
 		VersionPlusOne,
diff --git a/Engine/Plugins/2D/Paper2D/Source/Paper2D/Private/PaperSprite.cpp b/Engine/Plugins/2D/Paper2D/Source/Paper2D/Private/PaperSprite.cpp
--- a/Engine/Plugins/2D/Paper2D/Source/Paper2D/Private/PaperSprite.cpp
+++ b/Engine/Plugins/2D/Paper2D/Source/Paper2D/Private/PaperSprite.cpp
@@ -1,10 +1,14 @@
-// Copyright 1998-2015 Epic Games, Inc. All Rights Reserved.
+// Copyright 1998-2014 Epic Games, Inc. All Rights Reserved.
 
 #include "Paper2DPrivatePCH.h"
 #include "PaperSprite.h"
 #include "PhysicsEngine/BodySetup.h"
 #include "PaperCustomVersion.h"
 
+#if WITH_EDITOR
+#include "UnrealEd.h"
+#endif
+
 #if WITH_EDITOR
 
 #include "PhysicsEngine/BodySetup2D.h"
@@ -14,6 +18,7 @@
 #include "BitmapUtils.h"
 #include "ComponentReregisterContext.h"
 
+
 //////////////////////////////////////////////////////////////////////////
 // maf
 
@@ -278,10 +283,36 @@ UPaperSprite::UPaperSprite(const FObjectInitializer& ObjectInitializer)
 
 	static ConstructorHelpers::FObjectFinder<UMaterialInterface> OpaqueMaterialRef(TEXT("/Paper2D/OpaqueUnlitSpriteMaterial"));
 	AlternateMaterial = OpaqueMaterialRef.Object;
+
+#if WITH_EDITOR
+	// Hook into notifications for object re-imports so that the gameplay tag tree can be reconstructed if the table changes
+	if (GIsEditor && !bRegisteredObjectReimport)
+	{
+		bRegisteredObjectReimport = true;
+		FEditorDelegates::OnAssetReimport.AddUObject(this, &UPaperSprite::OnObjectReimported);
+	}
+#endif
 }
 
 #if WITH_EDITOR
 
+void UPaperSprite::OnObjectReimported(UObject* InObject)
+{
+	// Check if its our source texture, and if its dimensions have changed
+	// If SourceTetxureDimension == 0, we don't have a previous dimension to work off, so can't
+	// rescale sensibly
+	UTexture2D* Texture = Cast<UTexture2D>(InObject);
+	if (Texture != nullptr && Texture == GetSourceTexture() && NeedRescaleSpriteData())
+	{
+		RescaleSpriteData(GetSourceTexture());
+		PostEditChange();
+	}
+}
+
+#endif
+
+
+#if WITH_EDITOR
 
 /** Removes all components that use the specified sprite asset from their scenes for the lifetime of the class. */
 class FSpriteReregisterContext
@@ -416,14 +447,12 @@ void UPaperSprite::PostEditChangeProperty(FPropertyChangedEvent& PropertyChanged
 	}
 
 	// The texture dimensions have changed
-	if (SourceTexture != nullptr && (SourceTextureDimension.X != SourceTexture->GetSizeX() || SourceTextureDimension.Y != SourceTexture->GetSizeY()))
-	{
-		FVector2D PreviousTextureDimension = SourceTextureDimension;
-		SourceTextureDimension = FVector2D(SourceTexture->GetSizeX(), SourceTexture->GetSizeY());
-		// Tmp - disabled, not sure if we want it on here or not
-		// RescaleSpriteData(PreviousTextureDimension, SourceTextureDimension);
-		bBothModified = true;
-	}
+	//if (NeedRescaleSpriteData())
+	//{
+		// TMP: Disabled, not sure if we want this here
+		// RescaleSpriteData(GetSourceTexture());
+		// bBothModified = true;
+	//}
 
 
 	if (bCollisionDataModified || bBothModified)
@@ -439,18 +468,22 @@ void UPaperSprite::PostEditChangeProperty(FPropertyChangedEvent& PropertyChanged
 	Super::PostEditChangeProperty(PropertyChangedEvent);
 }
 
-void UPaperSprite::RescaleSpriteData(const FVector2D& PreviousTextureDimension, const FVector2D& CurrentTextureDimension)
+void UPaperSprite::RescaleSpriteData(const UTexture2D* Texture)
 {
+	FVector2D PreviousTextureDimension = SourceTextureDimension;
+	FVector2D NewTextureDimension(Texture->GetImportedSize().X, Texture->GetImportedSize().Y);
+
 	// Don't ever divby0 (no previously stored texture dimensions)
-	// or scale to 0
-	if (CurrentTextureDimension.X == 0 || CurrentTextureDimension.Y == 0 ||
+	// or scale to 0, should be covered by NeedRescaleSpriteData
+	if (NewTextureDimension.X == 0 || NewTextureDimension.Y == 0 ||
 		PreviousTextureDimension.X == 0 || PreviousTextureDimension.Y == 0)
 	{
 		return;
 	}
 
-	const FVector2D& S = CurrentTextureDimension;
+	const FVector2D& S = NewTextureDimension;
 	const FVector2D& D = PreviousTextureDimension;
+
 	struct Local
 	{
 		static float NoSnap(const float Value, const float Scale, const float Divisor)
@@ -484,6 +517,7 @@ void UPaperSprite::RescaleSpriteData(const FVector2D& PreviousTextureDimension,
 	SourceUV = Local::Rescale(SourceUV, S, D);
 	SourceDimension = Local::Rescale(SourceDimension, S, D);
 	SourceImageDimensionBeforeTrimming = Local::Rescale(SourceImageDimensionBeforeTrimming, S, D);
+	SourceTextureDimension = NewTextureDimension;
 
 	if (bSnapPivotToPixelGrid)
 	{
@@ -522,6 +556,18 @@ void UPaperSprite::RescaleSpriteData(const FVector2D& PreviousTextureDimension,
 	}
 }
 
+bool UPaperSprite::NeedRescaleSpriteData()
+{
+	if (UTexture2D* Texture = GetSourceTexture())
+	{
+		FIntPoint TextureSize = Texture->GetImportedSize();
+		bool bTextureSizeIsZero = TextureSize.X == 0 || TextureSize.Y == 0;
+		return !SourceTextureDimension.IsZero() && !bTextureSizeIsZero && (TextureSize.X != SourceTextureDimension.X || TextureSize.Y != SourceTextureDimension.Y);
+	}
+
+	return false;
+}
+
 void UPaperSprite::RebuildCollisionData()
 {
 	UBodySetup* OldBodySetup = BodySetup;
@@ -1214,7 +1260,7 @@ void UPaperSprite::InitializeSprite(const FSpriteAssetInitParameters& InitParams
 	SourceTexture = InitParams.Texture;
 	if (SourceTexture != nullptr)
 	{
-		SourceTextureDimension.Set(SourceTexture->GetSizeX(), SourceTexture->GetSizeY());
+		SourceTextureDimension.Set(SourceTexture->GetImportedSize().X, SourceTexture->GetImportedSize().Y);
 	}
 	else
 	{
@@ -1528,12 +1574,19 @@ void UPaperSprite::PostLoad()
 	{
 		bRebuildCollision = true;
 	}
-	
+
 	if ((PaperVer < FPaperCustomVersion::AddDefaultCollisionProfileInSpriteAsset) && (BodySetup != nullptr))
 	{
 		BodySetup->DefaultInstance.SetCollisionProfileName(UCollisionProfile::BlockAllDynamic_ProfileName);
 	}
 
+	if (PaperVer >= FPaperCustomVersion::AddSourceTextureSize && NeedRescaleSpriteData())
+	{
+		RescaleSpriteData(GetSourceTexture());
+		bRebuildCollision = true;
+		bRebuildRenderData = true;
+	}
+
 	if (bRebuildCollision)
 	{
 		RebuildCollisionData();
diff --git a/Engine/Source/Editor/UnrealEd/Classes/Editor/EditorEngine.h b/Engine/Source/Editor/UnrealEd/Classes/Editor/EditorEngine.h
--- a/Engine/Source/Editor/UnrealEd/Classes/Editor/EditorEngine.h
+++ b/Engine/Source/Editor/UnrealEd/Classes/Editor/EditorEngine.h
@@ -603,7 +603,7 @@ public:
 	void BroadcastEndObjectMovement( UObject& Object ) const { OnEndObjectTransformEvent.Broadcast( Object ); }
 
 	/**	Broadcasts that an object has been reimported. THIS SHOULD NOT BE PUBLIC */
-	void BroadcastObjectReimported(UObject* InObject) { ObjectReimportedEvent.Broadcast(InObject); }
+	void BroadcastObjectReimported(UObject* InObject);
 
 	// Begin UObject interface.
 	virtual void FinishDestroy() override;	
diff --git a/Engine/Source/Editor/UnrealEd/Private/Editor.cpp b/Engine/Source/Editor/UnrealEd/Private/Editor.cpp
--- a/Engine/Source/Editor/UnrealEd/Private/Editor.cpp
+++ b/Engine/Source/Editor/UnrealEd/Private/Editor.cpp
@@ -132,6 +132,7 @@ FEditorDelegates::FOnNewAssetCreation					FEditorDelegates::OnConfigureNewAssetP
 FEditorDelegates::FOnNewAssetCreation					FEditorDelegates::OnNewAssetCreated;
 FEditorDelegates::FOnAssetPreImport						FEditorDelegates::OnAssetPreImport;
 FEditorDelegates::FOnAssetPostImport					FEditorDelegates::OnAssetPostImport;
+FEditorDelegates::FOnAssetReimport						FEditorDelegates::OnAssetReimport;
 FEditorDelegates::FOnNewActorsDropped					FEditorDelegates::OnNewActorsDropped;
 FEditorDelegates::FOnGridSnappingChanged				FEditorDelegates::OnGridSnappingChanged;
 FSimpleMulticastDelegate								FEditorDelegates::OnLightingBuildStarted;
@@ -746,6 +747,12 @@ void UEditorEngine::InitBuilderBrush( UWorld* InWorld )
 	}
 }
 
+void UEditorEngine::BroadcastObjectReimported(UObject* InObject)
+{
+	ObjectReimportedEvent.Broadcast(InObject);
+	FEditorDelegates::OnAssetReimport.Broadcast(InObject);
+}
+
 void UEditorEngine::FinishDestroy()
 {
 	if ( !HasAnyFlags(RF_ClassDefaultObject) )
diff --git a/Engine/Source/Editor/UnrealEd/Public/Editor.h b/Engine/Source/Editor/UnrealEd/Public/Editor.h
--- a/Engine/Source/Editor/UnrealEd/Public/Editor.h
+++ b/Engine/Source/Editor/UnrealEd/Public/Editor.h
@@ -72,6 +72,8 @@ struct UNREALED_API FEditorDelegates
 	DECLARE_MULTICAST_DELEGATE_FiveParams(FOnAssetPreImport, UFactory*, UClass*, UObject*, const FName&, const TCHAR*);
 	/** delegate type fired when new assets have been (re-)imported. Note: InCreatedObject can be NULL if import failed. Params: UFactory* InFactory, UObject* InCreatedObject */
 	DECLARE_MULTICAST_DELEGATE_TwoParams(FOnAssetPostImport, UFactory*, UObject*);
+	/** delegate type fired when new assets have been reimported. Note: InCreatedObject can be NULL if import failed. UObject* InCreatedObject */
+	DECLARE_MULTICAST_DELEGATE_OneParam(FOnAssetReimport, UObject*);
 	/** delegate type for finishing up construction of a new blueprint */
 	DECLARE_MULTICAST_DELEGATE_OneParam(FOnFinishPickingBlueprintClass, UClass*);
 	/** delegate type for triggering when new actors are dropped on to the viewport */
@@ -159,6 +161,8 @@ struct UNREALED_API FEditorDelegates
 	static FOnAssetPreImport OnAssetPreImport;
 	/** Called when new assets have been (re-)imported. */
 	static FOnAssetPostImport OnAssetPostImport;
+	/** Called after an asset has been reimported */
+	static FOnAssetReimport OnAssetReimport;
 	/** Called when new actors are dropped on to the viewport */
 	static FOnNewActorsDropped OnNewActorsDropped;
 	/** Called when grid snapping is changed */
