Paper2D: Remove a out-of-date (uncalled) AddReferencedObjects impl.

diff --git a/Engine/Plugins/2D/Paper2D/Source/Paper2DEditor/Private/SPaperEditorViewport.cpp b/Engine/Plugins/2D/Paper2D/Source/Paper2DEditor/Private/SPaperEditorViewport.cpp
--- a/Engine/Plugins/2D/Paper2D/Source/Paper2DEditor/Private/SPaperEditorViewport.cpp
+++ b/Engine/Plugins/2D/Paper2D/Source/Paper2DEditor/Private/SPaperEditorViewport.cpp
@@ -474,11 +474,6 @@ void SPaperEditorViewport::RefreshViewport()
 	Viewport->InvalidateDisplay();
 }
 
-void SPaperEditorViewport::AddReferencedObjects(FReferenceCollector& Collector)
-{
-	ViewportClient->AddReferencedObjects(Collector);
-}
-
 /** Search through the list of zoom levels and find the one closest to what was passed in. */
 int32 SPaperEditorViewport::FindNearestZoomLevel(int32 CurrentZoomLevel, float InZoomAmount) const
 {
diff --git a/Engine/Plugins/2D/Paper2D/Source/Paper2DEditor/Private/SPaperEditorViewport.h b/Engine/Plugins/2D/Paper2D/Source/Paper2DEditor/Private/SPaperEditorViewport.h
--- a/Engine/Plugins/2D/Paper2D/Source/Paper2DEditor/Private/SPaperEditorViewport.h
+++ b/Engine/Plugins/2D/Paper2D/Source/Paper2DEditor/Private/SPaperEditorViewport.h
@@ -95,10 +95,6 @@ public:
 	/** Refreshes the viewport */
 	void RefreshViewport();
 
-	/** Serialization */
-	void AddReferencedObjects(FReferenceCollector& Collector);
-
-
 	float GetZoomAmount() const;
 	FText GetZoomText() const;
 	FSlateColor GetZoomTextColorAndOpacity() const;
@@ -162,7 +158,7 @@ protected:
 	bool bShowSoftwareCursor;
 
 	/** Level viewport client */
-	TSharedPtr<FPaperEditorViewportClient> ViewportClient;
+	TSharedPtr<class FPaperEditorViewportClient> ViewportClient;
 
 	/** Slate viewport for rendering and I/O */
 	TSharedPtr<FSceneViewport> Viewport;
