Paper2D: Make the selection process in tile sets start on mouse down so that single-clicks without a drag still select something

diff --git a/Engine/Plugins/2D/Paper2D/Source/Paper2DEditor/Private/SPaperEditorViewport.cpp b/Engine/Plugins/2D/Paper2D/Source/Paper2DEditor/Private/SPaperEditorViewport.cpp
--- a/Engine/Plugins/2D/Paper2D/Source/Paper2DEditor/Private/SPaperEditorViewport.cpp
+++ b/Engine/Plugins/2D/Paper2D/Source/Paper2DEditor/Private/SPaperEditorViewport.cpp
@@ -202,32 +202,13 @@ FReply SPaperEditorViewport::OnMouseButtonDown(const FGeometry& MyGeometry, cons
 	}
 	else if (MouseEvent.GetEffectingButton() == EKeys::LeftMouseButton)
 	{
-		// LEFT BUTTON is for selecting nodes and manipulating pins.
-// 		FArrangedChildren ArrangedChildren(EVisibility::Visible);
-// 		ArrangeChildren(MyGeometry, ArrangedChildren);
-// 
-// 		const int32 NodeUnderMouseIndex = SWidget::FindChildUnderMouse( ArrangedChildren, MouseEvent );
-// 		if ( NodeUnderMouseIndex != INDEX_NONE )
-// 		{
-// 			// PRESSING ON A NODE!
-// 
-// 			// This changes selection and starts dragging it.
-// 			const FArrangedWidget& NodeGeometry = ArrangedChildren(NodeUnderMouseIndex);
-// 			const FVector2D MousePositionInNode = NodeGeometry.Geometry.AbsoluteToLocal(MouseEvent.GetScreenSpacePosition());
-// 			TSharedRef<SNode> NodeWidgetUnderMouse = StaticCastSharedRef<SNode>( NodeGeometry.Widget );
-// 
-// 			// Track the node that we're dragging; we will move it in OnMouseMove.
-// 			this->OnBeginNodeInteraction(NodeWidgetUnderMouse, MousePositionInNode);
-// 		}
-// 		else
-		{	
-			// START MARQUEE SELECTION.
-			const FVector2D GraphMousePos = PanelCoordToGraphCoord( MyGeometry.AbsoluteToLocal( MouseEvent.GetScreenSpacePosition() ) );
-			Marquee.Start( GraphMousePos, FMarqueeOperation::OperationTypeFromMouseEvent(MouseEvent) );
-
-			// If we're marquee selecting, then we're not clicking on a node!
-			//NodeUnderMousePtr.Reset();
-		}
+		// START MARQUEE SELECTION.
+		const FVector2D GraphMousePos = PanelCoordToGraphCoord( MyGeometry.AbsoluteToLocal( MouseEvent.GetScreenSpacePosition() ) );
+		Marquee.Start( GraphMousePos, FMarqueeOperation::OperationTypeFromMouseEvent(MouseEvent) );
+
+		// Trigger a selection update now so that single-clicks without a drag still select something
+		OnSelectionChanged.ExecuteIfBound(Marquee, true);
+		ViewportClient->Invalidate();
 
 		return FReply::Handled().CaptureMouse( SharedThis(this) );
 	}
