Paper2D: Fix a crash when leaving tile map editing mode (freeing mode tools too early)

diff --git a/Engine/Plugins/2D/Paper2D/Source/Paper2DEditor/Private/PaperEditorViewportClient.cpp b/Engine/Plugins/2D/Paper2D/Source/Paper2DEditor/Private/PaperEditorViewportClient.cpp
--- a/Engine/Plugins/2D/Paper2D/Source/Paper2DEditor/Private/PaperEditorViewportClient.cpp
+++ b/Engine/Plugins/2D/Paper2D/Source/Paper2DEditor/Private/PaperEditorViewportClient.cpp
@@ -81,8 +81,6 @@ FPaperEditorViewportClient::FPaperEditorViewportClient()
 
 FPaperEditorViewportClient::~FPaperEditorViewportClient()
 {
-	//DestroyCheckerboardTexture();
-	delete ModeTools;
 }
 
 void FPaperEditorViewportClient::Draw(FViewport* Viewport, FCanvas* Canvas)
