Paper2D: Add FPaperTileInfo::IsValid to determine if a single tile is valid or not

diff --git a/Engine/Plugins/2D/Paper2D/Source/Paper2D/Classes/PaperTileLayer.h b/Engine/Plugins/2D/Paper2D/Source/Paper2D/Classes/PaperTileLayer.h
--- a/Engine/Plugins/2D/Paper2D/Source/Paper2D/Classes/PaperTileLayer.h
+++ b/Engine/Plugins/2D/Paper2D/Source/Paper2D/Classes/PaperTileLayer.h
@@ -32,6 +32,11 @@ struct FPaperTileInfo
 	{
 		return !(*this == Other);
 	}
+
+	bool IsValid() const
+	{
+		return (TileSet != nullptr) && (PackedTileIndex != INDEX_NONE);
+	}
 };
 
 UCLASS()
