Paper2D: Change colors for tile set and tile map assets away from cyan

diff --git a/Engine/Plugins/2D/Paper2D/Source/Paper2DEditor/Private/TileMapEditing/TileMapAssetTypeActions.cpp b/Engine/Plugins/2D/Paper2D/Source/Paper2DEditor/Private/TileMapEditing/TileMapAssetTypeActions.cpp
--- a/Engine/Plugins/2D/Paper2D/Source/Paper2DEditor/Private/TileMapEditing/TileMapAssetTypeActions.cpp
+++ b/Engine/Plugins/2D/Paper2D/Source/Paper2DEditor/Private/TileMapEditing/TileMapAssetTypeActions.cpp
@@ -17,7 +17,7 @@ FText FTileMapAssetTypeActions::GetName() const
 
 FColor FTileMapAssetTypeActions::GetTypeColor() const
 {
-	return FColor::Cyan;
+	return FColorList::BrightGold;
 }
 
 UClass* FTileMapAssetTypeActions::GetSupportedClass() const
diff --git a/Engine/Plugins/2D/Paper2D/Source/Paper2DEditor/Private/TileSetAssetTypeActions.cpp b/Engine/Plugins/2D/Paper2D/Source/Paper2DEditor/Private/TileSetAssetTypeActions.cpp
--- a/Engine/Plugins/2D/Paper2D/Source/Paper2DEditor/Private/TileSetAssetTypeActions.cpp
+++ b/Engine/Plugins/2D/Paper2D/Source/Paper2DEditor/Private/TileSetAssetTypeActions.cpp
@@ -1,9 +1,6 @@
 // Copyright 1998-2015 Epic Games, Inc. All Rights Reserved.
 
 #include "Paper2DEditorPrivatePCH.h"
-//#include "AssetToolsModule.h"
-//#include "Toolkits/AssetEditorManager.h"
-//#include "AssetRegistryModule.h"
 #include "TileSetAssetTypeActions.h"
 #include "TileSetEditor.h"
 
@@ -21,7 +18,7 @@ FText FTileSetAssetTypeActions::GetName() const
 
 FColor FTileSetAssetTypeActions::GetTypeColor() const
 {
-	return FColor::Cyan;
+	return FColorList::Orange;
 }
 
 UClass* FTileSetAssetTypeActions::GetSupportedClass() const
