Paper2D: Allow FPaperTileMapDetailsCustomization to work directly on UPaperTileMapRenderComponent instances as well

diff --git a/Engine/Plugins/2D/Paper2D/Source/Paper2DEditor/Private/TileMapEditing/PaperTileMapDetailsCustomization.cpp b/Engine/Plugins/2D/Paper2D/Source/Paper2DEditor/Private/TileMapEditing/PaperTileMapDetailsCustomization.cpp
--- a/Engine/Plugins/2D/Paper2D/Source/Paper2DEditor/Private/TileMapEditing/PaperTileMapDetailsCustomization.cpp
+++ b/Engine/Plugins/2D/Paper2D/Source/Paper2DEditor/Private/TileMapEditing/PaperTileMapDetailsCustomization.cpp
@@ -32,13 +32,18 @@ void FPaperTileMapDetailsCustomization::CustomizeDetails(IDetailLayoutBuilder& D
 		UObject* TestObject = SelectedObjects[ObjectIndex].Get();
 		if (AActor* CurrentActor = Cast<AActor>(TestObject))
 		{
-			if (UPaperTileMapRenderComponent* CurrentTileMap = CurrentActor->FindComponentByClass<UPaperTileMapRenderComponent>())
+			if (UPaperTileMapRenderComponent* CurrentComponent = CurrentActor->FindComponentByClass<UPaperTileMapRenderComponent>())
 			{
 				bEditingActor = true;
-				TileMap = CurrentTileMap->TileMap;
+				TileMap = CurrentComponent->TileMap;
 				break;
 			}
 		}
+		else if (UPaperTileMapRenderComponent* TestComponent = Cast<UPaperTileMapRenderComponent>(TestObject))
+		{
+			TileMap = TestComponent->TileMap;
+			break;
+		}
 		else if (UPaperTileMap* TestTileMap = Cast<UPaperTileMap>(TestObject))
 		{
 			TileMap = TestTileMap;
