Paper2D: Make the tile map layer list restore the layer selection state from last time when created, rather than always selecting the top layer

diff --git a/Engine/Plugins/2D/Paper2D/Source/Paper2DEditor/Private/TileMapEditing/STileLayerList.cpp b/Engine/Plugins/2D/Paper2D/Source/Paper2DEditor/Private/TileMapEditing/STileLayerList.cpp
--- a/Engine/Plugins/2D/Paper2D/Source/Paper2DEditor/Private/TileMapEditing/STileLayerList.cpp
+++ b/Engine/Plugins/2D/Paper2D/Source/Paper2DEditor/Private/TileMapEditing/STileLayerList.cpp
@@ -75,10 +75,11 @@ void STileLayerList::Construct(const FArguments& InArgs, UPaperTileMap* TileMap)
 		.OnGenerateRow(this, &STileLayerList::OnGenerateLayerListRow)
 		.OnContextMenuOpening(this, &STileLayerList::OnConstructContextMenu);
 
-	// Select the top item by default
-	if (TileMap->TileLayers.Num() > 0)
+	// Restore the selection
+	TileMap->ValidateSelectedLayerIndex();
+	if (TileMap->SelectedLayerIndex != INDEX_NONE)
 	{
-		SetSelectedLayer(TileMap->TileLayers[TileMap->TileLayers.Num() - 1]);
+		SetSelectedLayer(TileMap->TileLayers[TileMap->SelectedLayerIndex]);
 	}
 
 	ChildSlot
