Paper2D: Implement tile map layer merging

diff --git a/Engine/Plugins/2D/Paper2D/Source/Paper2DEditor/Private/TileMapEditing/STileLayerList.cpp b/Engine/Plugins/2D/Paper2D/Source/Paper2DEditor/Private/TileMapEditing/STileLayerList.cpp
--- a/Engine/Plugins/2D/Paper2D/Source/Paper2DEditor/Private/TileMapEditing/STileLayerList.cpp
+++ b/Engine/Plugins/2D/Paper2D/Source/Paper2DEditor/Private/TileMapEditing/STileLayerList.cpp
@@ -338,7 +338,43 @@ void STileLayerList::DuplicateLayer()
 
 void STileLayerList::MergeLayerDown()
 {
-	//@TODO: TILEMAPS: Support merging layers down
+	if (UPaperTileMap* TileMap = TileMapPtr.Get())
+	{
+		const int32 SourceIndex = GetSelectionIndex();
+		const int32 TargetIndex = SourceIndex + 1;
+		if ((SourceIndex != INDEX_NONE) && (TargetIndex != INDEX_NONE))
+		{
+			const FScopedTransaction Transaction(LOCTEXT("TileMapMergeLayerDown", "Merge Layer Down"));
+			TileMap->SetFlags(RF_Transactional);
+			TileMap->Modify();
+
+			UPaperTileLayer* SourceLayer = TileMap->TileLayers[SourceIndex];
+			UPaperTileLayer* TargetLayer = TileMap->TileLayers[TargetIndex];
+
+			TargetLayer->SetFlags(RF_Transactional);
+			TargetLayer->Modify();
+
+			// Copy the non-empty tiles from the source to the target layer
+			for (int32 Y = 0; Y < SourceLayer->LayerWidth; ++Y)
+			{
+				for (int32 X = 0; X < SourceLayer->LayerWidth; ++X)
+				{
+					FPaperTileInfo TileInfo = SourceLayer->GetCell(X, Y);
+					if (TileInfo.IsValid())
+					{
+						TargetLayer->SetCell(X, Y, TileInfo);
+					}
+				}
+			}
+
+			// Remove the source layer
+			TileMap->TileLayers.RemoveAt(SourceIndex);
+
+			// Update viewers
+			TileMap->PostEditChange();
+			ListViewWidget->RequestListRefresh();
+		}
+	}
 }
 
 void STileLayerList::MoveLayerUp()
