Paper2D: Prevent drawing the grid lines for the selected layer in a tile map when there are no layers

diff --git a/Engine/Plugins/2D/Paper2D/Source/Paper2D/Private/PaperTileMapRenderSceneProxy.cpp b/Engine/Plugins/2D/Paper2D/Source/Paper2D/Private/PaperTileMapRenderSceneProxy.cpp
--- a/Engine/Plugins/2D/Paper2D/Source/Paper2D/Private/PaperTileMapRenderSceneProxy.cpp
+++ b/Engine/Plugins/2D/Paper2D/Source/Paper2D/Private/PaperTileMapRenderSceneProxy.cpp
@@ -116,10 +116,12 @@ void FPaperTileMapRenderSceneProxy::GetDynamicMeshElements(const TArray<const FS
 
 					if (bUseOverrideColor)
 					{
+						const int32 SelectedLayerIndex = TileMap->SelectedLayerIndex;
+
 						// Draw a bound for any invisible layers
 						for (int32 LayerIndex = 0; LayerIndex < TileMap->TileLayers.Num(); ++LayerIndex)
 						{
-							if (LayerIndex != TileMap->SelectedLayerIndex)
+							if (LayerIndex != SelectedLayerIndex)
 							{
 								const FVector TL(LocalToWorld.TransformPosition(TileMap->GetTilePositionInLocalSpace(0, 0, LayerIndex)));
 								const FVector TR(LocalToWorld.TransformPosition(TileMap->GetTilePositionInLocalSpace(TileMap->MapWidth, 0, LayerIndex)));
@@ -133,28 +135,31 @@ void FPaperTileMapRenderSceneProxy::GetDynamicMeshElements(const TArray<const FS
 							}
 						}
 
-						// Draw horizontal lines on the selection
-						for (int32 Y = 0; Y <= TileMap->MapHeight; ++Y)
+						if (SelectedLayerIndex != INDEX_NONE)
 						{
-							int32 X = 0;
-							const FVector Start(TileMap->GetTilePositionInLocalSpace(X, Y, TileMap->SelectedLayerIndex));
+							// Draw horizontal lines on the selection
+							for (int32 Y = 0; Y <= TileMap->MapHeight; ++Y)
+							{
+								int32 X = 0;
+								const FVector Start(TileMap->GetTilePositionInLocalSpace(X, Y, SelectedLayerIndex));
 
-							X = TileMap->MapWidth;
-							const FVector End(TileMap->GetTilePositionInLocalSpace(X, Y, TileMap->SelectedLayerIndex));
+								X = TileMap->MapWidth;
+								const FVector End(TileMap->GetTilePositionInLocalSpace(X, Y, SelectedLayerIndex));
 
-							PDI->DrawLine(LocalToWorld.TransformPosition(Start), LocalToWorld.TransformPosition(End), OverrideColor, DPG, 0.0f, DepthBias);
-						}
+								PDI->DrawLine(LocalToWorld.TransformPosition(Start), LocalToWorld.TransformPosition(End), OverrideColor, DPG, 0.0f, DepthBias);
+							}
 
-						// Draw vertical lines
-						for (int32 X = 0; X <= TileMap->MapWidth; ++X)
-						{
-							int32 Y = 0;
-							const FVector Start(TileMap->GetTilePositionInLocalSpace(X, Y, TileMap->SelectedLayerIndex));
+							// Draw vertical lines
+							for (int32 X = 0; X <= TileMap->MapWidth; ++X)
+							{
+								int32 Y = 0;
+								const FVector Start(TileMap->GetTilePositionInLocalSpace(X, Y, SelectedLayerIndex));
 
-							Y = TileMap->MapHeight;
-							const FVector End(TileMap->GetTilePositionInLocalSpace(X, Y, TileMap->SelectedLayerIndex));
+								Y = TileMap->MapHeight;
+								const FVector End(TileMap->GetTilePositionInLocalSpace(X, Y, SelectedLayerIndex));
 
-							PDI->DrawLine(LocalToWorld.TransformPosition(Start), LocalToWorld.TransformPosition(End), OverrideColor, DPG, 0.0f, DepthBias);
+								PDI->DrawLine(LocalToWorld.TransformPosition(Start), LocalToWorld.TransformPosition(End), OverrideColor, DPG, 0.0f, DepthBias);
+							}
 						}
 					}
 				}
