Paper2D: Add a method to determine if a tile map render component is instanced or referencing an asset (UPaperTileMapRenderComponent::OwnsTileMap)

diff --git a/Engine/Plugins/2D/Paper2D/Source/Paper2D/Classes/PaperTileMapRenderComponent.h b/Engine/Plugins/2D/Paper2D/Source/Paper2D/Classes/PaperTileMapRenderComponent.h
--- a/Engine/Plugins/2D/Paper2D/Source/Paper2D/Classes/PaperTileMapRenderComponent.h
+++ b/Engine/Plugins/2D/Paper2D/Source/Paper2D/Classes/PaperTileMapRenderComponent.h
@@ -66,4 +66,7 @@ public:
 	virtual FBoxSphereBounds CalcBounds(const FTransform& LocalToWorld) const override;
 	virtual class UBodySetup* GetBodySetup() override;
 	// End of UPrimitiveComponent interface
+
+	// Does this component own the tile map (is it instanced instead of being an asset reference)?
+	bool OwnsTileMap() const;
 };
diff --git a/Engine/Plugins/2D/Paper2D/Source/Paper2D/Private/PaperTileMapRenderComponent.cpp b/Engine/Plugins/2D/Paper2D/Source/Paper2D/Private/PaperTileMapRenderComponent.cpp
--- a/Engine/Plugins/2D/Paper2D/Source/Paper2D/Private/PaperTileMapRenderComponent.cpp
+++ b/Engine/Plugins/2D/Paper2D/Source/Paper2D/Private/PaperTileMapRenderComponent.cpp
@@ -310,4 +310,9 @@ void UPaperTileMapRenderComponent::RebuildRenderData(FPaperTileMapRenderScenePro
 	Proxy->SetBatchesHack(BatchedSprites);
 }
 
+bool UPaperTileMapRenderComponent::OwnsTileMap() const
+{
+	return (TileMap != nullptr) && (TileMap->GetOuter() == this);
+}
+
 #undef LOCTEXT_NAMESPACE
\ No newline at end of file
