Paper2D: Add rudimentary draw call batching support to tile map rendering

diff --git a/Engine/Plugins/2D/Paper2D/Source/Paper2D/Private/PaperTileMapRenderComponent.cpp b/Engine/Plugins/2D/Paper2D/Source/Paper2D/Private/PaperTileMapRenderComponent.cpp
--- a/Engine/Plugins/2D/Paper2D/Source/Paper2D/Private/PaperTileMapRenderComponent.cpp
+++ b/Engine/Plugins/2D/Paper2D/Source/Paper2D/Private/PaperTileMapRenderComponent.cpp
@@ -149,7 +149,6 @@ void UPaperTileMapRenderComponent::RebuildRenderData(FPaperTileMapRenderScenePro
 	FVector WorldOffsetYFactor = FVector::ZeroVector;
 	FVector WorldStepX;
 	FVector WorldStepY;
-	FVector TileSetOffset = FVector::ZeroVector;
 
 	switch (TileMap->ProjectionMode)
 	{
@@ -173,6 +172,11 @@ void UPaperTileMapRenderComponent::RebuildRenderData(FPaperTileMapRenderScenePro
 		break;
 	}
 
+	UTexture2D* LastSourceTexture = nullptr;
+	FVector TileSetOffset = FVector::ZeroVector;
+	FVector2D InverseTextureSize(1.0f, 1.0f);
+	FVector2D SourceDimensionsUV(1.0f, 1.0f);
+	FVector2D TileSizeXY(0.0f, 0.0f);
 
 	for (int32 Z = 0; Z < TileMap->TileLayers.Num(); ++Z)
 	{
@@ -193,10 +197,11 @@ void UPaperTileMapRenderComponent::RebuildRenderData(FPaperTileMapRenderScenePro
 		DrawColor.A = Layer->LayerOpacity;
 #endif
 
+		FSpriteDrawCallRecord* CurrentBatch = nullptr;
+
 		for (int32 Y = 0; Y < TileMap->MapHeight; ++Y)
 		{
 			// In pixels
-			FVector TileSetOffset = FVector::ZeroVector;
 			FVector WorldOffset;
 
 			switch (TileMap->ProjectionMode)
@@ -225,12 +230,6 @@ void UPaperTileMapRenderComponent::RebuildRenderData(FPaperTileMapRenderScenePro
 				const int32 TileWidth = TileMap->TileWidth;
 				const int32 TileHeight = TileMap->TileHeight;
 
-				UTexture2D* LastSourceTexture = nullptr;
-				FVector2D InverseTextureSize(1.0f, 1.0f);
-				FVector2D SourceDimensionsUV(1.0f, 1.0f);
-
-				FVector2D TileSeparation(TileMap->TileWidth, TileMap->TileHeight);
-				FVector2D TileSizeXY(TileMap->TileWidth, TileMap->TileHeight);
 
 				{
 					UTexture2D* SourceTexture = nullptr;
@@ -263,6 +262,14 @@ void UPaperTileMapRenderComponent::RebuildRenderData(FPaperTileMapRenderScenePro
 						}
 					}
 
+					if ((SourceTexture != LastSourceTexture) || (CurrentBatch == nullptr))
+					{
+						CurrentBatch = (new (BatchedSprites) FSpriteDrawCallRecord());
+						CurrentBatch->Texture = SourceTexture;
+						CurrentBatch->Color = DrawColor;
+						CurrentBatch->Destination = WorldPos.ProjectOnTo(PaperAxisZ);
+					}
+
 					if (SourceTexture != LastSourceTexture)
 					{
 						InverseTextureSize = FVector2D(1.0f / SourceTexture->GetSizeX(), 1.0f / SourceTexture->GetSizeY());
@@ -270,30 +277,31 @@ void UPaperTileMapRenderComponent::RebuildRenderData(FPaperTileMapRenderScenePro
 						if (TileInfo.TileSet != nullptr)
 						{
 							SourceDimensionsUV = FVector2D(TileInfo.TileSet->TileWidth * InverseTextureSize.X, TileInfo.TileSet->TileHeight * InverseTextureSize.Y);
-							TileSeparation = FVector2D(TileInfo.TileSet->TileWidth, TileInfo.TileSet->TileHeight);
+							TileSizeXY = FVector2D(TileInfo.TileSet->TileWidth, TileInfo.TileSet->TileHeight);
 							TileSetOffset = (TileInfo.TileSet->DrawingOffset.X * PaperAxisX) + (TileInfo.TileSet->DrawingOffset.Y * PaperAxisY);
 						}
 						else
 						{
 							SourceDimensionsUV = FVector2D(TileWidth * InverseTextureSize.X, TileHeight * InverseTextureSize.Y);
-							TileSeparation = FVector2D(TileWidth, TileHeight);
+							TileSizeXY = FVector2D(TileWidth, TileHeight);
 							TileSetOffset = FVector::ZeroVector;
 						}
 						LastSourceTexture = SourceTexture;
 					}
+					WorldPos += TileSetOffset;
 
 					SourceUV.X *= InverseTextureSize.X;
 					SourceUV.Y *= InverseTextureSize.Y;
 
-					FSpriteDrawCallRecord& NewTile = *(new (BatchedSprites) FSpriteDrawCallRecord());
-					NewTile.Destination = WorldPos + TileSetOffset;
-					NewTile.Texture = SourceTexture;
-					NewTile.Color = DrawColor;
+					FSpriteDrawCallRecord& NewTile = *CurrentBatch;
+
+					const float WX0 = FVector::DotProduct(WorldPos, PaperAxisX);
+					const float WY0 = FVector::DotProduct(WorldPos, PaperAxisY);
 
-					const FVector4 BottomLeft(0.0f, 0.0f, SourceUV.X, SourceUV.Y + SourceDimensionsUV.Y);
-					const FVector4 BottomRight(TileSeparation.X, 0.0f, SourceUV.X + SourceDimensionsUV.X, SourceUV.Y + SourceDimensionsUV.Y);
-					const FVector4 TopRight(TileSeparation.X, TileSeparation.Y, SourceUV.X + SourceDimensionsUV.X, SourceUV.Y);
-					const FVector4 TopLeft(0.0f, TileSeparation.Y, SourceUV.X, SourceUV.Y);
+					const FVector4 BottomLeft(WX0, WY0, SourceUV.X, SourceUV.Y + SourceDimensionsUV.Y);
+					const FVector4 BottomRight(WX0 + TileSizeXY.X, WY0, SourceUV.X + SourceDimensionsUV.X, SourceUV.Y + SourceDimensionsUV.Y);
+					const FVector4 TopRight(WX0 + TileSizeXY.X, WY0 + TileSizeXY.Y, SourceUV.X + SourceDimensionsUV.X, SourceUV.Y);
+					const FVector4 TopLeft(WX0, WY0 + TileSizeXY.Y, SourceUV.X, SourceUV.Y);
 
 					new (NewTile.RenderVerts) FVector4(BottomLeft);
 					new (NewTile.RenderVerts) FVector4(TopRight);
