Paper2D: Tile selection improvements
- Tile selection now works on the right and bottom edges
- Changing tile sets updates the rectangle preview appropriately

diff --git a/Engine/Plugins/2D/Paper2D/Source/Paper2DEditor/Private/TileSetEditor.cpp b/Engine/Plugins/2D/Paper2D/Source/Paper2DEditor/Private/TileSetEditor.cpp
--- a/Engine/Plugins/2D/Paper2D/Source/Paper2DEditor/Private/TileSetEditor.cpp
+++ b/Engine/Plugins/2D/Paper2D/Source/Paper2DEditor/Private/TileSetEditor.cpp
@@ -135,6 +135,9 @@ STileSetSelectorViewport::~STileSetSelectorViewport()
 
 void STileSetSelectorViewport::Construct(const FArguments& InArgs, UPaperTileSet* InTileSet, FEdModeTileMap* InTileMapEditor)
 {
+	SelectionTopLeft = FIntPoint::ZeroValue;
+	SelectionDimensions = FIntPoint::ZeroValue;
+
 	TileSetPtr = InTileSet;
 	TileMapEditor = InTileMapEditor;
 
@@ -150,8 +153,15 @@ void STileSetSelectorViewport::Construct(const FArguments& InArgs, UPaperTileSet
 
 void STileSetSelectorViewport::ChangeTileSet(UPaperTileSet* InTileSet)
 {
-	TileSetPtr = InTileSet;
-	TypedViewportClient->TileSetBeingEdited = InTileSet;
+	if (InTileSet != TileSetPtr.Get())
+	{
+		TileSetPtr = InTileSet;
+		TypedViewportClient->TileSetBeingEdited = InTileSet;
+
+		// Update the selection rectangle
+		RefreshSelectionRectangle();
+		TypedViewportClient->Invalidate();
+	}
 }
 
 FText STileSetSelectorViewport::GetTitleText() const
@@ -173,22 +183,18 @@ void STileSetSelectorViewport::OnSelectionChanged(FMarqueeOperation Marquee, boo
 		return;
 	}
 
-	FIntPoint SelectionTopLeft;
-	FIntPoint SelectionDimensions;
-
-	if (Marquee.IsValid())
+	const FVector2D TopLeftUnrounded = Marquee.Rect.GetUpperLeft();
+	const FVector2D BottomRightUnrounded = Marquee.Rect.GetLowerRight();
+	if ((TopLeftUnrounded != FVector2D::ZeroVector) || Marquee.IsValid())
 	{
-		const FVector2D TopLeftUnrounded = Marquee.Rect.GetUpperLeft();
-		const FVector2D BottomRightUnrounded = Marquee.Rect.GetLowerRight();
-
 		// Round down the top left corner
-		SelectionTopLeft.X = FMath::Clamp<int>((int)(TopLeftUnrounded.X / TileSetBeingEdited->TileWidth), 0, TileSetBeingEdited->GetTileCountX()-1);
-		SelectionTopLeft.Y = FMath::Clamp<int>((int)(TopLeftUnrounded.Y / TileSetBeingEdited->TileHeight), 0, TileSetBeingEdited->GetTileCountY()-1);
+		SelectionTopLeft.X = FMath::Clamp<int>((int)(TopLeftUnrounded.X / TileSetBeingEdited->TileWidth), 0, TileSetBeingEdited->GetTileCountX());
+		SelectionTopLeft.Y = FMath::Clamp<int>((int)(TopLeftUnrounded.Y / TileSetBeingEdited->TileHeight), 0, TileSetBeingEdited->GetTileCountY());
 
 		// Round up the bottom right corner
 		FIntPoint SelectionBottomRight;
-		SelectionBottomRight.X = FMath::Clamp<int>(FMath::DivideAndRoundUp((int)BottomRightUnrounded.X, TileSetBeingEdited->TileWidth), 0, TileSetBeingEdited->GetTileCountX()-1);
-		SelectionBottomRight.Y = FMath::Clamp<int>(FMath::DivideAndRoundUp((int)BottomRightUnrounded.Y, TileSetBeingEdited->TileHeight), 0, TileSetBeingEdited->GetTileCountY()-1);
+		SelectionBottomRight.X = FMath::Clamp<int>(FMath::DivideAndRoundUp((int)BottomRightUnrounded.X, TileSetBeingEdited->TileWidth), 0, TileSetBeingEdited->GetTileCountX());
+		SelectionBottomRight.Y = FMath::Clamp<int>(FMath::DivideAndRoundUp((int)BottomRightUnrounded.Y, TileSetBeingEdited->TileHeight), 0, TileSetBeingEdited->GetTileCountY());
 
 		// Compute the new selection dimensions
 		SelectionDimensions = SelectionBottomRight - SelectionTopLeft;
@@ -201,7 +207,7 @@ void STileSetSelectorViewport::OnSelectionChanged(FMarqueeOperation Marquee, boo
 		SelectionDimensions.Y = 0;
 	}
 
-	const bool bHasSelection = (SelectionDimensions.X + SelectionDimensions.Y > 0);
+	const bool bHasSelection = (SelectionDimensions.X * SelectionDimensions.Y) > 0;
 	if (bIsPreview && bHasSelection)
 	{
 		if (TileMapEditor != nullptr)
@@ -209,9 +215,21 @@ void STileSetSelectorViewport::OnSelectionChanged(FMarqueeOperation Marquee, boo
 			TileMapEditor->SetActivePaint(TileSetBeingEdited, SelectionTopLeft, SelectionDimensions);
 		}
 
-		if (FTileSetEditorViewportClient* Client = TypedViewportClient.Get())
+		RefreshSelectionRectangle();
+	}
+}
+
+void STileSetSelectorViewport::RefreshSelectionRectangle()
+{
+	if (FTileSetEditorViewportClient* Client = TypedViewportClient.Get())
+	{
+		UPaperTileSet* TileSetBeingEdited = TileSetPtr.Get();
+
+		const bool bHasSelection = (SelectionDimensions.X * SelectionDimensions.Y) > 0;
+		Client->bHasValidPaintRectangle = bHasSelection && (TileSetBeingEdited != nullptr);
+
+		if (bHasSelection)
 		{
-			Client->bHasValidPaintRectangle = true;
 			Client->ValidPaintRectangle.Color = FLinearColor::White;
 			Client->ValidPaintRectangle.Dimensions = FVector2D(SelectionDimensions.X * TileSetBeingEdited->TileWidth, SelectionDimensions.Y * TileSetBeingEdited->TileHeight);
 			Client->ValidPaintRectangle.TopLeft = FVector2D(SelectionTopLeft.X * TileSetBeingEdited->TileWidth, SelectionTopLeft.Y * TileSetBeingEdited->TileHeight);
diff --git a/Engine/Plugins/2D/Paper2D/Source/Paper2DEditor/Private/TileSetEditor.h b/Engine/Plugins/2D/Paper2D/Source/Paper2DEditor/Private/TileSetEditor.h
--- a/Engine/Plugins/2D/Paper2D/Source/Paper2DEditor/Private/TileSetEditor.h
+++ b/Engine/Plugins/2D/Paper2D/Source/Paper2DEditor/Private/TileSetEditor.h
@@ -56,9 +56,12 @@ protected:
 
 private:
 	void OnSelectionChanged(FMarqueeOperation Marquee, bool bIsPreview);
+	void RefreshSelectionRectangle();
 
 private:
 	TWeakObjectPtr<class UPaperTileSet> TileSetPtr;
 	TSharedPtr<class FTileSetEditorViewportClient> TypedViewportClient;
 	class FEdModeTileMap* TileMapEditor;
+	FIntPoint SelectionTopLeft;
+	FIntPoint SelectionDimensions;
 };
\ No newline at end of file
