Paper2D: Switch to paint brush mode if we were in the eraser mode when a new tile selection is made

diff --git a/Engine/Plugins/2D/Paper2D/Source/Paper2DEditor/Private/TileSetEditor.cpp b/Engine/Plugins/2D/Paper2D/Source/Paper2DEditor/Private/TileSetEditor.cpp
--- a/Engine/Plugins/2D/Paper2D/Source/Paper2DEditor/Private/TileSetEditor.cpp
+++ b/Engine/Plugins/2D/Paper2D/Source/Paper2DEditor/Private/TileSetEditor.cpp
@@ -213,6 +213,12 @@ void STileSetSelectorViewport::OnSelectionChanged(FMarqueeOperation Marquee, boo
 		if (TileMapEditor != nullptr)
 		{
 			TileMapEditor->SetActivePaint(TileSetBeingEdited, SelectionTopLeft, SelectionDimensions);
+
+			// Switch to paint brush mode if we were in the eraser mode since the user is trying to select some ink to paint with
+			if (TileMapEditor->GetActiveTool() == ETileMapEditorTool::Eraser)
+			{
+				TileMapEditor->SetActiveTool(ETileMapEditorTool::Paintbrush);
+			}
 		}
 
 		RefreshSelectionRectangle();
