Paper2D: Move AddNewLayer into UPaperTileMap from STileLayerList

diff --git a/Engine/Plugins/2D/Paper2D/Source/Paper2D/Classes/PaperTileMap.h b/Engine/Plugins/2D/Paper2D/Source/Paper2D/Classes/PaperTileMap.h
--- a/Engine/Plugins/2D/Paper2D/Source/Paper2D/Classes/PaperTileMap.h
+++ b/Engine/Plugins/2D/Paper2D/Source/Paper2D/Classes/PaperTileMap.h
@@ -112,6 +112,13 @@ public:
 	FVector GetTilePositionInLocalSpace(int32 TileX, int32 TileY, int32 LayerIndex = 0) const;
 
 	FBoxSphereBounds GetRenderBounds() const;
+
+	// Creates and adds a new layer and returns it
+	class UPaperTileLayer* AddNewLayer(bool bCollisionLayer = false, int32 InsertionIndex = INDEX_NONE);
+
+	// Creates a reasonable new layer name
+	static FText GenerateNewLayerName(UPaperTileMap* TileMap);
+
 protected:
 	virtual void UpdateBodySetup();
 };
diff --git a/Engine/Plugins/2D/Paper2D/Source/Paper2D/Private/PaperTileMap.cpp b/Engine/Plugins/2D/Paper2D/Source/Paper2D/Private/PaperTileMap.cpp
--- a/Engine/Plugins/2D/Paper2D/Source/Paper2D/Private/PaperTileMap.cpp
+++ b/Engine/Plugins/2D/Paper2D/Source/Paper2D/Private/PaperTileMap.cpp
@@ -4,6 +4,8 @@
 #include "PhysicsEngine/BodySetup.h"
 #include "PhysicsEngine/BodySetup2D.h"
 
+#define LOCTEXT_NAMESPACE "Paper2D"
+
 //////////////////////////////////////////////////////////////////////////
 // UPaperTileMap
 
@@ -157,4 +159,57 @@ FBoxSphereBounds UPaperTileMap::GetRenderBounds() const
 
 	const FBox Box(TopLeft, TopLeft + Dimenisons);
 	return FBoxSphereBounds(Box);
-}
\ No newline at end of file
+}
+
+UPaperTileLayer* UPaperTileMap::AddNewLayer(bool bCollisionLayer, int32 InsertionIndex)
+{
+	// Create the new layer
+	UPaperTileLayer* NewLayer = NewObject<UPaperTileLayer>(this);
+	NewLayer->SetFlags(RF_Transactional);
+
+	NewLayer->LayerWidth = MapWidth;
+	NewLayer->LayerHeight = MapHeight;
+	NewLayer->DestructiveAllocateMap(NewLayer->LayerWidth, NewLayer->LayerHeight);
+	NewLayer->LayerName = GenerateNewLayerName(this);
+	NewLayer->bCollisionLayer = bCollisionLayer;
+
+	// Insert the new layer
+	if (TileLayers.IsValidIndex(InsertionIndex))
+	{
+		TileLayers.Insert(NewLayer, InsertionIndex);
+	}
+	else
+	{
+		TileLayers.Add(NewLayer);
+	}
+
+	return NewLayer;
+}
+
+FText UPaperTileMap::GenerateNewLayerName(UPaperTileMap* TileMap)
+{
+	// Create a set of existing names
+	TSet<FString> ExistingNames;
+	for (UPaperTileLayer* ExistingLayer : TileMap->TileLayers)
+	{
+		ExistingNames.Add(ExistingLayer->LayerName.ToString());
+	}
+
+	// Find a good name
+	FText TestLayerName;
+	do
+	{
+		TileMap->LayerNameIndex++;
+
+		FNumberFormattingOptions NoGroupingFormat;
+		NoGroupingFormat.SetUseGrouping(false);
+
+		TestLayerName = FText::Format(LOCTEXT("NewLayerNameFormatString", "Layer {0}"), FText::AsNumber(TileMap->LayerNameIndex, &NoGroupingFormat));
+	} while (ExistingNames.Contains(TestLayerName.ToString()));
+
+	return TestLayerName;
+}
+
+//////////////////////////////////////////////////////////////////////////
+
+#undef LOCTEXT_NAMESPACE
\ No newline at end of file
diff --git a/Engine/Plugins/2D/Paper2D/Source/Paper2DEditor/Private/TileMapEditing/STileLayerList.cpp b/Engine/Plugins/2D/Paper2D/Source/Paper2DEditor/Private/TileMapEditing/STileLayerList.cpp
--- a/Engine/Plugins/2D/Paper2D/Source/Paper2DEditor/Private/TileMapEditing/STileLayerList.cpp
+++ b/Engine/Plugins/2D/Paper2D/Source/Paper2DEditor/Private/TileMapEditing/STileLayerList.cpp
@@ -119,30 +119,6 @@ UPaperTileLayer* STileLayerList::GetSelectedLayer() const
 	return (ListViewWidget->GetNumItemsSelected() > 0) ? ListViewWidget->GetSelectedItems()[0] : nullptr;
 }
 
-FText STileLayerList::GenerateNewLayerName(UPaperTileMap* TileMap)
-{
-	// Create a set of existing names
-	TSet<FString> ExistingNames;
-	for (UPaperTileLayer* ExistingLayer : TileMap->TileLayers)
-	{
-		ExistingNames.Add(ExistingLayer->LayerName.ToString());
-	}
-
-	// Find a good name
-	FText TestLayerName;
-	do
-	{
-		TileMap->LayerNameIndex++;
-
-		FNumberFormattingOptions NoGroupingFormat;
-		NoGroupingFormat.SetUseGrouping(false);
-
-		TestLayerName = FText::Format(LOCTEXT("NewLayerNameFormatString", "Layer {0}"), FText::AsNumber(TileMap->LayerNameIndex, &NoGroupingFormat));
-	} while (ExistingNames.Contains(TestLayerName.ToString()));
-
-	return TestLayerName;
-}
-
 FText STileLayerList::GenerateDuplicatedLayerName(const FString& InputNameRaw, UPaperTileMap* TileMap)
 {
 	// Create a set of existing names
@@ -207,25 +183,7 @@ UPaperTileLayer* STileLayerList::AddLayer(bool bCollisionLayer, int32 InsertionI
 		TileMap->SetFlags(RF_Transactional);
 		TileMap->Modify();
 
-		// Create the new layer
-		NewLayer = NewObject<UPaperTileLayer>(TileMap);
-		NewLayer->SetFlags(RF_Transactional);
-
-		NewLayer->LayerWidth = TileMap->MapWidth;
-		NewLayer->LayerHeight = TileMap->MapHeight;
-		NewLayer->DestructiveAllocateMap(NewLayer->LayerWidth, NewLayer->LayerHeight);
-		NewLayer->LayerName = GenerateNewLayerName(TileMap);
-		NewLayer->bCollisionLayer = bCollisionLayer;
-
-		// Insert the new layer
-		if (TileMap->TileLayers.IsValidIndex(InsertionIndex))
-		{
-			TileMap->TileLayers.Insert(NewLayer, InsertionIndex);
-		}
-		else
-		{
-			TileMap->TileLayers.Add(NewLayer);
-		}
+		NewLayer = TileMap->AddNewLayer(bCollisionLayer, InsertionIndex);
 
 		ListViewWidget->RequestListRefresh();
 		TileMap->PostEditChange();
diff --git a/Engine/Plugins/2D/Paper2D/Source/Paper2DEditor/Private/TileMapEditing/STileLayerList.h b/Engine/Plugins/2D/Paper2D/Source/Paper2DEditor/Private/TileMapEditing/STileLayerList.h
--- a/Engine/Plugins/2D/Paper2D/Source/Paper2DEditor/Private/TileMapEditing/STileLayerList.h
+++ b/Engine/Plugins/2D/Paper2D/Source/Paper2DEditor/Private/TileMapEditing/STileLayerList.h
@@ -29,7 +29,6 @@ protected:
 	// Returns the selected index if anything is selected, or the top item otherwise (only returns INDEX_NONE if there are no layers)
 	int32 GetSelectionIndex() const;
 
-	static FText GenerateNewLayerName(UPaperTileMap* TileMap);
 	static FText GenerateDuplicatedLayerName(const FString& InputNameRaw, UPaperTileMap* TileMap);
 
 	class UPaperTileLayer* AddLayer(bool bCollisionLayer, int32 InsertionIndex = INDEX_NONE);
