Paper2D: Switch UPaperTileMap to use a factory instead of being a UDataAsset and always add a new empty layer to a tile map when it is first created

diff --git a/Engine/Plugins/2D/Paper2D/Source/Paper2D/Classes/PaperTileMap.h b/Engine/Plugins/2D/Paper2D/Source/Paper2D/Classes/PaperTileMap.h
--- a/Engine/Plugins/2D/Paper2D/Source/Paper2D/Classes/PaperTileMap.h
+++ b/Engine/Plugins/2D/Paper2D/Source/Paper2D/Classes/PaperTileMap.h
@@ -26,7 +26,7 @@ namespace ETileMapProjectionMode
 
 
 UCLASS()
-class PAPER2D_API UPaperTileMap : public UDataAsset //@TODO: Just to make it easy to spawn for now
+class PAPER2D_API UPaperTileMap : public UObject
 {
 	GENERATED_UCLASS_BODY()
 
diff --git a/Engine/Plugins/2D/Paper2D/Source/Paper2D/Private/PaperTileMap.cpp b/Engine/Plugins/2D/Paper2D/Source/Paper2D/Private/PaperTileMap.cpp
--- a/Engine/Plugins/2D/Paper2D/Source/Paper2D/Private/PaperTileMap.cpp
+++ b/Engine/Plugins/2D/Paper2D/Source/Paper2D/Private/PaperTileMap.cpp
@@ -24,7 +24,7 @@ UPaperTileMap::UPaperTileMap(const FObjectInitializer& ObjectInitializer)
 
 #if WITH_EDITORONLY_DATA
 	SelectedLayerIndex = INDEX_NONE;
-	LayerNameIndex = 1;
+	LayerNameIndex = 0;
 #endif
 
 	static ConstructorHelpers::FObjectFinder<UMaterial> DefaultMaterial(TEXT("/Paper2D/DefaultSpriteMaterial"));
diff --git a/Engine/Plugins/2D/Paper2D/Source/Paper2DEditor/Classes/PaperTileMapFactory.h b/Engine/Plugins/2D/Paper2D/Source/Paper2DEditor/Classes/PaperTileMapFactory.h
new file mode 100644
--- /dev/null
+++ b/Engine/Plugins/2D/Paper2D/Source/Paper2DEditor/Classes/PaperTileMapFactory.h
@@ -0,0 +1,22 @@
+// Copyright 1998-2015 Epic Games, Inc. All Rights Reserved.
+
+#pragma once
+
+/**
+ * Factory for tile maps
+ */
+
+#include "PaperTileMapFactory.generated.h"
+
+UCLASS()
+class UPaperTileMapFactory : public UFactory
+{
+	GENERATED_UCLASS_BODY()
+
+	// Initial tile set to create the tile map from (Can be nullptr)
+	class UPaperTileSet* InitialTileSet;
+
+	// UFactory interface
+	virtual UObject* FactoryCreateNew(UClass* Class, UObject* InParent, FName Name, EObjectFlags Flags, UObject* Context, FFeedbackContext* Warn) override;
+	// End of UFactory interface
+};
diff --git a/Engine/Plugins/2D/Paper2D/Source/Paper2DEditor/Private/PaperTileMapFactory.cpp b/Engine/Plugins/2D/Paper2D/Source/Paper2DEditor/Private/PaperTileMapFactory.cpp
new file mode 100644
--- /dev/null
+++ b/Engine/Plugins/2D/Paper2D/Source/Paper2DEditor/Private/PaperTileMapFactory.cpp
@@ -0,0 +1,33 @@
+// Copyright 1998-2015 Epic Games, Inc. All Rights Reserved.
+
+#include "Paper2DEditorPrivatePCH.h"
+
+#define LOCTEXT_NAMESPACE "Paper2D"
+
+/////////////////////////////////////////////////////
+// UPaperTileMapFactory
+
+UPaperTileMapFactory::UPaperTileMapFactory(const FObjectInitializer& ObjectInitializer)
+	: Super(ObjectInitializer)
+{
+	bCreateNew = true;
+	bEditAfterNew = true;
+	SupportedClass = UPaperTileMap::StaticClass();
+}
+
+UObject* UPaperTileMapFactory::FactoryCreateNew(UClass* Class, UObject* InParent, FName Name, EObjectFlags Flags, UObject* Context, FFeedbackContext* Warn)
+{
+	UPaperTileMap* NewTileMap = ConstructObject<UPaperTileMap>(Class, InParent, Name, Flags | RF_Transactional);
+
+	if (InitialTileSet != nullptr)
+	{
+		NewTileMap->TileWidth = InitialTileSet->TileWidth;
+		NewTileMap->TileHeight = InitialTileSet->TileHeight;
+	}
+
+	NewTileMap->AddNewLayer();
+
+	return NewTileMap;
+}
+
+#undef LOCTEXT_NAMESPACE
\ No newline at end of file
diff --git a/Engine/Plugins/2D/Paper2D/Source/Paper2DEditor/Private/TileMapEditing/TileMapActorFactory.cpp b/Engine/Plugins/2D/Paper2D/Source/Paper2DEditor/Private/TileMapEditing/TileMapActorFactory.cpp
--- a/Engine/Plugins/2D/Paper2D/Source/Paper2DEditor/Private/TileMapEditing/TileMapActorFactory.cpp
+++ b/Engine/Plugins/2D/Paper2D/Source/Paper2DEditor/Private/TileMapEditing/TileMapActorFactory.cpp
@@ -39,6 +39,11 @@ void UTileMapActorFactory::PostSpawnActor(UObject* Asset, AActor* NewActor)
 			RenderComponent->RegisterComponent();
 		}
 	}
+
+	if (RenderComponent->OwnsTileMap())
+	{
+		RenderComponent->TileMap->AddNewLayer();
+	}
 }
 
 void UTileMapActorFactory::PostCreateBlueprint(UObject* Asset, AActor* CDO)
@@ -60,6 +65,11 @@ void UTileMapActorFactory::PostCreateBlueprint(UObject* Asset, AActor* CDO)
 				RenderComponent->TileMap->TileHeight = TileSet->TileHeight;
 			}
 		}
+
+		if (RenderComponent->OwnsTileMap())
+		{
+			RenderComponent->TileMap->AddNewLayer();
+		}
 	}
 }
 
