Paper2D: Fix grid rendering/etc... not showing up in the stand alone tile map editor (due to the component not being 'selected')

diff --git a/Engine/Plugins/2D/Paper2D/Source/Paper2DEditor/Private/TileMapEditing/TileMapEditorViewportClient.cpp b/Engine/Plugins/2D/Paper2D/Source/Paper2DEditor/Private/TileMapEditing/TileMapEditorViewportClient.cpp
--- a/Engine/Plugins/2D/Paper2D/Source/Paper2DEditor/Private/TileMapEditing/TileMapEditorViewportClient.cpp
+++ b/Engine/Plugins/2D/Paper2D/Source/Paper2DEditor/Private/TileMapEditing/TileMapEditorViewportClient.cpp
@@ -45,6 +45,7 @@ FTileMapEditorViewportClient::FTileMapEditorViewportClient(TWeakPtr<FTileMapEdit
 		RenderTileMapComponent = NewObject<UPaperTileMapRenderComponent>();
 		UPaperTileMap* TileMap = GetTileMapBeingEdited();
 		RenderTileMapComponent->TileMap = TileMap;
+		GSelectedAnnotation.Set(RenderTileMapComponent);
 
 		PreviewScene->AddComponent(RenderTileMapComponent, FTransform::Identity);
 	}
