Paper2D: Changing selection to no longer contain a tile map component now causes the tile map editor mode to exit

diff --git a/Engine/Plugins/2D/Paper2D/Source/Paper2DEditor/Private/TileMapEditing/EdModeTileMap.cpp b/Engine/Plugins/2D/Paper2D/Source/Paper2DEditor/Private/TileMapEditing/EdModeTileMap.cpp
--- a/Engine/Plugins/2D/Paper2D/Source/Paper2DEditor/Private/TileMapEditing/EdModeTileMap.cpp
+++ b/Engine/Plugins/2D/Paper2D/Source/Paper2DEditor/Private/TileMapEditing/EdModeTileMap.cpp
@@ -117,13 +117,24 @@ void FEdModeTileMap::Enter()
 
 void FEdModeTileMap::Exit()
 {
-	FToolkitManager::Get().CloseToolkit(Toolkit.ToSharedRef());
-	Toolkit.Reset();
+	if (Toolkit.IsValid())
+	{
+		FToolkitManager::Get().CloseToolkit(Toolkit.ToSharedRef());
+		Toolkit.Reset();
+	}
 
 	// Call base Exit method to ensure proper cleanup
 	FEdMode::Exit();
 }
 
+void FEdModeTileMap::ActorSelectionChangeNotify()
+{
+	if (FindSelectedComponent() == nullptr)
+	{
+		Owner->DeactivateMode(FEdModeTileMap::EM_TileMap);
+	}
+}
+
 bool FEdModeTileMap::MouseMove(FEditorViewportClient* InViewportClient, FViewport* InViewport, int32 x, int32 y)
 {
 	if (InViewportClient->EngineShowFlags.ModeWidgets)
@@ -301,13 +312,8 @@ bool FEdModeTileMap::UsesTransformWidget() const
 	return false;
 }
 
-UPaperTileLayer* FEdModeTileMap::GetSelectedLayerUnderCursor(const FViewportCursorLocation& Ray, int32& OutTileX, int32& OutTileY) const
+UPaperTileMapRenderComponent* FEdModeTileMap::FindSelectedComponent() const
 {
-	const FVector TraceStart = Ray.GetOrigin();
-	const FVector TraceDir = Ray.GetDirection();
-
-	const bool bCollisionPainting = (GetActiveLayerPaintingMode() == ETileMapLayerPaintingMode::CollisionLayers);
-
 	UPaperTileMapRenderComponent* TileMapComponent = nullptr;
 
 	USelection* SelectedActors = Owner->GetSelectedActors();
@@ -337,6 +343,18 @@ UPaperTileLayer* FEdModeTileMap::GetSelectedLayerUnderCursor(const FViewportCurs
 		}
 	}
 
+	return TileMapComponent;
+}
+
+UPaperTileLayer* FEdModeTileMap::GetSelectedLayerUnderCursor(const FViewportCursorLocation& Ray, int32& OutTileX, int32& OutTileY) const
+{
+	const FVector TraceStart = Ray.GetOrigin();
+	const FVector TraceDir = Ray.GetDirection();
+
+	const bool bCollisionPainting = (GetActiveLayerPaintingMode() == ETileMapLayerPaintingMode::CollisionLayers);
+
+	UPaperTileMapRenderComponent* TileMapComponent = FindSelectedComponent();
+
 	if (TileMapComponent != nullptr)
 	{
 		if (UPaperTileMap* TileMap = TileMapComponent->TileMap)
@@ -752,7 +770,7 @@ FViewportCursorLocation FEdModeTileMap::CalculateViewRay(FEditorViewportClient*
 		.SetRealtimeUpdate( InViewportClient->IsRealtime() ));
 
 	FSceneView* View = InViewportClient->CalcSceneView( &ViewFamily );
-	FViewportCursorLocation MouseViewportRay( View, (FEditorViewportClient*)InViewport->GetClient(), InViewport->GetMouseX(), InViewport->GetMouseY() );
+	FViewportCursorLocation MouseViewportRay( View, InViewportClient, InViewport->GetMouseX(), InViewport->GetMouseY() );
 
 	return MouseViewportRay;
 }
diff --git a/Engine/Plugins/2D/Paper2D/Source/Paper2DEditor/Private/TileMapEditing/EdModeTileMap.h b/Engine/Plugins/2D/Paper2D/Source/Paper2DEditor/Private/TileMapEditing/EdModeTileMap.h
--- a/Engine/Plugins/2D/Paper2D/Source/Paper2DEditor/Private/TileMapEditing/EdModeTileMap.h
+++ b/Engine/Plugins/2D/Paper2D/Source/Paper2DEditor/Private/TileMapEditing/EdModeTileMap.h
@@ -54,10 +54,7 @@ public:
 	virtual bool InputDelta(FEditorViewportClient* InViewportClient, FViewport* InViewport, FVector& InDrag, FRotator& InRot, FVector& InScale) override;
 	virtual void Render(const FSceneView* View, FViewport* Viewport, FPrimitiveDrawInterface* PDI) override;
 	virtual void DrawHUD(FEditorViewportClient* ViewportClient, FViewport* Viewport, const FSceneView* View, FCanvas* Canvas) override;
-	// 	virtual bool Select(AActor* InActor, bool bInSelected) override;
-	// 	virtual bool IsSelectionAllowed(AActor* InActor) const override;
-	// 	virtual void ActorSelectionChangeNotify() override;
-	// 	virtual FVector GetWidgetLocation() const override;
+	virtual void ActorSelectionChangeNotify() override;
 	virtual bool AllowWidgetMove();
 	virtual bool ShouldDrawWidget() const override;
 	virtual bool UsesTransformWidget() const override;
@@ -93,6 +90,8 @@ protected:
 	
 	TSharedRef<FExtender> AddCreationModeExtender(const TSharedRef<FUICommandList> InCommandList);
 
+	UPaperTileMapRenderComponent* FindSelectedComponent() const;
+
 	void EnableTileMapEditMode();
 	bool IsTileMapEditModeActive() const;
 protected:
