Paper2D: Fix incorrect math in tile map editing intersection code that cause selection wonkiness

diff --git a/Engine/Plugins/2D/Paper2D/Source/Paper2DEditor/Private/TileMapEditing/EdModeTileMap.cpp b/Engine/Plugins/2D/Paper2D/Source/Paper2DEditor/Private/TileMapEditing/EdModeTileMap.cpp
--- a/Engine/Plugins/2D/Paper2D/Source/Paper2DEditor/Private/TileMapEditing/EdModeTileMap.cpp
+++ b/Engine/Plugins/2D/Paper2D/Source/Paper2DEditor/Private/TileMapEditing/EdModeTileMap.cpp
@@ -371,19 +371,19 @@ UPaperTileLayer* FEdModeTileMap::GetSelectedLayerUnderCursor(const FViewportCurs
 				const float WY = TileMap->MapHeight * TileMap->TileHeight;
 
 				ComponentToWorld = (TileMapComponent != nullptr) ? TileMapComponent->ComponentToWorld : FTransform::Identity;
-				FVector LocalStart = ComponentToWorld.InverseTransformPosition(TraceStart);
-				FVector LocalDirection = ComponentToWorld.InverseTransformVector(TraceDir);
-
+				const FVector LocalStart = ComponentToWorld.InverseTransformPosition(TraceStart);
+				const FVector LocalDirection = ComponentToWorld.InverseTransformVector(TraceDir);
+				const FVector LocalEnd = LocalStart + (LocalDirection * HALF_WORLD_MAX);
+					 
 				const FVector LSPlaneCorner = PaperAxisZ * TileMap->SeparationPerLayer;
 
 				const FPlane LayerPlane(LSPlaneCorner + PaperAxisX, LSPlaneCorner, LSPlaneCorner + PaperAxisY);
 
 				FVector Intersection;
-				if (FMath::SegmentPlaneIntersection(LocalStart, LocalDirection * HALF_WORLD_MAX, LayerPlane, /*out*/ Intersection))
+				if (FMath::SegmentPlaneIntersection(LocalStart, LocalEnd, LayerPlane, /*out*/ Intersection))
 				{
-					//@TODO: Ideally tile pivots weren't in the center!
-					const float NormalizedX = (Intersection.X + 0.5f * TileMap->TileWidth) / WX;
-					const float NormalizedY = (-Intersection.Z + 0.5f * TileMap->TileHeight) / WY;
+					const float NormalizedX = (FVector::DotProduct(Intersection, PaperAxisX) + 0.5f * TileMap->TileWidth) / WX;
+					const float NormalizedY = (-FVector::DotProduct(Intersection, PaperAxisY) + 0.5f * TileMap->TileHeight) / WY;
 
 					OutTileX = FMath::FloorToInt(NormalizedX * TileMap->MapWidth);
 					OutTileY = FMath::FloorToInt(NormalizedY * TileMap->MapHeight);
