Paper2D: Prevent pivot cross from being drawn in tile map edit mode

diff --git a/Engine/Plugins/2D/Paper2D/Source/Paper2DEditor/Private/TileMapEditing/EdModeTileMap.cpp b/Engine/Plugins/2D/Paper2D/Source/Paper2DEditor/Private/TileMapEditing/EdModeTileMap.cpp
--- a/Engine/Plugins/2D/Paper2D/Source/Paper2DEditor/Private/TileMapEditing/EdModeTileMap.cpp
+++ b/Engine/Plugins/2D/Paper2D/Source/Paper2DEditor/Private/TileMapEditing/EdModeTileMap.cpp
@@ -86,6 +86,8 @@ FEdModeTileMap::FEdModeTileMap()
 	, DrawPreviewDimensionsLS(0.0f, 0.0f, 0.0f)
 	, EraseBrushSize(1)
 {
+	bDrawPivot = false;
+
 	SetActiveTool(ETileMapEditorTool::Paintbrush);
 	SetActiveLayerPaintingMode(ETileMapLayerPaintingMode::VisualLayers);
 }
