Paper2D: Improve display message when no ink is selected in tile map edit mode

diff --git a/Engine/Plugins/2D/Paper2D/Source/Paper2DEditor/Private/TileMapEditing/EdModeTileMap.cpp b/Engine/Plugins/2D/Paper2D/Source/Paper2DEditor/Private/TileMapEditing/EdModeTileMap.cpp
--- a/Engine/Plugins/2D/Paper2D/Source/Paper2DEditor/Private/TileMapEditing/EdModeTileMap.cpp
+++ b/Engine/Plugins/2D/Paper2D/Source/Paper2DEditor/Private/TileMapEditing/EdModeTileMap.cpp
@@ -293,7 +293,7 @@ void FEdModeTileMap::Render(const FSceneView* View, FViewport* Viewport, FPrimit
 void FEdModeTileMap::DrawHUD(FEditorViewportClient* ViewportClient, FViewport* Viewport, const FSceneView* View, FCanvas* Canvas)
 {
 	FString InkInfo = FString::Printf(TEXT("Ink: (%d, %d)  %dx%d  %s"), PaintSourceTopLeft.X, PaintSourceTopLeft.Y, PaintSourceDimensions.X, PaintSourceDimensions.Y, 
-		(PaintSourceTileSet.Get() != NULL) ? (*PaintSourceTileSet.Get()->GetName()) : TEXT("(null)"));
+		(PaintSourceTileSet.Get() != NULL) ? (*PaintSourceTileSet.Get()->GetName()) : TEXT("(no tile ink selected)"));
 
 	FCanvasTextItem Msg(FVector2D(10, 30), FText::FromString(InkInfo), GEngine->GetMediumFont(), FLinearColor::White);
 	Canvas->DrawItem(Msg);
