Paper2D: Change tile map editor to keep track of the last selected tile set for a tile map

diff --git a/Engine/Plugins/2D/Paper2D/Source/Paper2D/Classes/PaperTileMap.h b/Engine/Plugins/2D/Paper2D/Source/Paper2D/Classes/PaperTileMap.h
--- a/Engine/Plugins/2D/Paper2D/Source/Paper2D/Classes/PaperTileMap.h
+++ b/Engine/Plugins/2D/Paper2D/Source/Paper2D/Classes/PaperTileMap.h
@@ -60,9 +60,9 @@ class PAPER2D_API UPaperTileMap : public UObject
 	UPROPERTY(Category=Setup, EditAnywhere)
 	float SeparationPerLayer;
 
-	// Default tile set to use for new layers
-	UPROPERTY(Category=Setup, EditAnywhere)
-	class UPaperTileSet* DefaultLayerTileSet;
+	// Last tile set that was selected when editing the tile map
+	UPROPERTY()
+	TAssetPtr<class UPaperTileSet> SelectedTileSet;
 
 	// Test material
 	UPROPERTY(Category=Setup, EditAnywhere)
diff --git a/Engine/Plugins/2D/Paper2D/Source/Paper2D/Private/PaperTileMapRenderComponent.cpp b/Engine/Plugins/2D/Paper2D/Source/Paper2D/Private/PaperTileMapRenderComponent.cpp
--- a/Engine/Plugins/2D/Paper2D/Source/Paper2D/Private/PaperTileMapRenderComponent.cpp
+++ b/Engine/Plugins/2D/Paper2D/Source/Paper2D/Private/PaperTileMapRenderComponent.cpp
@@ -97,7 +97,7 @@ void UPaperTileMapRenderComponent::PostLoad()
 		TileMap->TileWidth = TileWidth_DEPRECATED;
 		TileMap->TileHeight = TileHeight_DEPRECATED;
 		TileMap->PixelsPerUnit = 1.0f;
-		TileMap->DefaultLayerTileSet = DefaultLayerTileSet_DEPRECATED;
+		TileMap->SelectedTileSet = DefaultLayerTileSet_DEPRECATED;
 		TileMap->Material = Material_DEPRECATED;
 		TileMap->TileLayers = TileLayers_DEPRECATED;
 
diff --git a/Engine/Plugins/2D/Paper2D/Source/Paper2DEditor/Private/TileMapEditing/EdModeTileMap.h b/Engine/Plugins/2D/Paper2D/Source/Paper2DEditor/Private/TileMapEditing/EdModeTileMap.h
--- a/Engine/Plugins/2D/Paper2D/Source/Paper2DEditor/Private/TileMapEditing/EdModeTileMap.h
+++ b/Engine/Plugins/2D/Paper2D/Source/Paper2DEditor/Private/TileMapEditing/EdModeTileMap.h
@@ -90,10 +90,12 @@ protected:
 	
 	TSharedRef<FExtender> AddCreationModeExtender(const TSharedRef<FUICommandList> InCommandList);
 
-	UPaperTileMapRenderComponent* FindSelectedComponent() const;
-
 	void EnableTileMapEditMode();
 	bool IsTileMapEditModeActive() const;
+
+public:
+	UPaperTileMapRenderComponent* FindSelectedComponent() const;
+
 protected:
 	bool bIsPainting;
 	TWeakObjectPtr<UPaperTileSet> PaintSourceTileSet;
diff --git a/Engine/Plugins/2D/Paper2D/Source/Paper2DEditor/Private/TileMapEditing/TileMapEdModeToolkit.cpp b/Engine/Plugins/2D/Paper2D/Source/Paper2DEditor/Private/TileMapEditing/TileMapEdModeToolkit.cpp
--- a/Engine/Plugins/2D/Paper2D/Source/Paper2DEditor/Private/TileMapEditing/TileMapEdModeToolkit.cpp
+++ b/Engine/Plugins/2D/Paper2D/Source/Paper2DEditor/Private/TileMapEditing/TileMapEdModeToolkit.cpp
@@ -67,17 +67,11 @@ void FTileMapEdModeToolkit::Init(const TSharedPtr<IToolkitHost>& InitToolkitHost
 	BindCommands();
 
 	// Try to determine a good default tile set based on the current selection set
-	USelection* SelectedActors = GEditor->GetSelectedActors();
-	for (FSelectionIterator Iter(*SelectedActors); Iter; ++Iter)
+	if (UPaperTileMapRenderComponent* TargetComponent = TileMapEditor->FindSelectedComponent())
 	{
-		AActor* Actor = CastChecked<AActor>(*Iter);
-		if (UPaperTileMapRenderComponent* TileMapComponent = Actor->FindComponentByClass<UPaperTileMapRenderComponent>())
+		if (TargetComponent->TileMap != nullptr)
 		{
-			if (TileMapComponent->TileMap != nullptr)
-			{
-				CurrentTileSetPtr = TileMapComponent->TileMap->DefaultLayerTileSet;
-				break;
-			}
+			CurrentTileSetPtr = TargetComponent->TileMap->SelectedTileSet.Get();
 		}
 	}
 
@@ -154,10 +148,22 @@ void FTileMapEdModeToolkit::OnChangeTileSet(UObject* NewAsset)
 {
 	if (UPaperTileSet* NewTileSet = Cast<UPaperTileSet>(NewAsset))
 	{
-		CurrentTileSetPtr = NewTileSet;
-		if (TileSetPalette.IsValid())
+		if (CurrentTileSetPtr.Get() != NewTileSet)
 		{
-			TileSetPalette->ChangeTileSet(NewTileSet);
+			CurrentTileSetPtr = NewTileSet;
+			if (TileSetPalette.IsValid())
+			{
+				TileSetPalette->ChangeTileSet(NewTileSet);
+			}
+
+			// Save the newly selected tile set in the asset so it can be restored next time we edit it
+			if (UPaperTileMapRenderComponent* TargetComponent = TileMapEditor->FindSelectedComponent())
+			{
+				if (TargetComponent->TileMap != nullptr)
+				{
+					TargetComponent->TileMap->SelectedTileSet = NewTileSet;
+				}
+			}
 		}
 	}
 }
