Paper2D: Prevent the tile map editor mode from showing up as a tab, it is only activated by pressing 'edit tile map' in the details panel now

diff --git a/Engine/Plugins/2D/Paper2D/Source/Paper2DEditor/Private/Paper2DEditorModule.cpp b/Engine/Plugins/2D/Paper2D/Source/Paper2DEditor/Private/Paper2DEditorModule.cpp
--- a/Engine/Plugins/2D/Paper2D/Source/Paper2DEditor/Private/Paper2DEditorModule.cpp
+++ b/Engine/Plugins/2D/Paper2D/Source/Paper2DEditor/Private/Paper2DEditorModule.cpp
@@ -246,7 +246,7 @@ private:
 				FEdModeTileMap::EM_TileMap,
 				LOCTEXT("TileMapEditMode", "Tile Map Editor"),
 				FSlateIcon(),
-				true);
+				false);
 		}
 	}
 };
