Paper2D: Pass in notify hook from details customization to tile layer list, allowing some limited functionality in the BP editor

diff --git a/Engine/Plugins/2D/Paper2D/Source/Paper2DEditor/Private/TileMapEditing/PaperTileMapDetailsCustomization.cpp b/Engine/Plugins/2D/Paper2D/Source/Paper2DEditor/Private/TileMapEditing/PaperTileMapDetailsCustomization.cpp
--- a/Engine/Plugins/2D/Paper2D/Source/Paper2DEditor/Private/TileMapEditing/PaperTileMapDetailsCustomization.cpp
+++ b/Engine/Plugins/2D/Paper2D/Source/Paper2DEditor/Private/TileMapEditing/PaperTileMapDetailsCustomization.cpp
@@ -9,6 +9,7 @@
 
 #include "STileLayerList.h"
 #include "ScopedTransaction.h"
+#include "IPropertyUtilities.h"
 
 #define LOCTEXT_NAMESPACE "Paper2D"
 
@@ -24,6 +25,8 @@ void FPaperTileMapDetailsCustomization::CustomizeDetails(IDetailLayoutBuilder& D
 {
 	const TArray< TWeakObjectPtr<UObject> >& SelectedObjects = DetailLayout.GetDetailsView().GetSelectedObjects();
 	
+	FNotifyHook* NotifyHook = DetailLayout.GetPropertyUtilities()->GetNotifyHook();
+
 	bool bEditingActor = false;
 
 	UPaperTileMap* TileMap = NULL;
@@ -80,7 +83,7 @@ void FPaperTileMapDetailsCustomization::CustomizeDetails(IDetailLayoutBuilder& D
 
 		LayersCategory.AddCustomRow(LOCTEXT("TileLayerList", "Tile layer list"))
 		[
-			SNew(STileLayerList, TileMap)
+			SNew(STileLayerList, TileMap, NotifyHook)
 		];
 	}
 
diff --git a/Engine/Plugins/2D/Paper2D/Source/Paper2DEditor/Private/TileMapEditing/STileLayerList.cpp b/Engine/Plugins/2D/Paper2D/Source/Paper2DEditor/Private/TileMapEditing/STileLayerList.cpp
--- a/Engine/Plugins/2D/Paper2D/Source/Paper2DEditor/Private/TileMapEditing/STileLayerList.cpp
+++ b/Engine/Plugins/2D/Paper2D/Source/Paper2DEditor/Private/TileMapEditing/STileLayerList.cpp
@@ -14,9 +14,10 @@
 //////////////////////////////////////////////////////////////////////////
 // STileLayerList
 
-void STileLayerList::Construct(const FArguments& InArgs, UPaperTileMap* TileMap)
+void STileLayerList::Construct(const FArguments& InArgs, UPaperTileMap* InTileMap, FNotifyHook* InNotifyHook)
 {
-	TileMapPtr = TileMap;
+	TileMapPtr = InTileMap;
+	NotifyHook = InNotifyHook;
 
 	FTileMapEditorCommands::Register();
 	const FTileMapEditorCommands& Commands = FTileMapEditorCommands::Get();
@@ -70,16 +71,16 @@ void STileLayerList::Construct(const FArguments& InArgs, UPaperTileMap* TileMap)
 	ListViewWidget = SNew(SPaperLayerListView)
 		.SelectionMode(ESelectionMode::Single)
 		.ClearSelectionOnClick(false)
-		.ListItemsSource(&(TileMap->TileLayers))
+		.ListItemsSource(&(InTileMap->TileLayers))
 		.OnSelectionChanged(this, &STileLayerList::OnSelectionChanged)
 		.OnGenerateRow(this, &STileLayerList::OnGenerateLayerListRow)
 		.OnContextMenuOpening(this, &STileLayerList::OnConstructContextMenu);
 
 	// Restore the selection
-	TileMap->ValidateSelectedLayerIndex();
-	if (TileMap->SelectedLayerIndex != INDEX_NONE)
+	InTileMap->ValidateSelectedLayerIndex();
+	if (InTileMap->SelectedLayerIndex != INDEX_NONE)
 	{
-		SetSelectedLayer(TileMap->TileLayers[TileMap->SelectedLayerIndex]);
+		SetSelectedLayer(InTileMap->TileLayers[InTileMap->SelectedLayerIndex]);
 	}
 
 	ChildSlot
@@ -185,8 +186,7 @@ UPaperTileLayer* STileLayerList::AddLayer(bool bCollisionLayer, int32 InsertionI
 
 		NewLayer = TileMap->AddNewLayer(bCollisionLayer, InsertionIndex);
 
-		ListViewWidget->RequestListRefresh();
-		TileMap->PostEditChange();
+		PostEditNotfications();
 
 		// Change the selection set to select it
 		SetSelectedLayer(NewLayer);
@@ -209,8 +209,7 @@ void STileLayerList::ChangeLayerOrdering(int32 OldIndex, int32 NewIndex)
 			TileMap->TileLayers.RemoveAt(OldIndex);
 			TileMap->TileLayers.Insert(LayerToMove, NewIndex);
 
-			ListViewWidget->RequestListRefresh();
-			TileMap->PostEditChange();
+			PostEditNotfications();
 		}
 	}
 }
@@ -257,8 +256,7 @@ void STileLayerList::DeleteLayer()
 
 			TileMap->TileLayers.RemoveAt(DeleteIndex);
 
-			TileMap->PostEditChange();
-			ListViewWidget->RequestListRefresh();
+			PostEditNotfications();
 
 			// Select the item below the one that just got deleted
 			const int32 NewSelectionIndex = FMath::Min<int32>(DeleteIndex, TileMap->TileLayers.Num() - 1);
@@ -285,8 +283,7 @@ void STileLayerList::DuplicateLayer()
 			TileMap->TileLayers.Insert(NewLayer, DuplicateIndex);
 			NewLayer->LayerName = GenerateDuplicatedLayerName(NewLayer->LayerName.ToString(), TileMap);
 
-			TileMap->PostEditChange();
-			ListViewWidget->RequestListRefresh();
+			PostEditNotfications();
 
 			// Select the duplicated layer
 			SetSelectedLayer(NewLayer);
@@ -329,8 +326,7 @@ void STileLayerList::MergeLayerDown()
 			TileMap->TileLayers.RemoveAt(SourceIndex);
 
 			// Update viewers
-			TileMap->PostEditChange();
-			ListViewWidget->RequestListRefresh();
+			PostEditNotfications();
 		}
 	}
 }
@@ -384,6 +380,8 @@ void STileLayerList::OnSelectionChanged(UPaperTileLayer* ItemChangingState, ESel
 	if (UPaperTileMap* TileMap = TileMapPtr.Get())
 	{
 		TileMap->SelectedLayerIndex = GetSelectionIndex();
+
+		PostEditNotfications();
 	}
 }
 
@@ -405,6 +403,23 @@ TSharedPtr<SWidget> STileLayerList::OnConstructContextMenu()
 	return MenuBuilder.MakeWidget();
 }
 
+void STileLayerList::PostEditNotfications()
+{
+	ListViewWidget->RequestListRefresh();
+
+	if (UPaperTileMap* TileMap = TileMapPtr.Get())
+	{
+		TileMap->PostEditChange();
+	}
+
+	if (NotifyHook != nullptr)
+	{
+		UProperty* TileMapProperty = FindFieldChecked<UProperty>(UPaperTileMapRenderComponent::StaticClass(), GET_MEMBER_NAME_CHECKED(UPaperTileMapRenderComponent, TileMap));
+		NotifyHook->NotifyPreChange(TileMapProperty);
+		NotifyHook->NotifyPostChange(FPropertyChangedEvent(TileMapProperty), TileMapProperty);
+	}
+}
+
 //////////////////////////////////////////////////////////////////////////
 
 #undef LOCTEXT_NAMESPACE
\ No newline at end of file
diff --git a/Engine/Plugins/2D/Paper2D/Source/Paper2DEditor/Private/TileMapEditing/STileLayerList.h b/Engine/Plugins/2D/Paper2D/Source/Paper2DEditor/Private/TileMapEditing/STileLayerList.h
--- a/Engine/Plugins/2D/Paper2D/Source/Paper2DEditor/Private/TileMapEditing/STileLayerList.h
+++ b/Engine/Plugins/2D/Paper2D/Source/Paper2DEditor/Private/TileMapEditing/STileLayerList.h
@@ -11,7 +11,7 @@ public:
 	SLATE_BEGIN_ARGS(STileLayerList) {}
 	SLATE_END_ARGS()
 
-	void Construct(const FArguments& InArgs, UPaperTileMap* TileMap);
+	void Construct(const FArguments& InArgs, UPaperTileMap* TileMap, FNotifyHook* InNotifyHook);
 
 protected:
 	typedef SListView<class UPaperTileLayer*> SPaperLayerListView;
@@ -20,7 +20,7 @@ protected:
 	TSharedPtr<SPaperLayerListView> ListViewWidget;
 	TSharedPtr<class FUICommandList> CommandList;
 	TWeakObjectPtr<class UPaperTileMap> TileMapPtr;
-
+	FNotifyHook* NotifyHook;
 protected:
 	TSharedRef<ITableRow> OnGenerateLayerListRow(class UPaperTileLayer* Item, const TSharedRef<STableViewBase>& OwnerTable);
 
@@ -53,4 +53,7 @@ protected:
 
 	void OnSelectionChanged(UPaperTileLayer* ItemChangingState, ESelectInfo::Type SelectInfo);
 	TSharedPtr<SWidget> OnConstructContextMenu();
+
+	// Called after edits are finished
+	void PostEditNotfications();
 };
