Paper2D: Fix move layer up/down not properly updating the selection state on the tile map itself

diff --git a/Engine/Plugins/2D/Paper2D/Source/Paper2DEditor/Private/TileMapEditing/STileLayerList.cpp b/Engine/Plugins/2D/Paper2D/Source/Paper2DEditor/Private/TileMapEditing/STileLayerList.cpp
--- a/Engine/Plugins/2D/Paper2D/Source/Paper2DEditor/Private/TileMapEditing/STileLayerList.cpp
+++ b/Engine/Plugins/2D/Paper2D/Source/Paper2DEditor/Private/TileMapEditing/STileLayerList.cpp
@@ -209,6 +209,11 @@ void STileLayerList::ChangeLayerOrdering(int32 OldIndex, int32 NewIndex)
 			TileMap->TileLayers.RemoveAt(OldIndex);
 			TileMap->TileLayers.Insert(LayerToMove, NewIndex);
 
+			if (TileMap->SelectedLayerIndex == OldIndex)
+			{
+				TileMap->SelectedLayerIndex = NewIndex;
+			}
+
 			PostEditNotfications();
 		}
 	}
