Paper2D: Add limited support for flipbook collision (sourced either from the first frame of the animation or each frame of the animation; defaults to no collision)

diff --git a/Engine/Plugins/2D/Paper2D/Source/Paper2D/Classes/PaperFlipbook.h b/Engine/Plugins/2D/Paper2D/Source/Paper2D/Classes/PaperFlipbook.h
--- a/Engine/Plugins/2D/Paper2D/Source/Paper2D/Classes/PaperFlipbook.h
+++ b/Engine/Plugins/2D/Paper2D/Source/Paper2D/Classes/PaperFlipbook.h
@@ -24,6 +24,22 @@ public:
 	}
 };
 
+UENUM()
+namespace EFlipbookCollisionMode
+{
+	enum Type
+	{
+		// The flipbook has no collision
+		NoCollision,
+
+		// The flipbook has non-animated collision based on the first frame of the animation
+		FirstFrameCollision,
+
+		// The flipbook changes collision each frame based on the animation (Note: This setting is not recommended and is very expensive, recreating the physics state every frame)
+		EachFrameCollision,
+	};
+}
+
 /**
  * Contains an animation sequence of sprite frames
  */
@@ -45,7 +61,14 @@ protected:
 	UPROPERTY(Category=Sprite, EditAnywhere, BlueprintReadOnly)
 	UMaterialInterface* DefaultMaterial;
 
+	// Collision source
+	UPROPERTY(Category=Sprite, EditAnywhere, BlueprintReadOnly)
+	TEnumAsByte<EFlipbookCollisionMode::Type> CollisionSource;
+
 public:
+	// Returns the collision source of this flipbook animation (if any)
+	TEnumAsByte<EFlipbookCollisionMode::Type> GetCollisionSource() const { return CollisionSource; }
+
 	// Returns the nominal frame rate to play this flipbook animation back at
 	float GetFramesPerSecond() const;
 
diff --git a/Engine/Plugins/2D/Paper2D/Source/Paper2D/Classes/PaperFlipbookComponent.h b/Engine/Plugins/2D/Paper2D/Source/Paper2D/Classes/PaperFlipbookComponent.h
--- a/Engine/Plugins/2D/Paper2D/Source/Paper2D/Classes/PaperFlipbookComponent.h
+++ b/Engine/Plugins/2D/Paper2D/Source/Paper2D/Classes/PaperFlipbookComponent.h
@@ -50,6 +50,10 @@ protected:
 	UPROPERTY(BlueprintReadOnly, Interp, Category=Sprite)
 	FLinearColor SpriteColor;
 
+	/** The cached body setup */
+	UPROPERTY(Transient)
+	UBodySetup* CachedBodySetup;
+
 public:
 	/** Event called whenever a non-looping flipbook finishes playing (either reaching the beginning or the end, depending on the play direction) */
 	UPROPERTY(BlueprintAssignable)
@@ -149,6 +153,7 @@ protected:
 	UPaperSprite* GetSpriteAtCachedIndex() const;
 
 	void TickFlipbook(float DeltaTime);
+	void FlipbookChangedPhysicsState();
 
 public:
 	// UObject interface
@@ -181,6 +186,7 @@ public:
 	virtual void GetUsedMaterials(TArray<UMaterialInterface*>& OutMaterials) const override;
 	virtual void GetStreamingTextureInfo(TArray<FStreamingTexturePrimitiveInfo>& OutStreamingTextures) const override;
 	virtual int32 GetNumMaterials() const override;
+	virtual UBodySetup* GetBodySetup() override;
 	// End of UPrimitiveComponent interface
 };
 
diff --git a/Engine/Plugins/2D/Paper2D/Source/Paper2D/Private/PaperFlipbookComponent.cpp b/Engine/Plugins/2D/Paper2D/Source/Paper2D/Private/PaperFlipbookComponent.cpp
--- a/Engine/Plugins/2D/Paper2D/Source/Paper2D/Private/PaperFlipbookComponent.cpp
+++ b/Engine/Plugins/2D/Paper2D/Source/Paper2D/Private/PaperFlipbookComponent.cpp
@@ -18,9 +18,7 @@
 UPaperFlipbookComponent::UPaperFlipbookComponent(const FObjectInitializer& ObjectInitializer)
 	: Super(ObjectInitializer)
 {
-	BodyInstance.SetCollisionEnabled(ECollisionEnabled::NoCollision);
-
-	SetCollisionProfileName(UCollisionProfile::NoCollision_ProfileName);
+	SetCollisionProfileName(UCollisionProfile::BlockAllDynamic_ProfileName);
 	Material_DEPRECATED = nullptr;
 
 	CastShadow = false;
@@ -99,18 +97,15 @@ FBoxSphereBounds UPaperFlipbookComponent::CalcBounds(const FTransform& LocalToWo
 		// Graphics bounds.
 		FBoxSphereBounds NewBounds = SourceFlipbook->GetRenderBounds().TransformBy(LocalToWorld);
 
-		//@TODO: PAPER2D: Add collision support to flipbooks
-#if 0
 		// Add bounds of collision geometry (if present).
-		if (UBodySetup* BodySetup = SourceFlipbook->BodySetup)
+		if (CachedBodySetup != nullptr)
 		{
-			const FBox AggGeomBox = BodySetup->AggGeom.CalcAABB(LocalToWorld);
+			const FBox AggGeomBox = CachedBodySetup->AggGeom.CalcAABB(LocalToWorld);
 			if (AggGeomBox.IsValid)
 			{
 				NewBounds = Union(NewBounds, FBoxSphereBounds(AggGeomBox));
 			}
 		}
-#endif
 
 		// Apply bounds scale
 		NewBounds.BoxExtent *= BoundsScale;
@@ -176,6 +171,24 @@ int32 UPaperFlipbookComponent::GetNumMaterials() const
 	return FMath::Max<int32>(OverrideMaterials.Num(), 1);
 }
 
+UBodySetup* UPaperFlipbookComponent::GetBodySetup()
+{
+	CachedBodySetup = nullptr;
+
+	if ((SourceFlipbook != nullptr) && (SourceFlipbook->GetCollisionSource() != EFlipbookCollisionMode::NoCollision))
+	{
+		const int32 PotentialSpriteIndex = (SourceFlipbook->GetCollisionSource() == EFlipbookCollisionMode::FirstFrameCollision) ? 0 : CachedFrameIndex;
+		UPaperSprite* PotentialSpriteSource = SourceFlipbook->GetSpriteAtFrame(PotentialSpriteIndex);
+
+		if (PotentialSpriteSource != nullptr)
+		{
+			CachedBodySetup = PotentialSpriteSource->BodySetup;
+		}
+	}
+
+	return CachedBodySetup;
+}
+
 void UPaperFlipbookComponent::CalculateCurrentFrame()
 {
 	const int32 LastCachedFrame = CachedFrameIndex;
@@ -186,11 +199,25 @@ void UPaperFlipbookComponent::CalculateCurrentFrame()
 		// Update children transforms in case we have anything attached to an animated socket
 		UpdateChildTransforms();
 
+		if ((SourceFlipbook != nullptr) && (SourceFlipbook->GetCollisionSource() == EFlipbookCollisionMode::EachFrameCollision))
+		{
+			FlipbookChangedPhysicsState();
+		}
+
 		// Indicate we need to send new dynamic data.
 		MarkRenderDynamicDataDirty();
 	}
 }
 
+void UPaperFlipbookComponent::FlipbookChangedPhysicsState()
+{
+	// If the frame has changed and we're using animated collision, we need to recreate that state as well
+	RecreatePhysicsState();
+
+	// We just totally changed the physics setup so update overlaps too
+	UpdateOverlaps();
+}
+
 void UPaperFlipbookComponent::TickFlipbook(float DeltaTime)
 {
 	bool bIsFinished = false;
@@ -313,12 +340,14 @@ void UPaperFlipbookComponent::SendRenderDynamicData_Concurrent()
 		DrawCall.BuildFromSprite(SpriteToSend);
 		DrawCall.Color = SpriteColor;
 
-		ENQUEUE_UNIQUE_RENDER_COMMAND_TWOPARAMETER(
+		ENQUEUE_UNIQUE_RENDER_COMMAND_THREEPARAMETER(
 				FSendPaperRenderComponentDynamicData,
 				FPaperRenderSceneProxy*,InSceneProxy,(FPaperRenderSceneProxy*)SceneProxy,
 				FSpriteDrawCallRecord,InSpriteToSend,DrawCall,
+				UBodySetup*,InBodySetup,CachedBodySetup,
 			{
 				InSceneProxy->SetDrawCall_RenderThread(InSpriteToSend);
+				InSceneProxy->SetBodySetup_RenderThread(InBodySetup);
 			});
 	}
 }
@@ -341,7 +370,7 @@ bool UPaperFlipbookComponent::SetFlipbook(class UPaperFlipbook* NewFlipbook)
 			MarkRenderStateDirty();
 
 			// Update physics representation right away
-			RecreatePhysicsState();
+			FlipbookChangedPhysicsState();
 
 			// Notify the streaming system. Don't use Update(), because this may be the first time the mesh has been set
 			// and the component may have to be added to the streaming system for the first time.
@@ -384,7 +413,6 @@ const UObject* UPaperFlipbookComponent::AdditionalStatObject() const
 	return SourceFlipbook;
 }
 
-
 void UPaperFlipbookComponent::Play()
 {
 	Activate();
diff --git a/Engine/Plugins/2D/Paper2D/Source/Paper2D/Private/PaperRenderSceneProxy.cpp b/Engine/Plugins/2D/Paper2D/Source/Paper2D/Private/PaperRenderSceneProxy.cpp
--- a/Engine/Plugins/2D/Paper2D/Source/Paper2D/Private/PaperRenderSceneProxy.cpp
+++ b/Engine/Plugins/2D/Paper2D/Source/Paper2D/Private/PaperRenderSceneProxy.cpp
@@ -545,6 +545,11 @@ void FPaperRenderSceneProxy::SetDrawCall_RenderThread(const FSpriteDrawCallRecor
 	Record = NewDynamicData;
 }
 
+void FPaperRenderSceneProxy::SetBodySetup_RenderThread(UBodySetup* NewSetup)
+{
+	BodySetup = NewSetup;
+}
+
 bool FPaperRenderSceneProxy::IsCollisionView(const FEngineShowFlags& EngineShowFlags, bool& bDrawSimpleCollision, bool& bDrawComplexCollision) const
 {
 	bDrawSimpleCollision = false;
diff --git a/Engine/Plugins/2D/Paper2D/Source/Paper2D/Private/PaperRenderSceneProxy.h b/Engine/Plugins/2D/Paper2D/Source/Paper2D/Private/PaperRenderSceneProxy.h
--- a/Engine/Plugins/2D/Paper2D/Source/Paper2D/Private/PaperRenderSceneProxy.h
+++ b/Engine/Plugins/2D/Paper2D/Source/Paper2D/Private/PaperRenderSceneProxy.h
@@ -98,6 +98,7 @@ public:
 	// End of FPrimitiveSceneProxy interface.
 
 	void SetDrawCall_RenderThread(const FSpriteDrawCallRecord& NewDynamicData);
+	void SetBodySetup_RenderThread(UBodySetup* NewSetup);
 
 	class UMaterialInterface* GetMaterial() const
 	{
diff --git a/Engine/Plugins/2D/Paper2D/Source/Paper2DEditor/Private/FlipbookEditor/FlipbookEditorViewportClient.cpp b/Engine/Plugins/2D/Paper2D/Source/Paper2DEditor/Private/FlipbookEditor/FlipbookEditorViewportClient.cpp
--- a/Engine/Plugins/2D/Paper2D/Source/Paper2DEditor/Private/FlipbookEditor/FlipbookEditorViewportClient.cpp
+++ b/Engine/Plugins/2D/Paper2D/Source/Paper2DEditor/Private/FlipbookEditor/FlipbookEditorViewportClient.cpp
@@ -24,6 +24,7 @@ FFlipbookEditorViewportClient::FFlipbookEditorViewportClient(const TAttribute<UP
 
 	// Create a render component for the sprite being edited
 	AnimatedRenderComponent = NewObject<UPaperFlipbookComponent>();
+	AnimatedRenderComponent->SetCollisionProfileName(UCollisionProfile::BlockAll_ProfileName);
 	AnimatedRenderComponent->SetFlipbook(FlipbookBeingEdited.Get());
 	AnimatedRenderComponent->UpdateBounds();
 	PreviewScene->AddComponent(AnimatedRenderComponent.Get(), FTransform::Identity);
