Paper2D: Replaced uses of NULL with nullptr in Paper2D

diff --git a/Engine/Plugins/2D/Paper2D/Source/Paper2D/Classes/PaperFlipbook.h b/Engine/Plugins/2D/Paper2D/Source/Paper2D/Classes/PaperFlipbook.h
--- a/Engine/Plugins/2D/Paper2D/Source/Paper2D/Classes/PaperFlipbook.h
+++ b/Engine/Plugins/2D/Paper2D/Source/Paper2D/Classes/PaperFlipbook.h
@@ -18,7 +18,7 @@ public:
 	int32 FrameRun;
 
 	FPaperFlipbookKeyFrame()
-		: Sprite(NULL)
+		: Sprite(nullptr)
 		, FrameRun(1)
 	{
 	}
diff --git a/Engine/Plugins/2D/Paper2D/Source/Paper2D/Classes/SpriteDrawCall.h b/Engine/Plugins/2D/Paper2D/Source/Paper2D/Classes/SpriteDrawCall.h
--- a/Engine/Plugins/2D/Paper2D/Source/Paper2D/Classes/SpriteDrawCall.h
+++ b/Engine/Plugins/2D/Paper2D/Source/Paper2D/Classes/SpriteDrawCall.h
@@ -27,7 +27,7 @@ public:
 	void BuildFromSprite(const class UPaperSprite* Sprite);
 
 	FSpriteDrawCallRecord()
-		: Texture(NULL)
+		: Texture(nullptr)
 		, Color(FLinearColor::White)
 	{
 	}
diff --git a/Engine/Plugins/2D/Paper2D/Source/Paper2D/Private/AlphaBitmap.h b/Engine/Plugins/2D/Paper2D/Source/Paper2D/Private/AlphaBitmap.h
--- a/Engine/Plugins/2D/Paper2D/Source/Paper2D/Private/AlphaBitmap.h
+++ b/Engine/Plugins/2D/Paper2D/Source/Paper2D/Private/AlphaBitmap.h
@@ -18,8 +18,8 @@ struct FAlphaBitmap
 	void ExtractFromTexture(UTexture2D* SourceTexture)
 	{
 		// use the source art if it exists
-		FTextureSource* TextureSource = NULL;
-		if ((SourceTexture != NULL) && SourceTexture->Source.IsValid())
+		FTextureSource* TextureSource = nullptr;
+		if ((SourceTexture != nullptr) && SourceTexture->Source.IsValid())
 		{
 			switch (SourceTexture->Source.GetFormat())
 			{
@@ -32,7 +32,7 @@ struct FAlphaBitmap
 			};
 		}
 
-		if (TextureSource != NULL)
+		if (TextureSource != nullptr)
 		{
 			TArray<uint8> TextureRawData;
 			TextureSource->GetMipData(TextureRawData, 0);
diff --git a/Engine/Plugins/2D/Paper2D/Source/Paper2D/Private/PaperFlipbookComponent.cpp b/Engine/Plugins/2D/Paper2D/Source/Paper2D/Private/PaperFlipbookComponent.cpp
--- a/Engine/Plugins/2D/Paper2D/Source/Paper2D/Private/PaperFlipbookComponent.cpp
+++ b/Engine/Plugins/2D/Paper2D/Source/Paper2D/Private/PaperFlipbookComponent.cpp
@@ -332,7 +332,7 @@ void UPaperFlipbookComponent::TickComponent(float DeltaTime, enum ELevelTick Tic
 
 void UPaperFlipbookComponent::SendRenderDynamicData_Concurrent()
 {
-	if (SceneProxy != NULL)
+	if (SceneProxy != nullptr)
 	{
 		UPaperSprite* SpriteToSend = GetSpriteAtCachedIndex();
 
@@ -358,7 +358,7 @@ bool UPaperFlipbookComponent::SetFlipbook(class UPaperFlipbook* NewFlipbook)
 	{
 		// Don't allow changing the sprite if we are "static".
 		AActor* Owner = GetOwner();
-		if (!IsRegistered() || (Owner == NULL) || (Mobility != EComponentMobility::Static))
+		if (!IsRegistered() || (Owner == nullptr) || (Mobility != EComponentMobility::Static))
 		{
 			SourceFlipbook = NewFlipbook;
 
diff --git a/Engine/Plugins/2D/Paper2D/Source/Paper2D/Private/PaperRenderSceneProxy.cpp b/Engine/Plugins/2D/Paper2D/Source/Paper2D/Private/PaperRenderSceneProxy.cpp
--- a/Engine/Plugins/2D/Paper2D/Source/Paper2D/Private/PaperRenderSceneProxy.cpp
+++ b/Engine/Plugins/2D/Paper2D/Source/Paper2D/Private/PaperRenderSceneProxy.cpp
@@ -157,7 +157,7 @@ public:
 
 FPaperRenderSceneProxy::FPaperRenderSceneProxy(const UPrimitiveComponent* InComponent)
 	: FPrimitiveSceneProxy(InComponent)
-	, Material(NULL)
+	, Material(nullptr)
 	, Owner(InComponent->GetOwner())
 	, BodySetup(const_cast<UPrimitiveComponent*>(InComponent)->GetBodySetup())
 	, bCastShadow(InComponent->CastShadow)
@@ -275,7 +275,7 @@ void FPaperRenderSceneProxy::GetDynamicMeshElements(const TArray<const FSceneVie
 							// wireframe
 							FColor CollisionColor = FColor(157, 149, 223, 255);
 							FTransform GeomTransform(GetLocalToWorld());
-							BodySetup->AggGeom.GetAggGeom(GeomTransform, GetSelectionColor(CollisionColor, IsSelected(), IsHovered()), NULL, (Owner == NULL), false, UseEditorDepthTest(), ViewIndex, Collector);
+							BodySetup->AggGeom.GetAggGeom(GeomTransform, GetSelectionColor(CollisionColor, IsSelected(), IsHovered()), nullptr, (Owner == nullptr), false, UseEditorDepthTest(), ViewIndex, Collector);
 						}
 					}
 				}
@@ -285,7 +285,7 @@ void FPaperRenderSceneProxy::GetDynamicMeshElements(const TArray<const FSceneVie
 #if !(UE_BUILD_SHIPPING || UE_BUILD_TEST)
 			if (EngineShowFlags.Paper2DSprites)
 			{
-				RenderBounds(Collector.GetPDI(ViewIndex), EngineShowFlags, GetBounds(), (Owner == NULL) || IsSelected());
+				RenderBounds(Collector.GetPDI(ViewIndex), EngineShowFlags, GetBounds(), (Owner == nullptr) || IsSelected());
 			}
 #endif
 		}
@@ -366,7 +366,7 @@ void FPaperRenderSceneProxy::GetNewBatchMeshes(const FSceneView* View, bool bUse
 			Mesh.DynamicVertexStride = sizeof(FPaperSpriteVertex);
 
 			Mesh.VertexFactory = VertexFactory;
-			Mesh.LCI = NULL;
+			Mesh.LCI = nullptr;
 			Mesh.ReverseCulling = IsLocalToWorldDeterminantNegative() ? true : false;
 			Mesh.CastShadow = bCastShadow;
 			Mesh.DepthPriorityGroup = DPG;
@@ -376,8 +376,8 @@ void FPaperRenderSceneProxy::GetNewBatchMeshes(const FSceneView* View, bool bUse
 
 			// Set up the FMeshBatchElement.
 			FMeshBatchElement& BatchElement = Mesh.Elements[0];
-			BatchElement.IndexBuffer = NULL;
-			BatchElement.DynamicIndexData = NULL;
+			BatchElement.IndexBuffer = nullptr;
+			BatchElement.DynamicIndexData = nullptr;
 			BatchElement.DynamicIndexStride = 0;
 			BatchElement.PrimitiveUniformBuffer = GetUniformBuffer();
 			BatchElement.FirstIndex = Batch.VertexOffset;
@@ -458,7 +458,7 @@ void FPaperRenderSceneProxy::GetBatchMesh(const FSceneView* View, bool bUseOverr
 			Collector.RegisterOneFrameMaterialProxy(TextureOverrideMaterialProxy);
 
 			Mesh.MaterialRenderProxy = TextureOverrideMaterialProxy;
-			Mesh.LCI = NULL;
+			Mesh.LCI = nullptr;
 			Mesh.ReverseCulling = IsLocalToWorldDeterminantNegative() ? true : false;
 			Mesh.CastShadow = false;
 			Mesh.DepthPriorityGroup = DPG;
@@ -467,8 +467,8 @@ void FPaperRenderSceneProxy::GetBatchMesh(const FSceneView* View, bool bUseOverr
 
 			// Set up the FMeshBatchElement.
 			FMeshBatchElement& BatchElement = Mesh.Elements[0];
-			BatchElement.IndexBuffer = NULL;
-			BatchElement.DynamicIndexData = NULL;
+			BatchElement.IndexBuffer = nullptr;
+			BatchElement.DynamicIndexData = nullptr;
 			BatchElement.DynamicIndexStride = 0;
 			BatchElement.FirstIndex = 0;
 			BatchElement.MinVertexIndex = 0;
diff --git a/Engine/Plugins/2D/Paper2D/Source/Paper2D/Private/PaperSprite.cpp b/Engine/Plugins/2D/Paper2D/Source/Paper2D/Private/PaperSprite.cpp
--- a/Engine/Plugins/2D/Paper2D/Source/Paper2D/Private/PaperSprite.cpp
+++ b/Engine/Plugins/2D/Paper2D/Source/Paper2D/Private/PaperSprite.cpp
@@ -372,8 +372,8 @@ void UPaperSprite::PostEditChangeProperty(FPropertyChangedEvent& PropertyChanged
 	FSpriteReregisterContext ReregisterExistingComponents(this);
 
 	// Look for changed properties
-	const FName PropertyName = (PropertyChangedEvent.Property != NULL) ? PropertyChangedEvent.Property->GetFName() : NAME_None;
-	const FName MemberPropertyName = (PropertyChangedEvent.MemberProperty != NULL) ? PropertyChangedEvent.MemberProperty->GetFName() : NAME_None;
+	const FName PropertyName = (PropertyChangedEvent.Property != nullptr) ? PropertyChangedEvent.Property->GetFName() : NAME_None;
+	const FName MemberPropertyName = (PropertyChangedEvent.MemberProperty != nullptr) ? PropertyChangedEvent.MemberProperty->GetFName() : NAME_None;
 
 	if (PixelsPerUnrealUnit <= 0.0f)
 	{
@@ -1524,7 +1524,7 @@ FPaperSpriteSocket* UPaperSprite::FindSocket(FName SocketName)
 		}
 	}
 
-	return NULL;
+	return nullptr;
 }
 
 void UPaperSprite::QuerySupportedSockets(TArray<FComponentSocketDescription>& OutSockets) const
diff --git a/Engine/Plugins/2D/Paper2D/Source/Paper2D/Private/PaperSpriteComponent.cpp b/Engine/Plugins/2D/Paper2D/Source/Paper2D/Private/PaperSpriteComponent.cpp
--- a/Engine/Plugins/2D/Paper2D/Source/Paper2D/Private/PaperSpriteComponent.cpp
+++ b/Engine/Plugins/2D/Paper2D/Source/Paper2D/Private/PaperSpriteComponent.cpp
@@ -76,7 +76,7 @@ FPrimitiveSceneProxy* UPaperSpriteComponent::CreateSceneProxy()
 
 FBoxSphereBounds UPaperSpriteComponent::CalcBounds(const FTransform& LocalToWorld) const
 {
-	if (SourceSprite != NULL)
+	if (SourceSprite != nullptr)
 	{
 		// Graphics bounds.
 		FBoxSphereBounds NewBounds = SourceSprite->GetRenderBounds().TransformBy(LocalToWorld);
@@ -105,7 +105,7 @@ FBoxSphereBounds UPaperSpriteComponent::CalcBounds(const FTransform& LocalToWorl
 
 void UPaperSpriteComponent::SendRenderDynamicData_Concurrent()
 {
-	if (SceneProxy != NULL)
+	if (SceneProxy != nullptr)
 	{
 		FSpriteDrawCallRecord DrawCall;
 		DrawCall.BuildFromSprite(SourceSprite);
@@ -125,7 +125,7 @@ void UPaperSpriteComponent::SendRenderDynamicData_Concurrent()
 
 bool UPaperSpriteComponent::HasAnySockets() const
 {
-	if (SourceSprite != NULL)
+	if (SourceSprite != nullptr)
 	{
 		return SourceSprite->HasAnySockets();
 	}
@@ -135,7 +135,7 @@ bool UPaperSpriteComponent::HasAnySockets() const
 
 FTransform UPaperSpriteComponent::GetSocketTransform(FName InSocketName, ERelativeTransformSpace TransformSpace) const
 {
-	if (SourceSprite != NULL)
+	if (SourceSprite != nullptr)
 	{
 		if (FPaperSpriteSocket* Socket = SourceSprite->FindSocket(InSocketName))
 		{
@@ -169,7 +169,7 @@ FTransform UPaperSpriteComponent::GetSocketTransform(FName InSocketName, ERelati
 
 void UPaperSpriteComponent::QuerySupportedSockets(TArray<FComponentSocketDescription>& OutSockets) const
 {
-	if (SourceSprite != NULL)
+	if (SourceSprite != nullptr)
 	{
 		return SourceSprite->QuerySupportedSockets(OutSockets);
 	}
@@ -186,7 +186,7 @@ bool UPaperSpriteComponent::SetSprite(class UPaperSprite* NewSprite)
 	{
 		// Don't allow changing the sprite if we are "static".
 		AActor* Owner = GetOwner();
-		if (!IsRegistered() || (Owner == NULL) || (Mobility != EComponentMobility::Static))
+		if (!IsRegistered() || (Owner == nullptr) || (Mobility != EComponentMobility::Static))
 		{
 			SourceSprite = NewSprite;
 
diff --git a/Engine/Plugins/2D/Paper2D/Source/Paper2D/Private/PaperSpriteSceneProxy.cpp b/Engine/Plugins/2D/Paper2D/Source/Paper2D/Private/PaperSpriteSceneProxy.cpp
--- a/Engine/Plugins/2D/Paper2D/Source/Paper2D/Private/PaperSpriteSceneProxy.cpp
+++ b/Engine/Plugins/2D/Paper2D/Source/Paper2D/Private/PaperSpriteSceneProxy.cpp
@@ -71,7 +71,7 @@ void FPaperSpriteSceneProxy::GetDynamicMeshElements(const TArray<const FSceneVie
 							// Differentiate the color based on bBlockNonZeroExtent.  Helps greatly with skimming a level for optimization opportunities.
 							const FColor CollisionColor = FColor(220,149,223,255);
 
-							const bool bUseSeparateColorPerHull = (Owner == NULL);
+							const bool bUseSeparateColorPerHull = (Owner == nullptr);
 							const bool bDrawSolid = false;
 							BodySetup->AggGeom.GetAggGeom(GeomTransform, GetSelectionColor(CollisionColor, bDrawWireSelected, IsHovered()), CollisionMaterialInstance, bUseSeparateColorPerHull, bDrawSolid, false, ViewIndex, Collector);
 						}
diff --git a/Engine/Plugins/2D/Paper2D/Source/Paper2D/Private/PaperTileLayer.cpp b/Engine/Plugins/2D/Paper2D/Source/Paper2D/Private/PaperTileLayer.cpp
--- a/Engine/Plugins/2D/Paper2D/Source/Paper2D/Private/PaperTileLayer.cpp
+++ b/Engine/Plugins/2D/Paper2D/Source/Paper2D/Private/PaperTileLayer.cpp
@@ -111,7 +111,7 @@ void UPaperTileLayer::PostEditChangeProperty(struct FPropertyChangedEvent& Prope
 {
 	FTileMapLayerReregisterContext ReregisterExistingComponents(this);
 
-	FName PropertyName = (PropertyChangedEvent.Property != NULL) ? PropertyChangedEvent.Property->GetFName() : NAME_None;
+	FName PropertyName = (PropertyChangedEvent.Property != nullptr) ? PropertyChangedEvent.Property->GetFName() : NAME_None;
 
 	if ((PropertyName == GET_MEMBER_NAME_CHECKED(UPaperTileLayer, LayerWidth)) || (PropertyName == GET_MEMBER_NAME_CHECKED(UPaperTileLayer, LayerHeight)))
 	{
diff --git a/Engine/Plugins/2D/Paper2D/Source/Paper2D/Private/PaperTileMap.cpp b/Engine/Plugins/2D/Paper2D/Source/Paper2D/Private/PaperTileMap.cpp
--- a/Engine/Plugins/2D/Paper2D/Source/Paper2D/Private/PaperTileMap.cpp
+++ b/Engine/Plugins/2D/Paper2D/Source/Paper2D/Private/PaperTileMap.cpp
@@ -39,7 +39,7 @@ UPaperTileMap::UPaperTileMap(const FObjectInitializer& ObjectInitializer)
 
 void UPaperTileMap::PostEditChangeProperty(FPropertyChangedEvent& PropertyChangedEvent)
 {
-	const FName PropertyName = (PropertyChangedEvent.Property != NULL) ? PropertyChangedEvent.Property->GetFName() : NAME_None;
+	const FName PropertyName = (PropertyChangedEvent.Property != nullptr) ? PropertyChangedEvent.Property->GetFName() : NAME_None;
 
 	//@TODO: Determine when these are really needed, as they're seriously expensive!
 	TComponentReregisterContext<UPaperTileMapRenderComponent> ReregisterStaticComponents;
diff --git a/Engine/Plugins/2D/Paper2D/Source/Paper2D/Private/PaperTileMapRenderComponent.cpp b/Engine/Plugins/2D/Paper2D/Source/Paper2D/Private/PaperTileMapRenderComponent.cpp
--- a/Engine/Plugins/2D/Paper2D/Source/Paper2D/Private/PaperTileMapRenderComponent.cpp
+++ b/Engine/Plugins/2D/Paper2D/Source/Paper2D/Private/PaperTileMapRenderComponent.cpp
@@ -41,7 +41,7 @@ void UPaperTileMapRenderComponent::PostInitProperties()
 
 FBoxSphereBounds UPaperTileMapRenderComponent::CalcBounds(const FTransform& LocalToWorld) const
 {
-	if (TileMap != NULL)
+	if (TileMap != nullptr)
 	{
 		// Graphics bounds.
 		FBoxSphereBounds NewBounds = TileMap->GetRenderBounds().TransformBy(LocalToWorld);
diff --git a/Engine/Plugins/2D/Paper2D/Source/Paper2D/Private/PaperTileMapRenderSceneProxy.cpp b/Engine/Plugins/2D/Paper2D/Source/Paper2D/Private/PaperTileMapRenderSceneProxy.cpp
--- a/Engine/Plugins/2D/Paper2D/Source/Paper2D/Private/PaperTileMapRenderSceneProxy.cpp
+++ b/Engine/Plugins/2D/Paper2D/Source/Paper2D/Private/PaperTileMapRenderSceneProxy.cpp
@@ -41,7 +41,7 @@ void FPaperTileMapRenderSceneProxy::GetDynamicMeshElements(const TArray<const FS
 
 			// Draw the tile maps
 			//@TODO: RenderThread race condition
-			if (TileMap != NULL)
+			if (TileMap != nullptr)
 			{
 				FColor WireframeColor = FColor(0, 255, 255, 255);
 
diff --git a/Engine/Plugins/2D/Paper2D/Source/Paper2D/Private/PaperTileSet.cpp b/Engine/Plugins/2D/Paper2D/Source/Paper2D/Private/PaperTileSet.cpp
--- a/Engine/Plugins/2D/Paper2D/Source/Paper2D/Private/PaperTileSet.cpp
+++ b/Engine/Plugins/2D/Paper2D/Source/Paper2D/Private/PaperTileSet.cpp
@@ -14,7 +14,7 @@ UPaperTileSet::UPaperTileSet(const FObjectInitializer& ObjectInitializer)
 
 int32 UPaperTileSet::GetTileCount() const
 {
-	if (TileSheet != NULL)
+	if (TileSheet != nullptr)
 	{
 		checkSlow((TileWidth > 0) && (TileHeight > 0));
 		const int32 TextureWidth = TileSheet->GetSizeX();
@@ -33,7 +33,7 @@ int32 UPaperTileSet::GetTileCount() const
 
 int32 UPaperTileSet::GetTileCountX() const
 {
-	if (TileSheet != NULL)
+	if (TileSheet != nullptr)
 	{
 		checkSlow(TileWidth > 0);
 		const int32 TextureWidth = TileSheet->GetSizeX();
@@ -48,7 +48,7 @@ int32 UPaperTileSet::GetTileCountX() const
 
 int32 UPaperTileSet::GetTileCountY() const
 {
-	if (TileSheet != NULL)
+	if (TileSheet != nullptr)
 	{
 		checkSlow(TileHeight > 0);
 		const int32 TextureHeight = TileSheet->GetSizeY();
@@ -64,7 +64,7 @@ int32 UPaperTileSet::GetTileCountY() const
 bool UPaperTileSet::GetTileUV(int32 TileIndex, /*out*/ FVector2D& Out_TileUV) const
 {
 	//@TODO: Performance: should cache this stuff
-	if (TileSheet != NULL)
+	if (TileSheet != nullptr)
 	{
 		checkSlow((TileWidth > 0) && (TileHeight > 0));
 		const int32 TextureWidth = TileSheet->GetSizeX() - (Margin * 2) + Spacing;
@@ -102,7 +102,7 @@ bool UPaperTileSet::GetTileUV(int32 TileIndex, /*out*/ FVector2D& Out_TileUV) co
 
 void UPaperTileSet::PostEditChangeProperty(FPropertyChangedEvent& PropertyChangedEvent)
 {
-	const FName PropertyName = (PropertyChangedEvent.Property != NULL) ? PropertyChangedEvent.Property->GetFName() : NAME_None;
+	const FName PropertyName = (PropertyChangedEvent.Property != nullptr) ? PropertyChangedEvent.Property->GetFName() : NAME_None;
 
 	Margin = FMath::Max<int32>(Margin, 0);
 	Spacing = FMath::Max<int32>(Spacing, 0);
diff --git a/Engine/Plugins/2D/Paper2D/Source/Paper2DEditor/Classes/PaperSpriteFactory.h b/Engine/Plugins/2D/Paper2D/Source/Paper2DEditor/Classes/PaperSpriteFactory.h
--- a/Engine/Plugins/2D/Paper2D/Source/Paper2DEditor/Classes/PaperSpriteFactory.h
+++ b/Engine/Plugins/2D/Paper2D/Source/Paper2DEditor/Classes/PaperSpriteFactory.h
@@ -13,7 +13,7 @@ class UPaperSpriteFactory : public UFactory
 {
 	GENERATED_UCLASS_BODY()
 
-	// Initial texture to create the sprite from (Can be NULL)
+	// Initial texture to create the sprite from (Can be nullptr)
 	class UTexture2D* InitialTexture;
 
 	// Set bUseSourceRegion to get it to use/set initial sourceUV and dimensions
diff --git a/Engine/Plugins/2D/Paper2D/Source/Paper2DEditor/Private/ContentBrowserExtensions/ContentBrowserExtensions.cpp b/Engine/Plugins/2D/Paper2D/Source/Paper2DEditor/Private/ContentBrowserExtensions/ContentBrowserExtensions.cpp
--- a/Engine/Plugins/2D/Paper2D/Source/Paper2DEditor/Private/ContentBrowserExtensions/ContentBrowserExtensions.cpp
+++ b/Engine/Plugins/2D/Paper2D/Source/Paper2DEditor/Private/ContentBrowserExtensions/ContentBrowserExtensions.cpp
@@ -301,7 +301,7 @@ public:
 			Extender->AddMenuExtension(
 				"GetAssetActions",
 				EExtensionHook::After,
-				NULL,
+				nullptr,
 				FMenuExtensionDelegate::CreateStatic(&FPaperContentBrowserExtensions_Impl::CreateSpriteActionsSubMenu, SpriteCreatorFunctor, SpriteExtractorFunctor, TextureConfigFunctor));
 		}
 
diff --git a/Engine/Plugins/2D/Paper2D/Source/Paper2DEditor/Private/FlipbookEditor/FlipbookEditor.cpp b/Engine/Plugins/2D/Paper2D/Source/Paper2DEditor/Private/FlipbookEditor/FlipbookEditor.cpp
--- a/Engine/Plugins/2D/Paper2D/Source/Paper2DEditor/Private/FlipbookEditor/FlipbookEditor.cpp
+++ b/Engine/Plugins/2D/Paper2D/Source/Paper2DEditor/Private/FlipbookEditor/FlipbookEditor.cpp
@@ -45,7 +45,7 @@ class SFlipbookEditorViewport : public SEditorViewport, public ICommonEditorView
 {
 public:
 	SLATE_BEGIN_ARGS(SFlipbookEditorViewport)
-		: _FlipbookBeingEdited((UPaperFlipbook*)NULL)
+		: _FlipbookBeingEdited((UPaperFlipbook*)nullptr)
 	{}
 
 		SLATE_ATTRIBUTE( UPaperFlipbook*, FlipbookBeingEdited )
diff --git a/Engine/Plugins/2D/Paper2D/Source/Paper2DEditor/Private/FlipbookEditor/FlipbookEditorViewportClient.cpp b/Engine/Plugins/2D/Paper2D/Source/Paper2DEditor/Private/FlipbookEditor/FlipbookEditorViewportClient.cpp
--- a/Engine/Plugins/2D/Paper2D/Source/Paper2DEditor/Private/FlipbookEditor/FlipbookEditorViewportClient.cpp
+++ b/Engine/Plugins/2D/Paper2D/Source/Paper2DEditor/Private/FlipbookEditor/FlipbookEditorViewportClient.cpp
@@ -43,7 +43,7 @@ void FFlipbookEditorViewportClient::DrawCanvas(FViewport& Viewport, FSceneView&
 	const bool bIsHitTesting = Canvas.IsHitTesting();
 	if (!bIsHitTesting)
 	{
-		Canvas.SetHitProxy(NULL);
+		Canvas.SetHitProxy(nullptr);
 	}
 
 	int32 YPos = 42;
diff --git a/Engine/Plugins/2D/Paper2D/Source/Paper2DEditor/Private/FlipbookEditor/STimelineHeader.h b/Engine/Plugins/2D/Paper2D/Source/Paper2DEditor/Private/FlipbookEditor/STimelineHeader.h
--- a/Engine/Plugins/2D/Paper2D/Source/Paper2DEditor/Private/FlipbookEditor/STimelineHeader.h
+++ b/Engine/Plugins/2D/Paper2D/Source/Paper2DEditor/Private/FlipbookEditor/STimelineHeader.h
@@ -11,7 +11,7 @@ class STimelineHeader : public SCompoundWidget
 public:
 	SLATE_BEGIN_ARGS(STimelineHeader)
 		: _SlateUnitsPerFrame(1)
-		, _FlipbookBeingEdited((class UPaperFlipbook*)NULL)
+		, _FlipbookBeingEdited((class UPaperFlipbook*)nullptr)
 		, _PlayTime(0)
 	{}
 
diff --git a/Engine/Plugins/2D/Paper2D/Source/Paper2DEditor/Private/PaperEditorViewportClient.cpp b/Engine/Plugins/2D/Paper2D/Source/Paper2DEditor/Private/PaperEditorViewportClient.cpp
--- a/Engine/Plugins/2D/Paper2D/Source/Paper2DEditor/Private/PaperEditorViewportClient.cpp
+++ b/Engine/Plugins/2D/Paper2D/Source/Paper2DEditor/Private/PaperEditorViewportClient.cpp
@@ -42,7 +42,7 @@ USelection* FAssetEditorModeTools::GetSelectedObjects() const
 
 FPaperEditorViewportClient::FPaperEditorViewportClient()
 	: FEditorViewportClient(new FAssetEditorModeTools())
-	, CheckerboardTexture(NULL)
+	, CheckerboardTexture(nullptr)
 {
 	bOwnsModeTools = true;
 	ZoomPos = FVector2D::ZeroVector;
@@ -134,7 +134,7 @@ void FPaperEditorViewportClient::SetupCheckerboardTexture(const FColor& ColorOne
 	CheckerSize = FMath::RoundUpToPowerOfTwo(CheckerSize);
 	const int32 HalfPixelNum = CheckerSize >> 1;
 
-	if (CheckerboardTexture == NULL)
+	if (CheckerboardTexture == nullptr)
 	{
 		// Create the texture
 		CheckerboardTexture = UTexture2D::CreateTransient(CheckerSize, CheckerSize, PF_B8G8R8A8);
@@ -175,6 +175,6 @@ void FPaperEditorViewportClient::DestroyCheckerboardTexture()
 			CheckerboardTexture->ReleaseResource();
 		}
 		CheckerboardTexture->MarkPendingKill();
-		CheckerboardTexture = NULL;
+		CheckerboardTexture = nullptr;
 	}
 }
diff --git a/Engine/Plugins/2D/Paper2D/Source/Paper2DEditor/Private/PaperFlipbookAssetBroker.h b/Engine/Plugins/2D/Paper2D/Source/Paper2DEditor/Private/PaperFlipbookAssetBroker.h
--- a/Engine/Plugins/2D/Paper2D/Source/Paper2DEditor/Private/PaperFlipbookAssetBroker.h
+++ b/Engine/Plugins/2D/Paper2D/Source/Paper2DEditor/Private/PaperFlipbookAssetBroker.h
@@ -21,7 +21,7 @@ public:
 		{
 			UPaperFlipbook* Flipbook = Cast<UPaperFlipbook>(InAsset);
 
-			if ((Flipbook != NULL) || (InAsset == NULL))
+			if ((Flipbook != nullptr) || (InAsset == nullptr))
 			{
 				RenderComp->SetFlipbook(Flipbook);
 				return true;
@@ -37,7 +37,7 @@ public:
 		{
 			return RenderComp->GetFlipbook();
 		}
-		return NULL;
+		return nullptr;
 	}
 };
 
diff --git a/Engine/Plugins/2D/Paper2D/Source/Paper2DEditor/Private/PaperJSONHelpers.cpp b/Engine/Plugins/2D/Paper2D/Source/Paper2DEditor/Private/PaperJSONHelpers.cpp
--- a/Engine/Plugins/2D/Paper2D/Source/Paper2DEditor/Private/PaperJSONHelpers.cpp
+++ b/Engine/Plugins/2D/Paper2D/Source/Paper2DEditor/Private/PaperJSONHelpers.cpp
@@ -23,7 +23,7 @@ TSharedPtr<class FJsonObject> FPaperJSONHelpers::ReadObject(TSharedPtr<class FJs
 	}
 	else
 	{
-		return NULL;
+		return nullptr;
 	}
 }
 
diff --git a/Engine/Plugins/2D/Paper2D/Source/Paper2DEditor/Private/PaperSpriteThumbnailRenderer.cpp b/Engine/Plugins/2D/Paper2D/Source/Paper2DEditor/Private/PaperSpriteThumbnailRenderer.cpp
--- a/Engine/Plugins/2D/Paper2D/Source/Paper2DEditor/Private/PaperSpriteThumbnailRenderer.cpp
+++ b/Engine/Plugins/2D/Paper2D/Source/Paper2DEditor/Private/PaperSpriteThumbnailRenderer.cpp
@@ -19,10 +19,10 @@ void UPaperSpriteThumbnailRenderer::Draw(UObject* Object, int32 X, int32 Y, uint
 
 void UPaperSpriteThumbnailRenderer::DrawGrid(int32 X, int32 Y, uint32 Width, uint32 Height, FCanvas* Canvas)
 {
-	static UTexture2D* GridTexture = NULL;
-	if (GridTexture == NULL)
+	static UTexture2D* GridTexture = nullptr;
+	if (GridTexture == nullptr)
 	{
-		GridTexture = LoadObject<UTexture2D>(NULL, TEXT("/Engine/EngineMaterials/DefaultWhiteGrid.DefaultWhiteGrid"), NULL, LOAD_None, NULL);
+		GridTexture = LoadObject<UTexture2D>(nullptr, TEXT("/Engine/EngineMaterials/DefaultWhiteGrid.DefaultWhiteGrid"), nullptr, LOAD_None, nullptr);
 	}
 
 	const bool bAlphaBlend = false;
@@ -43,7 +43,7 @@ void UPaperSpriteThumbnailRenderer::DrawGrid(int32 X, int32 Y, uint32 Width, uin
 
 void UPaperSpriteThumbnailRenderer::DrawFrame(class UPaperSprite* Sprite, int32 X, int32 Y, uint32 Width, uint32 Height, FRenderTarget*, FCanvas* Canvas, FBoxSphereBounds* OverrideRenderBounds)
 {
-	if (const UTexture2D* SourceTexture = (Sprite != NULL) ? Sprite->GetSourceTexture() : NULL)
+	if (const UTexture2D* SourceTexture = (Sprite != nullptr) ? Sprite->GetSourceTexture() : nullptr)
 	{
 		const bool bUseTranslucentBlend = SourceTexture->HasAlphaChannel();
 
diff --git a/Engine/Plugins/2D/Paper2D/Source/Paper2DEditor/Private/PaperStyle.cpp b/Engine/Plugins/2D/Paper2D/Source/Paper2DEditor/Private/PaperStyle.cpp
--- a/Engine/Plugins/2D/Paper2D/Source/Paper2DEditor/Private/PaperStyle.cpp
+++ b/Engine/Plugins/2D/Paper2D/Source/Paper2DEditor/Private/PaperStyle.cpp
@@ -14,7 +14,7 @@ FString FPaperStyle::InContent(const FString& RelativePath, const ANSICHAR* Exte
 	return (ContentDir / RelativePath) + Extension;
 }
 
-TSharedPtr< FSlateStyleSet > FPaperStyle::StyleSet = NULL;
+TSharedPtr< FSlateStyleSet > FPaperStyle::StyleSet = nullptr;
 TSharedPtr< class ISlateStyle > FPaperStyle::Get() { return StyleSet; }
 
 void FPaperStyle::Initialize()
diff --git a/Engine/Plugins/2D/Paper2D/Source/Paper2DEditor/Private/PaperTileSetThumbnailRenderer.cpp b/Engine/Plugins/2D/Paper2D/Source/Paper2DEditor/Private/PaperTileSetThumbnailRenderer.cpp
--- a/Engine/Plugins/2D/Paper2D/Source/Paper2DEditor/Private/PaperTileSetThumbnailRenderer.cpp
+++ b/Engine/Plugins/2D/Paper2D/Source/Paper2DEditor/Private/PaperTileSetThumbnailRenderer.cpp
@@ -14,17 +14,17 @@ UPaperTileSetThumbnailRenderer::UPaperTileSetThumbnailRenderer(const FObjectInit
 void UPaperTileSetThumbnailRenderer::Draw(UObject* Object, int32 X, int32 Y, uint32 Width, uint32 Height, FRenderTarget*, FCanvas* Canvas)
 {
 	UPaperTileSet* TileSet = Cast<UPaperTileSet>(Object);
-	if ((TileSet != NULL) && (TileSet->TileSheet != NULL))
+	if ((TileSet != nullptr) && (TileSet->TileSheet != nullptr))
 	{
 		const bool bUseTranslucentBlend = TileSet->TileSheet->HasAlphaChannel();
 
 		// Draw the grid behind the sprite
 		if (bUseTranslucentBlend)
 		{
-			static UTexture2D* GridTexture = NULL;
-			if (GridTexture == NULL)
+			static UTexture2D* GridTexture = nullptr;
+			if (GridTexture == nullptr)
 			{
-				GridTexture = LoadObject<UTexture2D>(NULL, TEXT("/Engine/EngineMaterials/DefaultWhiteGrid.DefaultWhiteGrid"), NULL, LOAD_None, NULL);
+				GridTexture = LoadObject<UTexture2D>(nullptr, TEXT("/Engine/EngineMaterials/DefaultWhiteGrid.DefaultWhiteGrid"), nullptr, LOAD_None, nullptr);
 			}
 
 			const bool bAlphaBlend = false;
diff --git a/Engine/Plugins/2D/Paper2D/Source/Paper2DEditor/Private/SPaperEditorViewport.cpp b/Engine/Plugins/2D/Paper2D/Source/Paper2DEditor/Private/SPaperEditorViewport.cpp
--- a/Engine/Plugins/2D/Paper2D/Source/Paper2DEditor/Private/SPaperEditorViewport.cpp
+++ b/Engine/Plugins/2D/Paper2D/Source/Paper2DEditor/Private/SPaperEditorViewport.cpp
@@ -44,7 +44,7 @@ SPaperEditorViewport::~SPaperEditorViewport()
 	// Close viewport
 	if (ViewportClient.IsValid())
 	{
-		ViewportClient->Viewport = NULL;
+		ViewportClient->Viewport = nullptr;
 	}
 
 	Viewport.Reset();
@@ -196,7 +196,7 @@ FReply SPaperEditorViewport::OnMouseButtonDown(const FGeometry& MyGeometry, cons
 		SoftwareCursorPosition = PanelCoordToGraphCoord( MyGeometry.AbsoluteToLocal( MouseEvent.GetScreenSpacePosition() ) );
 
 		// clear any interpolation when you manually pan
-		//DeferredMovementTargetObject = NULL;
+		//DeferredMovementTargetObject = nullptr;
 
 		return ReplyState;
 	}
@@ -370,7 +370,7 @@ FReply SPaperEditorViewport::OnMouseMove(const FGeometry& MyGeometry, const FPoi
 // 						for (FGraphPanelSelectionSet::TIterator NodeIt(SelectionManager.SelectedNodes); NodeIt; ++NodeIt)
 // 						{
 // 							TSharedRef<SNode>* pWidget = NodeToWidgetLookup.Find(*NodeIt);
-// 							if (pWidget != NULL)
+// 							if (pWidget != nullptr)
 // 							{
 // 								SNode& Widget = pWidget->Get();
 // 								Widget.MoveTo(Widget.GetPosition() + DeltaPos);
diff --git a/Engine/Plugins/2D/Paper2D/Source/Paper2DEditor/Private/SpriteEditor/SSpriteList.cpp b/Engine/Plugins/2D/Paper2D/Source/Paper2DEditor/Private/SpriteEditor/SSpriteList.cpp
--- a/Engine/Plugins/2D/Paper2D/Source/Paper2DEditor/Private/SpriteEditor/SSpriteList.cpp
+++ b/Engine/Plugins/2D/Paper2D/Source/Paper2DEditor/Private/SpriteEditor/SSpriteList.cpp
@@ -32,7 +32,7 @@ void SSpriteList::RebuildWidget(UTexture2D* NewTextureFilter)
 	Config.ThumbnailScale = 0.0f;
 	Config.InitialAssetViewType = EAssetViewType::Tile;
 
-	if (NewTextureFilter != NULL)
+	if (NewTextureFilter != nullptr)
 	{
 		const FString TextureString = FAssetData(NewTextureFilter).GetExportTextName();
 		const FName SourceTexturePropName(TEXT("SourceTexture")); //@TODO: Protected: GET_MEMBER_NAME_CHECKED(UPaperSprite, SourceTexture);
diff --git a/Engine/Plugins/2D/Paper2D/Source/Paper2DEditor/Private/SpriteEditor/SpriteEditor.cpp b/Engine/Plugins/2D/Paper2D/Source/Paper2DEditor/Private/SpriteEditor/SpriteEditor.cpp
--- a/Engine/Plugins/2D/Paper2D/Source/Paper2DEditor/Private/SpriteEditor/SpriteEditor.cpp
+++ b/Engine/Plugins/2D/Paper2D/Source/Paper2DEditor/Private/SpriteEditor/SpriteEditor.cpp
@@ -556,7 +556,7 @@ void FSpriteEditor::ExtendToolbar()
 
 void FSpriteEditor::SetSpriteBeingEdited(UPaperSprite* NewSprite)
 {
-	if ((NewSprite != SpriteBeingEdited) && (NewSprite != NULL))
+	if ((NewSprite != SpriteBeingEdited) && (NewSprite != nullptr))
 	{
 		UPaperSprite* OldSprite = SpriteBeingEdited;
 		SpriteBeingEdited = NewSprite;
diff --git a/Engine/Plugins/2D/Paper2D/Source/Paper2DEditor/Private/SpriteEditor/SpriteEditorSelections.h b/Engine/Plugins/2D/Paper2D/Source/Paper2DEditor/Private/SpriteEditor/SpriteEditorSelections.h
--- a/Engine/Plugins/2D/Paper2D/Source/Paper2DEditor/Private/SpriteEditor/SpriteEditorSelections.h
+++ b/Engine/Plugins/2D/Paper2D/Source/Paper2DEditor/Private/SpriteEditor/SpriteEditorSelections.h
@@ -121,7 +121,7 @@ public:
 		if (UPaperSprite* Sprite = SpritePtr.Get())
 		{
 			UTexture2D* SourceTexture = Sprite->GetSourceTexture();
-			const FVector2D SourceDims = (SourceTexture != NULL) ? FVector2D(SourceTexture->GetSurfaceWidth(), SourceTexture->GetSurfaceHeight()) : FVector2D::ZeroVector;
+			const FVector2D SourceDims = (SourceTexture != nullptr) ? FVector2D(SourceTexture->GetSurfaceWidth(), SourceTexture->GetSurfaceHeight()) : FVector2D::ZeroVector;
             int XAxis = 0; // -1 = min, 0 = none, 1 = max
             int YAxis = 0; // ditto
             
@@ -192,7 +192,7 @@ public:
 		if (UPaperSprite* Sprite = SpritePtr.Get())
 		{
 			UTexture2D* SourceTexture = Sprite->GetSourceTexture();
-			const FVector2D SourceDims = (SourceTexture != NULL) ? FVector2D(SourceTexture->GetSurfaceWidth(), SourceTexture->GetSurfaceHeight()) : FVector2D::ZeroVector;
+			const FVector2D SourceDims = (SourceTexture != nullptr) ? FVector2D(SourceTexture->GetSurfaceWidth(), SourceTexture->GetSurfaceHeight()) : FVector2D::ZeroVector;
 			FVector2D BoundsVertex = Sprite->SourceUV;
 			switch (VertexIndex)
 			{
diff --git a/Engine/Plugins/2D/Paper2D/Source/Paper2DEditor/Private/SpriteEditor/SpriteEditorViewportClient.cpp b/Engine/Plugins/2D/Paper2D/Source/Paper2DEditor/Private/SpriteEditor/SpriteEditorViewportClient.cpp
--- a/Engine/Plugins/2D/Paper2D/Source/Paper2DEditor/Private/SpriteEditor/SpriteEditorViewportClient.cpp
+++ b/Engine/Plugins/2D/Paper2D/Source/Paper2DEditor/Private/SpriteEditor/SpriteEditorViewportClient.cpp
@@ -78,7 +78,7 @@ FSpriteEditorViewportClient::FSpriteEditorViewportClient(TWeakPtr<FSpriteEditor>
 	WidgetMode = FWidget::WM_Translate;
 	bManipulating = false;
 	bManipulationDirtiedSomething = false;
-	ScopedTransaction = NULL;
+	ScopedTransaction = nullptr;
 
 	bShowSourceTexture = false;
 	bShowSockets = true;
@@ -124,7 +124,7 @@ void FSpriteEditorViewportClient::UpdateSourceTextureSpriteFromSprite(UPaperSpri
 	UPaperSprite* TargetSprite = SourceTextureViewComponent->GetSprite();
 	check(TargetSprite);
 
-	if (SourceSprite != NULL)
+	if (SourceSprite != nullptr)
 	{
 		if ((SourceSprite->GetSourceTexture() != TargetSprite->GetSourceTexture()) || (TargetSprite->PixelsPerUnrealUnit != SourceSprite->PixelsPerUnrealUnit))
 		{
@@ -182,7 +182,7 @@ void FSpriteEditorViewportClient::UpdateSourceTextureSpriteFromSprite(UPaperSpri
 	else
 	{
 		// No source sprite, so don't draw the target either
-		TargetSprite->SourceTexture = NULL;
+		TargetSprite->SourceTexture = nullptr;
 	}
 }
 
@@ -317,7 +317,7 @@ void FSpriteEditorViewportClient::DrawSourceRegion(FViewport& InViewport, FScene
 
             if (bIsHitTesting)
             {
-                Canvas.SetHitProxy(NULL);
+                Canvas.SetHitProxy(nullptr);
             }
         }
 
@@ -336,7 +336,7 @@ void FSpriteEditorViewportClient::DrawSourceRegion(FViewport& InViewport, FScene
 
             if (bIsHitTesting)
             {
-                Canvas.SetHitProxy(NULL);
+                Canvas.SetHitProxy(nullptr);
             }
         }
 	}
@@ -419,7 +419,7 @@ void FSpriteEditorViewportClient::DrawGeometry(FViewport& InViewport, FSceneView
 
 				if (bIsHitTesting)
 				{
-					Canvas.SetHitProxy(NULL);
+					Canvas.SetHitProxy(nullptr);
 				}
 			}
 		}
@@ -449,7 +449,7 @@ void FSpriteEditorViewportClient::DrawGeometry(FViewport& InViewport, FSceneView
 			
 			if (bIsHitTesting)
 			{
-				Canvas.SetHitProxy(NULL);
+				Canvas.SetHitProxy(nullptr);
 			}
 		}
 	}
@@ -616,7 +616,7 @@ void FSpriteEditorViewportClient::DrawSockets(const FSceneView* View, FPrimitive
 		
 		DrawWireDiamond(PDI, SocketTM, DiamondSize, DiamondColor, SDPG_Foreground);
 		
-		PDI->SetHitProxy(NULL);
+		PDI->SetHitProxy(nullptr);
 	}
 }
 
@@ -671,7 +671,7 @@ void FSpriteEditorViewportClient::DrawCanvas(FViewport& Viewport, FSceneView& Vi
 	const bool bIsHitTesting = Canvas.IsHitTesting();
 	if (!bIsHitTesting)
 	{
-		Canvas.SetHitProxy(NULL);
+		Canvas.SetHitProxy(nullptr);
 	}
 
 	if (!SpriteEditorPtr.IsValid())
@@ -1394,7 +1394,7 @@ ECoordSystem FSpriteEditorViewportClient::GetWidgetCoordSystemSpace() const
 
 void FSpriteEditorViewportClient::BeginTransaction(const FText& SessionName)
 {
-	if (ScopedTransaction == NULL)
+	if (ScopedTransaction == nullptr)
 	{
 		ScopedTransaction = new FScopedTransaction(SessionName);
 
@@ -1423,10 +1423,10 @@ void FSpriteEditorViewportClient::EndTransaction()
 	
 	bManipulationDirtiedSomething = false;
 
-	if (ScopedTransaction != NULL)
+	if (ScopedTransaction != nullptr)
 	{
 		delete ScopedTransaction;
-		ScopedTransaction = NULL;
+		ScopedTransaction = nullptr;
 	}
 }
 
@@ -1597,7 +1597,7 @@ FSpritePolygonCollection* FSpriteEditorViewportClient::GetGeometryBeingEdited()
 	case ESpriteEditorMode::EditRenderingGeomMode:
 		return &(Sprite->RenderGeometry);
 	default:
-		return NULL;
+		return nullptr;
 	}
 }
 
diff --git a/Engine/Plugins/2D/Paper2D/Source/Paper2DEditor/Private/TileMapEditing/EdModeTileMap.cpp b/Engine/Plugins/2D/Paper2D/Source/Paper2DEditor/Private/TileMapEditing/EdModeTileMap.cpp
--- a/Engine/Plugins/2D/Paper2D/Source/Paper2DEditor/Private/TileMapEditing/EdModeTileMap.cpp
+++ b/Engine/Plugins/2D/Paper2D/Source/Paper2DEditor/Private/TileMapEditing/EdModeTileMap.cpp
@@ -80,7 +80,7 @@ const FEditorModeID FEdModeTileMap::EM_TileMap(TEXT("EM_TileMap"));
 
 FEdModeTileMap::FEdModeTileMap()
 	: bIsPainting(false)
-	, PaintSourceTileSet(NULL)
+	, PaintSourceTileSet(nullptr)
 	, PaintSourceTopLeft(0, 0)
 	, PaintSourceDimensions(0, 0)
 	, DrawPreviewDimensionsLS(0.0f, 0.0f, 0.0f)
@@ -265,7 +265,7 @@ void FEdModeTileMap::Render(const FSceneView* View, FViewport* Viewport, FPrimit
 
 	// If this viewport does not support Mode widgets we will not draw it here.
 	FEditorViewportClient* ViewportClient = (FEditorViewportClient*)Viewport->GetClient();
-	if ((ViewportClient != NULL) && !ViewportClient->EngineShowFlags.ModeWidgets)
+	if ((ViewportClient != nullptr) && !ViewportClient->EngineShowFlags.ModeWidgets)
 	{
 		return;
 	}
@@ -293,7 +293,7 @@ void FEdModeTileMap::Render(const FSceneView* View, FViewport* Viewport, FPrimit
 void FEdModeTileMap::DrawHUD(FEditorViewportClient* ViewportClient, FViewport* Viewport, const FSceneView* View, FCanvas* Canvas)
 {
 	FString InkInfo = FString::Printf(TEXT("Ink: (%d, %d)  %dx%d  %s"), PaintSourceTopLeft.X, PaintSourceTopLeft.Y, PaintSourceDimensions.X, PaintSourceDimensions.Y, 
-		(PaintSourceTileSet.Get() != NULL) ? (*PaintSourceTileSet.Get()->GetName()) : TEXT("(no tile ink selected)"));
+		(PaintSourceTileSet.Get() != nullptr) ? (*PaintSourceTileSet.Get()->GetName()) : TEXT("(no tile ink selected)"));
 
 	FCanvasTextItem Msg(FVector2D(10, 30), FText::FromString(InkInfo), GEngine->GetMediumFont(), FLinearColor::White);
 	Canvas->DrawItem(Msg);
@@ -401,7 +401,7 @@ UPaperTileLayer* FEdModeTileMap::GetSelectedLayerUnderCursor(const FViewportCurs
 
 	OutTileX = 0;
 	OutTileY = 0;
-	return NULL;
+	return nullptr;
 }
 
 
@@ -434,7 +434,7 @@ bool FEdModeTileMap::PaintTiles(const FViewportCursorLocation& Ray)
 	}
 
 	// If we are using an ink source, validate that it exists
-	UPaperTileSet* InkSource = NULL;
+	UPaperTileSet* InkSource = nullptr;
 	if ( GetActiveLayerPaintingMode() != ETileMapLayerPaintingMode::CollisionLayers )
 	{
 		InkSource = PaintSourceTileSet.Get();
@@ -600,7 +600,7 @@ bool FEdModeTileMap::FloodFillTiles(const FViewportCursorLocation& Ray)
 	}
 
 	// If we are using an ink source, validate that it exists
-	UPaperTileSet* InkSource = NULL;
+	UPaperTileSet* InkSource = nullptr;
 	if ( GetActiveLayerPaintingMode() != ETileMapLayerPaintingMode::CollisionLayers )
 	{
 		InkSource = PaintSourceTileSet.Get();
diff --git a/Engine/Plugins/2D/Paper2D/Source/Paper2DEditor/Private/TileMapEditing/PaperTileMapDetailsCustomization.cpp b/Engine/Plugins/2D/Paper2D/Source/Paper2DEditor/Private/TileMapEditing/PaperTileMapDetailsCustomization.cpp
--- a/Engine/Plugins/2D/Paper2D/Source/Paper2DEditor/Private/TileMapEditing/PaperTileMapDetailsCustomization.cpp
+++ b/Engine/Plugins/2D/Paper2D/Source/Paper2DEditor/Private/TileMapEditing/PaperTileMapDetailsCustomization.cpp
@@ -29,7 +29,7 @@ void FPaperTileMapDetailsCustomization::CustomizeDetails(IDetailLayoutBuilder& D
 
 	bool bEditingActor = false;
 
-	UPaperTileMap* TileMap = NULL;
+	UPaperTileMap* TileMap = nullptr;
 	for (int32 ObjectIndex = 0; ObjectIndex < SelectedObjects.Num(); ++ObjectIndex)
 	{
 		UObject* TestObject = SelectedObjects[ObjectIndex].Get();
@@ -77,7 +77,7 @@ void FPaperTileMapDetailsCustomization::CustomizeDetails(IDetailLayoutBuilder& D
 	}
 
 	//@TODO: Handle showing layers when multiple tile maps are selected
-	if (TileMap != NULL)
+	if (TileMap != nullptr)
 	{
 		IDetailCategoryBuilder& LayersCategory = DetailLayout.EditCategory("Tile Layers");
 
diff --git a/Engine/Plugins/2D/Paper2D/Source/Paper2DEditor/Private/TileMapEditing/STileLayerList.cpp b/Engine/Plugins/2D/Paper2D/Source/Paper2DEditor/Private/TileMapEditing/STileLayerList.cpp
--- a/Engine/Plugins/2D/Paper2D/Source/Paper2DEditor/Private/TileMapEditing/STileLayerList.cpp
+++ b/Engine/Plugins/2D/Paper2D/Source/Paper2DEditor/Private/TileMapEditing/STileLayerList.cpp
@@ -176,7 +176,7 @@ FText STileLayerList::GenerateDuplicatedLayerName(const FString& InputNameRaw, U
 
 UPaperTileLayer* STileLayerList::AddLayer(bool bCollisionLayer, int32 InsertionIndex)
 {
-	UPaperTileLayer* NewLayer = NULL;
+	UPaperTileLayer* NewLayer = nullptr;
 
 	if (UPaperTileMap* TileMap = TileMapPtr.Get())
 	{
diff --git a/Engine/Plugins/2D/Paper2D/Source/Paper2DEditor/Private/TileMapEditing/TileMapEditor.cpp b/Engine/Plugins/2D/Paper2D/Source/Paper2DEditor/Private/TileMapEditing/TileMapEditor.cpp
--- a/Engine/Plugins/2D/Paper2D/Source/Paper2DEditor/Private/TileMapEditing/TileMapEditor.cpp
+++ b/Engine/Plugins/2D/Paper2D/Source/Paper2DEditor/Private/TileMapEditing/TileMapEditor.cpp
@@ -507,7 +507,7 @@ void FTileMapEditor::ExtendToolbar()
 
 void FTileMapEditor::SetTileMapBeingEdited(UPaperTileMap* NewTileMap)
 {
-	if ((NewTileMap != TileMapBeingEdited) && (NewTileMap != NULL))
+	if ((NewTileMap != TileMapBeingEdited) && (NewTileMap != nullptr))
 	{
 		UPaperTileMap* OldTileMap = TileMapBeingEdited;
 		TileMapBeingEdited = NewTileMap;
diff --git a/Engine/Plugins/2D/Paper2D/Source/Paper2DEditor/Private/TileMapEditing/TileMapEditorViewportClient.cpp b/Engine/Plugins/2D/Paper2D/Source/Paper2DEditor/Private/TileMapEditing/TileMapEditorViewportClient.cpp
--- a/Engine/Plugins/2D/Paper2D/Source/Paper2DEditor/Private/TileMapEditing/TileMapEditorViewportClient.cpp
+++ b/Engine/Plugins/2D/Paper2D/Source/Paper2DEditor/Private/TileMapEditing/TileMapEditorViewportClient.cpp
@@ -31,7 +31,7 @@ FTileMapEditorViewportClient::FTileMapEditorViewportClient(TWeakPtr<FTileMapEdit
 	WidgetMode = FWidget::WM_Translate;
 	bManipulating = false;
 	bManipulationDirtiedSomething = false;
-	ScopedTransaction = NULL;
+	ScopedTransaction = nullptr;
 
 	bShowPivot = true;
 
@@ -93,7 +93,7 @@ void FTileMapEditorViewportClient::DrawCanvas(FViewport& Viewport, FSceneView& V
 	const bool bIsHitTesting = Canvas.IsHitTesting();
 	if (!bIsHitTesting)
 	{
-		Canvas.SetHitProxy(NULL);
+		Canvas.SetHitProxy(nullptr);
 	}
 
 	if (!TileMapEditorPtr.IsValid())
@@ -165,10 +165,10 @@ void FTileMapEditorViewportClient::EndTransaction()
 	
 	bManipulationDirtiedSomething = false;
 
-	if (ScopedTransaction != NULL)
+	if (ScopedTransaction != nullptr)
 	{
 		delete ScopedTransaction;
-		ScopedTransaction = NULL;
+		ScopedTransaction = nullptr;
 	}
 }
 
diff --git a/Engine/Plugins/2D/Paper2D/Source/Paper2DEditor/Private/TileSetEditor.cpp b/Engine/Plugins/2D/Paper2D/Source/Paper2DEditor/Private/TileSetEditor.cpp
--- a/Engine/Plugins/2D/Paper2D/Source/Paper2DEditor/Private/TileSetEditor.cpp
+++ b/Engine/Plugins/2D/Paper2D/Source/Paper2DEditor/Private/TileSetEditor.cpp
@@ -69,14 +69,14 @@ void FTileSetEditorViewportClient::Draw(FViewport* Viewport, FCanvas* Canvas)
 
 	// Can only proceed if we have a valid tile set
 	UPaperTileSet* TileSet = TileSetBeingEdited.Get();
-	if (TileSet == NULL)
+	if (TileSet == nullptr)
 	{
 		return;
 	}
 
 	UTexture2D* Texture = TileSet->TileSheet;
 
-	if (Texture != NULL)
+	if (Texture != nullptr)
 	{
 		const bool bUseTranslucentBlend = Texture->HasAlphaChannel();
 
@@ -130,7 +130,7 @@ void FTileSetEditorViewportClient::Draw(FViewport* Viewport, FCanvas* Canvas)
 
 STileSetSelectorViewport::~STileSetSelectorViewport()
 {
-	TypedViewportClient = NULL;
+	TypedViewportClient = nullptr;
 }
 
 void STileSetSelectorViewport::Construct(const FArguments& InArgs, UPaperTileSet* InTileSet, FEdModeTileMap* InTileMapEditor)
@@ -178,7 +178,7 @@ FText STileSetSelectorViewport::GetTitleText() const
 void STileSetSelectorViewport::OnSelectionChanged(FMarqueeOperation Marquee, bool bIsPreview)
 {
 	UPaperTileSet* TileSetBeingEdited = TileSetPtr.Get();
-	if (TileSetBeingEdited == NULL)
+	if (TileSetBeingEdited == nullptr)
 	{
 		return;
 	}
diff --git a/Engine/Plugins/2D/Paper2D/Source/Paper2DEditor/Public/Paper2DEditorModule.h b/Engine/Plugins/2D/Paper2D/Source/Paper2DEditor/Public/Paper2DEditorModule.h
--- a/Engine/Plugins/2D/Paper2D/Source/Paper2DEditor/Public/Paper2DEditorModule.h
+++ b/Engine/Plugins/2D/Paper2D/Source/Paper2DEditor/Public/Paper2DEditorModule.h
@@ -10,10 +10,10 @@
 class IPaper2DEditorModule : public IModuleInterface
 {
 public:
-	virtual TSharedPtr<class FExtensibilityManager> GetSpriteEditorMenuExtensibilityManager() { return NULL; }
-	virtual TSharedPtr<class FExtensibilityManager> GetSpriteEditorToolBarExtensibilityManager() { return NULL; }
+	virtual TSharedPtr<class FExtensibilityManager> GetSpriteEditorMenuExtensibilityManager() { return nullptr; }
+	virtual TSharedPtr<class FExtensibilityManager> GetSpriteEditorToolBarExtensibilityManager() { return nullptr; }
 
-	virtual TSharedPtr<class FExtensibilityManager> GetFlipbookEditorMenuExtensibilityManager() { return NULL; }
-	virtual TSharedPtr<class FExtensibilityManager> GetFlipbookEditorToolBarExtensibilityManager() { return NULL; }
+	virtual TSharedPtr<class FExtensibilityManager> GetFlipbookEditorMenuExtensibilityManager() { return nullptr; }
+	virtual TSharedPtr<class FExtensibilityManager> GetFlipbookEditorToolBarExtensibilityManager() { return nullptr; }
 };
 
diff --git a/Engine/Plugins/2D/Paper2D/Source/Paper2DEditor/Public/PaperJSONHelpers.h b/Engine/Plugins/2D/Paper2D/Source/Paper2DEditor/Public/PaperJSONHelpers.h
--- a/Engine/Plugins/2D/Paper2D/Source/Paper2DEditor/Public/PaperJSONHelpers.h
+++ b/Engine/Plugins/2D/Paper2D/Source/Paper2DEditor/Public/PaperJSONHelpers.h
@@ -11,10 +11,10 @@ public:
 	// returns the value of Key in Item, or DefaultValue if the Key is missing or the wrong type
 	static FString ReadString(TSharedPtr<class FJsonObject> Item, const FString& Key, const FString& DefaultValue);
 
-	// Returns the object named Key or NULL if it is missing or the wrong type
+	// Returns the object named Key or nullptr if it is missing or the wrong type
 	static TSharedPtr<class FJsonObject> ReadObject(TSharedPtr<class FJsonObject> Item, const FString& Key);
 
-	// Returns the array named Key or NULL if it is missing or the wrong type
+	// Returns the array named Key or nullptr if it is missing or the wrong type
 	static const TArray< TSharedPtr<FJsonValue> >& ReadArray(TSharedPtr<class FJsonObject> Item, const FString& Key);
 
 	// Returns the bool named Key or bDefaultIfMissing if it is missing or the wrong type (note: no way to determine errors!)
diff --git a/Engine/Plugins/2D/Paper2D/Source/PaperJsonImporter/Private/PaperJsonImporterFactory.cpp b/Engine/Plugins/2D/Paper2D/Source/PaperJsonImporter/Private/PaperJsonImporterFactory.cpp
--- a/Engine/Plugins/2D/Paper2D/Source/PaperJsonImporter/Private/PaperJsonImporterFactory.cpp
+++ b/Engine/Plugins/2D/Paper2D/Source/PaperJsonImporter/Private/PaperJsonImporterFactory.cpp
@@ -50,7 +50,7 @@ UObject* UPaperJsonImporterFactory::FactoryCreateText(UClass* InClass, UObject*
 	
 
 	const FString LongPackagePath = FPackageName::GetLongPackagePath(InParent->GetOutermost()->GetPathName());
-	UPaperSpriteSheet* Result = NULL;
+	UPaperSpriteSheet* Result = nullptr;
 	
 	const FString NameForErrors(InName.ToString());
 	const FString FileContent(BufferEnd - Buffer, Buffer);
diff --git a/Engine/Plugins/2D/Paper2D/Source/PaperJsonImporter/Private/PaperJsonSpriteSheetImporter.cpp b/Engine/Plugins/2D/Paper2D/Source/PaperJsonImporter/Private/PaperJsonSpriteSheetImporter.cpp
--- a/Engine/Plugins/2D/Paper2D/Source/PaperJsonImporter/Private/PaperJsonSpriteSheetImporter.cpp
+++ b/Engine/Plugins/2D/Paper2D/Source/PaperJsonImporter/Private/PaperJsonSpriteSheetImporter.cpp
@@ -26,13 +26,13 @@ TSharedPtr<FJsonObject> ParseJSON(const FString& FileContents, const FString& Na
 		else
 		{
 			UE_LOG(LogPaperJsonImporter, Warning, TEXT("Failed to parse sprite descriptor file '%s'.  Error: '%s'"), *NameForErrors, *Reader->GetErrorMessage());
-			return NULL;
+			return nullptr;
 		}
 	}
 	else
 	{
 		UE_LOG(LogPaperJsonImporter, Warning, TEXT("Sprite descriptor file '%s' was empty.  This sprite cannot be imported."), *NameForErrors);
-		return NULL;
+		return nullptr;
 	}
 }
 
@@ -49,7 +49,7 @@ TSharedPtr<FJsonObject> ParseJSON(FArchive* const Stream, const FString& NameFor
 	else
 	{
 		UE_LOG(LogPaperJsonImporter, Warning, TEXT("Failed to parse sprite descriptor file '%s'.  Error: '%s'"), *NameForErrors, *Reader->GetErrorMessage());
-		return NULL;
+		return nullptr;
 	}
 }
 
@@ -335,7 +335,7 @@ void FPaperJsonSpriteSheetImporter::SetReimportData(const FString& ExistingTextu
 			FStringAssetReference SpriteStringRef = SpriteAssetPtr.ToStringReference();
 			if (!SpriteStringRef.ToString().IsEmpty())
 			{
-				UPaperSprite* LoadedSprite = Cast<UPaperSprite>(StaticLoadObject(UPaperSprite::StaticClass(), NULL, *SpriteStringRef.ToString(), NULL, LOAD_None, NULL));
+				UPaperSprite* LoadedSprite = Cast<UPaperSprite>(StaticLoadObject(UPaperSprite::StaticClass(), nullptr, *SpriteStringRef.ToString(), nullptr, LOAD_None, nullptr));
 				if (LoadedSprite != nullptr)
 				{
 					ExistingSprites.Add(ExistingSpriteNames[i], LoadedSprite);
@@ -419,9 +419,9 @@ bool FPaperJsonSpriteSheetImporter::ImportTextures(const FString& LongPackagePat
 		FAssetToolsModule& AssetToolsModule = FModuleManager::GetModuleChecked<FAssetToolsModule>("AssetTools");
 		TArray<UObject*> ImportedSheets = AssetToolsModule.Get().ImportAssets(SheetFileNames, TargetSubPath);
 
-		UTexture2D* ImportedTexture = (ImportedSheets.Num() > 0) ? Cast<UTexture2D>(ImportedSheets[0]) : NULL;
+		UTexture2D* ImportedTexture = (ImportedSheets.Num() > 0) ? Cast<UTexture2D>(ImportedSheets[0]) : nullptr;
 
-		if (ImportedTexture == NULL)
+		if (ImportedTexture == nullptr)
 		{
 			UE_LOG(LogPaperJsonImporter, Warning, TEXT("Failed to import sprite sheet image '%s'."), *SourceSheetImageFilename);
 			bLoadedSuccessfully = false;
@@ -475,7 +475,7 @@ bool FPaperJsonSpriteSheetImporter::PerformImport(const FString& LongPackagePath
 		// Create a package for the frame
 		const FString TargetSubPath = LongPackagePath + TEXT("/Frames");
 
-		UObject* OuterForFrame = NULL; // @TODO: Use this if we don't want them to be individual assets - Flipbook;
+		UObject* OuterForFrame = nullptr; // @TODO: Use this if we don't want them to be individual assets - Flipbook;
 
 		// Create a unique package name and asset name for the frame
 		const FString TentativePackagePath = PackageTools::SanitizePackageName(TargetSubPath + TEXT("/") + Frame.FrameName.ToString());
@@ -485,7 +485,7 @@ bool FPaperJsonSpriteSheetImporter::PerformImport(const FString& LongPackagePath
 		AssetToolsModule.Get().CreateUniqueAssetName(TentativePackagePath, /*out*/ DefaultSuffix, /*out*/ PackageName, /*out*/ AssetName);
 
 		// Create a package for the frame
-		OuterForFrame = CreatePackage(NULL, *PackageName);
+		OuterForFrame = CreatePackage(nullptr, *PackageName);
 
 		// Create a frame in the package
 		UPaperSprite* TargetSprite = nullptr;
diff --git a/Engine/Plugins/2D/Paper2D/Source/PaperJsonImporter/Private/PaperSpriteSheetAssetTypeActions.cpp b/Engine/Plugins/2D/Paper2D/Source/PaperJsonImporter/Private/PaperSpriteSheetAssetTypeActions.cpp
--- a/Engine/Plugins/2D/Paper2D/Source/PaperJsonImporter/Private/PaperSpriteSheetAssetTypeActions.cpp
+++ b/Engine/Plugins/2D/Paper2D/Source/PaperJsonImporter/Private/PaperSpriteSheetAssetTypeActions.cpp
@@ -96,7 +96,7 @@ void FPaperSpriteSheetAssetTypeActions::ExecuteCreateFlipbooks(TArray<TWeakObjec
 				UPaperSprite* Sprite = nullptr;
 				if (!SpriteStringRef.ToString().IsEmpty())
 				{
-					Sprite = Cast<UPaperSprite>(StaticLoadObject(UPaperSprite::StaticClass(), NULL, *SpriteStringRef.ToString(), NULL, LOAD_None, NULL));
+					Sprite = Cast<UPaperSprite>(StaticLoadObject(UPaperSprite::StaticClass(), nullptr, *SpriteStringRef.ToString(), nullptr, LOAD_None, nullptr));
 				}
 				if (Sprite != nullptr)
 				{
diff --git a/Engine/Plugins/2D/Paper2D/Source/PaperJsonImporter/Private/ReimportPaperJsonImporterFactory.cpp b/Engine/Plugins/2D/Paper2D/Source/PaperJsonImporter/Private/ReimportPaperJsonImporterFactory.cpp
--- a/Engine/Plugins/2D/Paper2D/Source/PaperJsonImporter/Private/ReimportPaperJsonImporterFactory.cpp
+++ b/Engine/Plugins/2D/Paper2D/Source/PaperJsonImporter/Private/ReimportPaperJsonImporterFactory.cpp
@@ -57,7 +57,7 @@ EReimportResult::Type UReimportPaperJsonImporterFactory::Reimport(UObject* Obj)
 	ExistingTexture = SpriteSheet->Texture;
 
 	EReimportResult::Type Result = EReimportResult::Failed;
-	if (UFactory::StaticImportObject(SpriteSheet->GetClass(), SpriteSheet->GetOuter(), *SpriteSheet->GetName(), RF_Public | RF_Standalone, *Filename, NULL, this))
+	if (UFactory::StaticImportObject(SpriteSheet->GetClass(), SpriteSheet->GetOuter(), *SpriteSheet->GetName(), RF_Public | RF_Standalone, *Filename, nullptr, this))
 	{
 		UE_LOG(LogPaperJsonImporter, Log, TEXT("Imported successfully"));
 		// Try to find the outer package so we can dirty it up
diff --git a/Engine/Plugins/2D/Paper2D/Source/SmartSnapping/Private/SmartSnappingModule.cpp b/Engine/Plugins/2D/Paper2D/Source/SmartSnapping/Private/SmartSnappingModule.cpp
--- a/Engine/Plugins/2D/Paper2D/Source/SmartSnapping/Private/SmartSnappingModule.cpp
+++ b/Engine/Plugins/2D/Paper2D/Source/SmartSnapping/Private/SmartSnappingModule.cpp
@@ -84,7 +84,7 @@ public:
 		Extender->AddMenuExtension(
 			"Snapping",
 			EExtensionHook::After,
-			NULL,
+			nullptr,
 			FMenuExtensionDelegate::CreateRaw(this, &FSmartSnappingModule::CreateSnappingOptionsMenu));
 	
 		return Extender;
