Paper2D: Rename UPaperTileMapRenderComponent to UPaperTileMapComponent

diff --git a/Engine/Config/BaseEngine.ini b/Engine/Config/BaseEngine.ini
--- a/Engine/Config/BaseEngine.ini
+++ b/Engine/Config/BaseEngine.ini
@@ -550,6 +550,9 @@ bUseStreamingPause=false
 +ActiveClassRedirects=(OldClassName="PaperAnimatedRenderComponent",NewClassName="/Script/Paper2D.PaperFlipbookComponent")
 +ActiveClassRedirects=(OldClassName="PaperRenderActor",NewClassName="/Script/Paper2D.PaperSpriteActor")
 
+; Paper2D modifications (4.7 era)
++ActiveClassRedirects=(OldClassName="PaperTileMapRenderComponent",NewClassName="/Script/Paper2D.PaperTileMapComponent")
+
 ; collision response container refactor
 +TaggedPropertyRedirects=(ClassName="CollisionResponseContainer",OldPropertyName="RigidBody",NewPropertyName="PhysicsBody")
 +TaggedPropertyRedirects=(ClassName="CollisionResponseContainer",OldPropertyName="Static",NewPropertyName="WorldStatic")
diff --git a/Engine/Plugins/2D/Paper2D/Source/Paper2D/Classes/PaperTileMapActor.h b/Engine/Plugins/2D/Paper2D/Source/Paper2D/Classes/PaperTileMapActor.h
--- a/Engine/Plugins/2D/Paper2D/Source/Paper2D/Classes/PaperTileMapActor.h
+++ b/Engine/Plugins/2D/Paper2D/Source/Paper2D/Classes/PaperTileMapActor.h
@@ -10,7 +10,7 @@
  * An instance of a UPaperTileMap in a level.
  *
  * This actor is created when you drag a tile map asset from the content browser into the level, and
- * it is just a thin wrapper around a UPaperTileMapRenderComponent that actually references the asset.
+ * it is just a thin wrapper around a UPaperTileMapComponent that actually references the asset.
  */
 UCLASS()
 class PAPER2D_API APaperTileMapActor : public AActor
@@ -19,7 +19,7 @@ class PAPER2D_API APaperTileMapActor : public AActor
 
 private:
 	UPROPERTY(Category=TileMapActor, VisibleAnywhere, BlueprintReadOnly, meta=(ExposeFunctionCategories="Sprite,Rendering,Physics,Components|Sprite", AllowPrivateAccess="true"))
-	class UPaperTileMapRenderComponent* RenderComponent;
+	class UPaperTileMapComponent* RenderComponent;
 public:
 
 	// AActor interface
@@ -29,5 +29,5 @@ public:
 	// End of AActor interface
 
 	/** Returns RenderComponent subobject **/
-	FORCEINLINE class UPaperTileMapRenderComponent* GetRenderComponent() const { return RenderComponent; }
+	FORCEINLINE class UPaperTileMapComponent* GetRenderComponent() const { return RenderComponent; }
 };
diff --git a/Engine/Plugins/2D/Paper2D/Source/Paper2D/Classes/PaperTileMapComponent.h b/Engine/Plugins/2D/Paper2D/Source/Paper2D/Classes/PaperTileMapComponent.h
new file mode 100644
--- /dev/null
+++ b/Engine/Plugins/2D/Paper2D/Source/Paper2D/Classes/PaperTileMapComponent.h
@@ -0,0 +1,76 @@
+// Copyright 1998-2015 Epic Games, Inc. All Rights Reserved.
+
+#pragma once
+
+#include "PaperSprite.h"
+
+#include "PaperTileMapComponent.generated.h"
+
+UCLASS(hideCategories=Object, ClassGroup=Paper2D, Experimental, meta=(BlueprintSpawnableComponent))
+class PAPER2D_API UPaperTileMapComponent : public UPrimitiveComponent
+{
+	GENERATED_UCLASS_BODY()
+
+private:
+	// Width of map (in tiles)
+	UPROPERTY()
+	int32 MapWidth_DEPRECATED;
+
+	UPROPERTY()
+	int32 MapHeight_DEPRECATED;
+
+	UPROPERTY()
+	int32 TileWidth_DEPRECATED;
+
+	UPROPERTY()
+	int32 TileHeight_DEPRECATED;
+
+	// Default tile set to use for new layers
+	UPROPERTY()
+	UPaperTileSet* DefaultLayerTileSet_DEPRECATED;
+
+	// Test material
+	UPROPERTY()
+	UMaterialInterface* Material_DEPRECATED;
+
+	// The list of layers
+	UPROPERTY()
+	TArray<UPaperTileLayer*> TileLayers_DEPRECATED;
+
+public:
+	// The tile map used by this component
+	UPROPERTY(Category=Setup, EditAnywhere, Instanced, BlueprintReadOnly)
+	class UPaperTileMap* TileMap;
+
+protected:
+	friend class FPaperTileMapRenderSceneProxy;
+
+	void RebuildRenderData(class FPaperTileMapRenderSceneProxy* Proxy);
+
+public:
+	// UObject interface
+	virtual void PostInitProperties() override;
+#if WITH_EDITORONLY_DATA
+	virtual void Serialize(FArchive& Ar) override;
+	virtual void PostLoad() override;
+#endif
+	// End of UObject interface
+
+	// UActorComponent interface
+	virtual void TickComponent(float DeltaTime, enum ELevelTick TickType, FActorComponentTickFunction* ThisTickFunction) override;
+	virtual const UObject* AdditionalStatObject() const override;
+	// End of UActorComponent interface
+
+	// UPrimitiveComponent interface
+	virtual FPrimitiveSceneProxy* CreateSceneProxy() override;
+	virtual FBoxSphereBounds CalcBounds(const FTransform& LocalToWorld) const override;
+	virtual class UBodySetup* GetBodySetup() override;
+	// End of UPrimitiveComponent interface
+
+	// Does this component own the tile map (is it instanced instead of being an asset reference)?
+	bool OwnsTileMap() const;
+};
+
+// Allow the old name to continue to work for one release
+DEPRECATED(4.7, "UPaperTileMapRenderComponent has been renamed to UPaperTileMapComponent")
+typedef UPaperTileMapComponent UPaperTileMapRenderComponent;
diff --git a/Engine/Plugins/2D/Paper2D/Source/Paper2D/Classes/PaperTileMapRenderComponent.h b/Engine/Plugins/2D/Paper2D/Source/Paper2D/Classes/PaperTileMapRenderComponent.h
deleted file mode 100644
diff --git a/Engine/Plugins/2D/Paper2D/Source/Paper2D/Private/PaperTileLayer.cpp b/Engine/Plugins/2D/Paper2D/Source/Paper2D/Private/PaperTileLayer.cpp
--- a/Engine/Plugins/2D/Paper2D/Source/Paper2D/Private/PaperTileLayer.cpp
+++ b/Engine/Plugins/2D/Paper2D/Source/Paper2D/Private/PaperTileLayer.cpp
@@ -76,7 +76,7 @@ public:
 	FTileMapLayerReregisterContext(UPaperTileLayer* TargetAsset)
 	{
 		// Look at tile map components
-		for (TObjectIterator<UPaperTileMapRenderComponent> MapIt; MapIt; ++MapIt)
+		for (TObjectIterator<UPaperTileMapComponent> MapIt; MapIt; ++MapIt)
 		{
 			if (UPaperTileMap* TestMap = (*MapIt)->TileMap)
 			{
diff --git a/Engine/Plugins/2D/Paper2D/Source/Paper2D/Private/PaperTileMap.cpp b/Engine/Plugins/2D/Paper2D/Source/Paper2D/Private/PaperTileMap.cpp
--- a/Engine/Plugins/2D/Paper2D/Source/Paper2D/Private/PaperTileMap.cpp
+++ b/Engine/Plugins/2D/Paper2D/Source/Paper2D/Private/PaperTileMap.cpp
@@ -34,7 +34,7 @@ UPaperTileMap::UPaperTileMap(const FObjectInitializer& ObjectInitializer)
 
 #if WITH_EDITOR
 
-#include "PaperTileMapRenderComponent.h"
+#include "PaperTileMapComponent.h"
 #include "ComponentReregisterContext.h"
 
 void UPaperTileMap::PostEditChangeProperty(FPropertyChangedEvent& PropertyChangedEvent)
@@ -42,7 +42,7 @@ void UPaperTileMap::PostEditChangeProperty(FPropertyChangedEvent& PropertyChange
 	const FName PropertyName = (PropertyChangedEvent.Property != nullptr) ? PropertyChangedEvent.Property->GetFName() : NAME_None;
 
 	//@TODO: Determine when these are really needed, as they're seriously expensive!
-	TComponentReregisterContext<UPaperTileMapRenderComponent> ReregisterStaticComponents;
+	TComponentReregisterContext<UPaperTileMapComponent> ReregisterStaticComponents;
 
 	ValidateSelectedLayerIndex();
 
diff --git a/Engine/Plugins/2D/Paper2D/Source/Paper2D/Private/PaperTileMapActor.cpp b/Engine/Plugins/2D/Paper2D/Source/Paper2D/Private/PaperTileMapActor.cpp
--- a/Engine/Plugins/2D/Paper2D/Source/Paper2D/Private/PaperTileMapActor.cpp
+++ b/Engine/Plugins/2D/Paper2D/Source/Paper2D/Private/PaperTileMapActor.cpp
@@ -8,7 +8,7 @@
 APaperTileMapActor::APaperTileMapActor(const FObjectInitializer& ObjectInitializer)
 	: Super(ObjectInitializer)
 {
-	RenderComponent = ObjectInitializer.CreateDefaultSubobject<UPaperTileMapRenderComponent>(this, TEXT("RenderComponent"));
+	RenderComponent = ObjectInitializer.CreateDefaultSubobject<UPaperTileMapComponent>(this, TEXT("RenderComponent"));
 
 	RootComponent = RenderComponent;
 }
diff --git a/Engine/Plugins/2D/Paper2D/Source/Paper2D/Private/PaperTileMapComponent.cpp b/Engine/Plugins/2D/Paper2D/Source/Paper2D/Private/PaperTileMapComponent.cpp
new file mode 100644
--- /dev/null
+++ b/Engine/Plugins/2D/Paper2D/Source/Paper2D/Private/PaperTileMapComponent.cpp
@@ -0,0 +1,326 @@
+// Copyright 1998-2015 Epic Games, Inc. All Rights Reserved.
+
+#include "Paper2DPrivatePCH.h"
+#include "PaperTileMapRenderSceneProxy.h"
+#include "PhysicsEngine/BodySetup.h"
+#include "PaperCustomVersion.h"
+
+#define LOCTEXT_NAMESPACE "Paper2D"
+
+//////////////////////////////////////////////////////////////////////////
+// UPaperTileMapComponent
+
+UPaperTileMapComponent::UPaperTileMapComponent(const FObjectInitializer& ObjectInitializer)
+	: Super(ObjectInitializer)
+{
+	BodyInstance.SetCollisionEnabled(ECollisionEnabled::QueryAndPhysics);
+
+	MapWidth_DEPRECATED = 4;
+	MapHeight_DEPRECATED = 4;
+	TileWidth_DEPRECATED = 32;
+	TileHeight_DEPRECATED = 32;
+
+	static ConstructorHelpers::FObjectFinder<UMaterial> DefaultMaterial(TEXT("/Paper2D/DefaultSpriteMaterial"));
+	Material_DEPRECATED = DefaultMaterial.Object;
+}
+
+FPrimitiveSceneProxy* UPaperTileMapComponent::CreateSceneProxy()
+{
+	FPaperTileMapRenderSceneProxy* Proxy = new FPaperTileMapRenderSceneProxy(this);
+	RebuildRenderData(Proxy);
+
+	return Proxy;
+}
+
+void UPaperTileMapComponent::PostInitProperties()
+{
+	TileMap = NewObject<UPaperTileMap>(this);
+	TileMap->SetFlags(RF_Transactional);
+	Super::PostInitProperties();
+}
+
+FBoxSphereBounds UPaperTileMapComponent::CalcBounds(const FTransform& LocalToWorld) const
+{
+	if (TileMap != nullptr)
+	{
+		// Graphics bounds.
+		FBoxSphereBounds NewBounds = TileMap->GetRenderBounds().TransformBy(LocalToWorld);
+
+		// Add bounds of collision geometry (if present).
+		if (UBodySetup* BodySetup = TileMap->BodySetup)
+		{
+			const FBox AggGeomBox = BodySetup->AggGeom.CalcAABB(LocalToWorld);
+			if (AggGeomBox.IsValid)
+			{
+				NewBounds = Union(NewBounds, FBoxSphereBounds(AggGeomBox));
+			}
+		}
+
+		// Apply bounds scale
+		NewBounds.BoxExtent *= BoundsScale;
+		NewBounds.SphereRadius *= BoundsScale;
+
+		return NewBounds;
+	}
+	else
+	{
+		return FBoxSphereBounds(LocalToWorld.GetLocation(), FVector::ZeroVector, 0.f);
+	}
+}
+
+void UPaperTileMapComponent::TickComponent(float DeltaTime, enum ELevelTick TickType, FActorComponentTickFunction* ThisTickFunction)
+{
+	// Indicate we need to send new dynamic data.
+	MarkRenderDynamicDataDirty();
+}
+
+#if WITH_EDITORONLY_DATA
+void UPaperTileMapComponent::Serialize(FArchive& Ar)
+{
+	Super::Serialize(Ar);
+
+	Ar.UsingCustomVersion(FPaperCustomVersion::GUID);
+}
+
+void UPaperTileMapComponent::PostLoad()
+{
+	Super::PostLoad();
+
+#if WITH_EDITORONLY_DATA
+	if (GetLinkerCustomVersion(FPaperCustomVersion::GUID) < FPaperCustomVersion::MovedTileMapDataToSeparateClass)
+	{
+		// Create a tile map object and move our old properties over to it
+		TileMap = NewObject<UPaperTileMap>(this);
+		TileMap->SetFlags(RF_Transactional);
+		TileMap->MapWidth = MapWidth_DEPRECATED;
+		TileMap->MapHeight = MapHeight_DEPRECATED;
+		TileMap->TileWidth = TileWidth_DEPRECATED;
+		TileMap->TileHeight = TileHeight_DEPRECATED;
+		TileMap->PixelsPerUnit = 1.0f;
+		TileMap->SelectedTileSet = DefaultLayerTileSet_DEPRECATED;
+		TileMap->Material = Material_DEPRECATED;
+		TileMap->TileLayers = TileLayers_DEPRECATED;
+
+		// Convert the layers
+		for (UPaperTileLayer* Layer : TileMap->TileLayers)
+		{
+			Layer->Rename(nullptr, TileMap, REN_ForceNoResetLoaders | REN_DontCreateRedirectors);
+			Layer->ConvertToTileSetPerCell();
+		}
+
+		// Remove references in the deprecated variables to prevent issues with deleting referenced assets, etc...
+		DefaultLayerTileSet_DEPRECATED = nullptr;
+		Material_DEPRECATED = nullptr;
+		TileLayers_DEPRECATED.Empty();
+	}
+#endif
+}
+#endif
+
+UBodySetup* UPaperTileMapComponent::GetBodySetup()
+{
+	return (TileMap != nullptr) ? TileMap->BodySetup : nullptr;
+}
+
+const UObject* UPaperTileMapComponent::AdditionalStatObject() const
+{
+	if (TileMap != nullptr)
+	{
+		if (TileMap->GetOuter() != this)
+		{
+			return TileMap;
+		}
+	}
+
+	return nullptr;
+}
+
+void UPaperTileMapComponent::RebuildRenderData(FPaperTileMapRenderSceneProxy* Proxy)
+{
+	TArray<FSpriteDrawCallRecord> BatchedSprites;
+	
+	if (TileMap == nullptr)
+	{
+		return;
+	}
+
+
+	FVector WorldOffsetNoY;
+	FVector WorldOffsetYFactor = FVector::ZeroVector;
+	FVector WorldStepX;
+	FVector WorldStepY;
+
+	switch (TileMap->ProjectionMode)
+	{
+	case ETileMapProjectionMode::Orthogonal:
+	default:
+		WorldOffsetNoY = -(TileMap->TileWidth * PaperAxisX * 0.5f) - (TileMap->TileHeight * PaperAxisY * 0.5f);
+		WorldStepX = PaperAxisX * TileMap->TileWidth;
+		WorldStepY = -PaperAxisY * TileMap->TileHeight;
+		break;
+	case ETileMapProjectionMode::IsometricDiamond:
+		WorldOffsetNoY = (TileMap->MapWidth * TileMap->TileWidth * 0.5f * PaperAxisX);
+		WorldOffsetYFactor = -(TileMap->TileHeight * 0.5f * PaperAxisY);
+		WorldStepX = (TileMap->TileWidth * PaperAxisX * 0.5f) - (TileMap->TileHeight * PaperAxisY * 0.5f);
+		WorldStepY = (TileMap->TileWidth * PaperAxisX * -0.5f) - (TileMap->TileHeight * PaperAxisY * 0.5f);
+		break;
+	case ETileMapProjectionMode::IsometricStaggered:
+		WorldOffsetNoY = 0.5f * TileMap->TileHeight * PaperAxisY;
+		WorldOffsetYFactor = 0.5f * TileMap->TileWidth * PaperAxisX;
+		WorldStepX = PaperAxisX * TileMap->TileWidth;
+		WorldStepY = -PaperAxisY * TileMap->TileHeight;
+		break;
+	}
+
+	UTexture2D* LastSourceTexture = nullptr;
+	FVector TileSetOffset = FVector::ZeroVector;
+	FVector2D InverseTextureSize(1.0f, 1.0f);
+	FVector2D SourceDimensionsUV(1.0f, 1.0f);
+	FVector2D TileSizeXY(0.0f, 0.0f);
+
+	for (int32 Z = 0; Z < TileMap->TileLayers.Num(); ++Z)
+	{
+		UPaperTileLayer* Layer = TileMap->TileLayers[Z];
+
+		if (Layer == nullptr)
+		{
+			continue;
+		}
+
+		FLinearColor DrawColor = FLinearColor::White;
+#if WITH_EDITORONLY_DATA
+		if (Layer->bHiddenInEditor)
+		{
+			continue;
+		}
+
+		DrawColor.A = Layer->LayerOpacity;
+#endif
+
+		FSpriteDrawCallRecord* CurrentBatch = nullptr;
+
+		for (int32 Y = 0; Y < TileMap->MapHeight; ++Y)
+		{
+			// In pixels
+			FVector WorldOffset;
+
+			switch (TileMap->ProjectionMode)
+			{
+			case ETileMapProjectionMode::Orthogonal:
+			default:
+				WorldOffset = WorldOffsetNoY;
+				break;
+			case ETileMapProjectionMode::IsometricDiamond:
+				WorldOffset = WorldOffsetNoY + Y * WorldOffsetYFactor;
+				break;
+			case ETileMapProjectionMode::IsometricStaggered:
+				WorldOffset = WorldOffsetNoY + (Y & 1) * WorldOffsetYFactor;
+				break;
+			}
+
+			for (int32 X = 0; X < TileMap->MapWidth; ++X)
+			{
+				const FPaperTileInfo TileInfo = Layer->GetCell(X, Y);
+
+				// do stuff
+				const float TotalSeparation = (TileMap->SeparationPerLayer * Z) + (TileMap->SeparationPerTileX * X) + (TileMap->SeparationPerTileY * Y);
+				FVector WorldPos = (WorldStepX * X) + (WorldStepY * Y) + WorldOffset;
+				WorldPos += TotalSeparation * PaperAxisZ;
+
+				const int32 TileWidth = TileMap->TileWidth;
+				const int32 TileHeight = TileMap->TileHeight;
+
+
+				{
+					UTexture2D* SourceTexture = nullptr;
+
+					FVector2D SourceUV = FVector2D::ZeroVector;
+					if (Layer->bCollisionLayer)
+					{
+						if (TileInfo.PackedTileIndex == 0)
+						{
+							continue;
+						}
+						SourceTexture = UCanvas::StaticClass()->GetDefaultObject<UCanvas>()->DefaultTexture;
+					}
+					else
+					{
+						if (TileInfo.TileSet == nullptr)
+						{
+							continue;
+						}
+
+						if (!TileInfo.TileSet->GetTileUV(TileInfo.PackedTileIndex, /*out*/ SourceUV))
+						{
+							continue;
+						}
+
+						SourceTexture = TileInfo.TileSet->TileSheet;
+						if (SourceTexture == nullptr)
+						{
+							continue;
+						}
+					}
+
+					if ((SourceTexture != LastSourceTexture) || (CurrentBatch == nullptr))
+					{
+						CurrentBatch = (new (BatchedSprites) FSpriteDrawCallRecord());
+						CurrentBatch->Texture = SourceTexture;
+						CurrentBatch->Color = DrawColor;
+						CurrentBatch->Destination = WorldPos.ProjectOnTo(PaperAxisZ);
+					}
+
+					if (SourceTexture != LastSourceTexture)
+					{
+						InverseTextureSize = FVector2D(1.0f / SourceTexture->GetSizeX(), 1.0f / SourceTexture->GetSizeY());
+
+						if (TileInfo.TileSet != nullptr)
+						{
+							SourceDimensionsUV = FVector2D(TileInfo.TileSet->TileWidth * InverseTextureSize.X, TileInfo.TileSet->TileHeight * InverseTextureSize.Y);
+							TileSizeXY = FVector2D(TileInfo.TileSet->TileWidth, TileInfo.TileSet->TileHeight);
+							TileSetOffset = (TileInfo.TileSet->DrawingOffset.X * PaperAxisX) + (TileInfo.TileSet->DrawingOffset.Y * PaperAxisY);
+						}
+						else
+						{
+							SourceDimensionsUV = FVector2D(TileWidth * InverseTextureSize.X, TileHeight * InverseTextureSize.Y);
+							TileSizeXY = FVector2D(TileWidth, TileHeight);
+							TileSetOffset = FVector::ZeroVector;
+						}
+						LastSourceTexture = SourceTexture;
+					}
+					WorldPos += TileSetOffset;
+
+					SourceUV.X *= InverseTextureSize.X;
+					SourceUV.Y *= InverseTextureSize.Y;
+
+					FSpriteDrawCallRecord& NewTile = *CurrentBatch;
+
+					const float WX0 = FVector::DotProduct(WorldPos, PaperAxisX);
+					const float WY0 = FVector::DotProduct(WorldPos, PaperAxisY);
+
+					const FVector4 BottomLeft(WX0, WY0, SourceUV.X, SourceUV.Y + SourceDimensionsUV.Y);
+					const FVector4 BottomRight(WX0 + TileSizeXY.X, WY0, SourceUV.X + SourceDimensionsUV.X, SourceUV.Y + SourceDimensionsUV.Y);
+					const FVector4 TopRight(WX0 + TileSizeXY.X, WY0 + TileSizeXY.Y, SourceUV.X + SourceDimensionsUV.X, SourceUV.Y);
+					const FVector4 TopLeft(WX0, WY0 + TileSizeXY.Y, SourceUV.X, SourceUV.Y);
+
+					new (NewTile.RenderVerts) FVector4(BottomLeft);
+					new (NewTile.RenderVerts) FVector4(TopRight);
+					new (NewTile.RenderVerts) FVector4(BottomRight);
+
+					new (NewTile.RenderVerts) FVector4(BottomLeft);
+					new (NewTile.RenderVerts) FVector4(TopLeft);
+					new (NewTile.RenderVerts) FVector4(TopRight);
+				}
+			}
+		}
+	}
+
+	Proxy->SetBatchesHack(BatchedSprites);
+}
+
+bool UPaperTileMapComponent::OwnsTileMap() const
+{
+	return (TileMap != nullptr) && (TileMap->GetOuter() == this);
+}
+
+#undef LOCTEXT_NAMESPACE
\ No newline at end of file
diff --git a/Engine/Plugins/2D/Paper2D/Source/Paper2D/Private/PaperTileMapRenderComponent.cpp b/Engine/Plugins/2D/Paper2D/Source/Paper2D/Private/PaperTileMapRenderComponent.cpp
deleted file mode 100644
diff --git a/Engine/Plugins/2D/Paper2D/Source/Paper2D/Private/PaperTileMapRenderSceneProxy.cpp b/Engine/Plugins/2D/Paper2D/Source/Paper2D/Private/PaperTileMapRenderSceneProxy.cpp
--- a/Engine/Plugins/2D/Paper2D/Source/Paper2D/Private/PaperTileMapRenderSceneProxy.cpp
+++ b/Engine/Plugins/2D/Paper2D/Source/Paper2D/Private/PaperTileMapRenderSceneProxy.cpp
@@ -8,11 +8,11 @@
 //////////////////////////////////////////////////////////////////////////
 // FPaperTileMapRenderSceneProxy
 
-FPaperTileMapRenderSceneProxy::FPaperTileMapRenderSceneProxy(const UPaperTileMapRenderComponent* InComponent)
+FPaperTileMapRenderSceneProxy::FPaperTileMapRenderSceneProxy(const UPaperTileMapComponent* InComponent)
 	: FPaperRenderSceneProxy(InComponent)
 	, TileMap(nullptr)
 {
-	if (const UPaperTileMapRenderComponent* InTileComponent = Cast<const UPaperTileMapRenderComponent>(InComponent))
+	if (const UPaperTileMapComponent* InTileComponent = Cast<const UPaperTileMapComponent>(InComponent))
 	{
 		TileMap = InTileComponent->TileMap;
 		Material = (TileMap != nullptr) ? TileMap->Material : nullptr;
diff --git a/Engine/Plugins/2D/Paper2D/Source/Paper2D/Private/PaperTileMapRenderSceneProxy.h b/Engine/Plugins/2D/Paper2D/Source/Paper2D/Private/PaperTileMapRenderSceneProxy.h
--- a/Engine/Plugins/2D/Paper2D/Source/Paper2D/Private/PaperTileMapRenderSceneProxy.h
+++ b/Engine/Plugins/2D/Paper2D/Source/Paper2D/Private/PaperTileMapRenderSceneProxy.h
@@ -10,7 +10,7 @@
 class FPaperTileMapRenderSceneProxy : public FPaperRenderSceneProxy
 {
 public:
-	FPaperTileMapRenderSceneProxy(const UPaperTileMapRenderComponent* InComponent);
+	FPaperTileMapRenderSceneProxy(const UPaperTileMapComponent* InComponent);
 
 	// FPrimitiveSceneProxy interface.
 	virtual void GetDynamicMeshElements(const TArray<const FSceneView*>& Views, const FSceneViewFamily& ViewFamily, uint32 VisibilityMap, FMeshElementCollector& Collector) const override;
diff --git a/Engine/Plugins/2D/Paper2D/Source/Paper2D/Private/PaperTileSet.cpp b/Engine/Plugins/2D/Paper2D/Source/Paper2D/Private/PaperTileSet.cpp
--- a/Engine/Plugins/2D/Paper2D/Source/Paper2D/Private/PaperTileSet.cpp
+++ b/Engine/Plugins/2D/Paper2D/Source/Paper2D/Private/PaperTileSet.cpp
@@ -97,7 +97,7 @@ bool UPaperTileSet::GetTileUV(int32 TileIndex, /*out*/ FVector2D& Out_TileUV) co
 
 #if WITH_EDITOR
 
-#include "PaperTileMapRenderComponent.h"
+#include "PaperTileMapComponent.h"
 #include "ComponentReregisterContext.h"
 
 void UPaperTileSet::PostEditChangeProperty(FPropertyChangedEvent& PropertyChangedEvent)
diff --git a/Engine/Plugins/2D/Paper2D/Source/Paper2DEditor/Private/Paper2DEditorModule.cpp b/Engine/Plugins/2D/Paper2D/Source/Paper2DEditor/Private/Paper2DEditorModule.cpp
--- a/Engine/Plugins/2D/Paper2D/Source/Paper2DEditor/Private/Paper2DEditorModule.cpp
+++ b/Engine/Plugins/2D/Paper2D/Source/Paper2DEditor/Private/Paper2DEditorModule.cpp
@@ -106,12 +106,12 @@ public:
 		FComponentAssetBrokerage::RegisterBroker(PaperFlipbookBroker, UPaperFlipbookComponent::StaticClass(), true, true);
 
 		PaperTileMapBroker = MakeShareable(new FPaperTileMapAssetBroker);
-		FComponentAssetBrokerage::RegisterBroker(PaperTileMapBroker, UPaperTileMapRenderComponent::StaticClass(), true, true);
+		FComponentAssetBrokerage::RegisterBroker(PaperTileMapBroker, UPaperTileMapComponent::StaticClass(), true, true);
 
 		// Register the details customizations
 		{
 			FPropertyEditorModule& PropertyModule = FModuleManager::LoadModuleChecked<FPropertyEditorModule>("PropertyEditor");
-			PropertyModule.RegisterCustomClassLayout(UPaperTileMapRenderComponent::StaticClass()->GetFName(), FOnGetDetailCustomizationInstance::CreateStatic(&FPaperTileMapDetailsCustomization::MakeInstance));
+			PropertyModule.RegisterCustomClassLayout(UPaperTileMapComponent::StaticClass()->GetFName(), FOnGetDetailCustomizationInstance::CreateStatic(&FPaperTileMapDetailsCustomization::MakeInstance));
 			PropertyModule.RegisterCustomClassLayout(UPaperTileMap::StaticClass()->GetFName(), FOnGetDetailCustomizationInstance::CreateStatic(&FPaperTileMapDetailsCustomization::MakeInstance));
 			PropertyModule.RegisterCustomClassLayout(UPaperSprite::StaticClass()->GetFName(), FOnGetDetailCustomizationInstance::CreateStatic(&FSpriteDetailsCustomization::MakeInstance));
 
diff --git a/Engine/Plugins/2D/Paper2D/Source/Paper2DEditor/Private/TileMapEditing/EdModeTileMap.cpp b/Engine/Plugins/2D/Paper2D/Source/Paper2DEditor/Private/TileMapEditing/EdModeTileMap.cpp
--- a/Engine/Plugins/2D/Paper2D/Source/Paper2DEditor/Private/TileMapEditing/EdModeTileMap.cpp
+++ b/Engine/Plugins/2D/Paper2D/Source/Paper2DEditor/Private/TileMapEditing/EdModeTileMap.cpp
@@ -314,16 +314,16 @@ bool FEdModeTileMap::UsesTransformWidget() const
 	return false;
 }
 
-UPaperTileMapRenderComponent* FEdModeTileMap::FindSelectedComponent() const
+UPaperTileMapComponent* FEdModeTileMap::FindSelectedComponent() const
 {
-	UPaperTileMapRenderComponent* TileMapComponent = nullptr;
+	UPaperTileMapComponent* TileMapComponent = nullptr;
 
 	USelection* SelectedActors = Owner->GetSelectedActors();
 	for (FSelectionIterator Iter(*SelectedActors); Iter; ++Iter)
 	{
 		AActor* Actor = CastChecked<AActor>(*Iter);
 
-		TileMapComponent = Actor->FindComponentByClass<UPaperTileMapRenderComponent>();
+		TileMapComponent = Actor->FindComponentByClass<UPaperTileMapComponent>();
 		if (TileMapComponent != nullptr)
 		{
 			break;
@@ -336,7 +336,7 @@ UPaperTileMapRenderComponent* FEdModeTileMap::FindSelectedComponent() const
 		for (FSelectionIterator Iter(*SelectedObjects); Iter; ++Iter)
 		{
 			UObject* Foo = *Iter;
-			TileMapComponent = Cast<UPaperTileMapRenderComponent>(Foo);
+			TileMapComponent = Cast<UPaperTileMapComponent>(Foo);
 
 			if (TileMapComponent != nullptr)
 			{
@@ -355,7 +355,7 @@ UPaperTileLayer* FEdModeTileMap::GetSelectedLayerUnderCursor(const FViewportCurs
 
 	const bool bCollisionPainting = (GetActiveLayerPaintingMode() == ETileMapLayerPaintingMode::CollisionLayers);
 
-	UPaperTileMapRenderComponent* TileMapComponent = FindSelectedComponent();
+	UPaperTileMapComponent* TileMapComponent = FindSelectedComponent();
 
 	if (TileMapComponent != nullptr)
 	{
diff --git a/Engine/Plugins/2D/Paper2D/Source/Paper2DEditor/Private/TileMapEditing/EdModeTileMap.h b/Engine/Plugins/2D/Paper2D/Source/Paper2DEditor/Private/TileMapEditing/EdModeTileMap.h
--- a/Engine/Plugins/2D/Paper2D/Source/Paper2DEditor/Private/TileMapEditing/EdModeTileMap.h
+++ b/Engine/Plugins/2D/Paper2D/Source/Paper2DEditor/Private/TileMapEditing/EdModeTileMap.h
@@ -94,7 +94,7 @@ protected:
 	bool IsTileMapEditModeActive() const;
 
 public:
-	UPaperTileMapRenderComponent* FindSelectedComponent() const;
+	UPaperTileMapComponent* FindSelectedComponent() const;
 
 protected:
 	bool bIsPainting;
diff --git a/Engine/Plugins/2D/Paper2D/Source/Paper2DEditor/Private/TileMapEditing/PaperTileMapAssetBroker.h b/Engine/Plugins/2D/Paper2D/Source/Paper2DEditor/Private/TileMapEditing/PaperTileMapAssetBroker.h
--- a/Engine/Plugins/2D/Paper2D/Source/Paper2DEditor/Private/TileMapEditing/PaperTileMapAssetBroker.h
+++ b/Engine/Plugins/2D/Paper2D/Source/Paper2DEditor/Private/TileMapEditing/PaperTileMapAssetBroker.h
@@ -17,7 +17,7 @@ public:
 
 	virtual bool AssignAssetToComponent(UActorComponent* InComponent, UObject* InAsset) override
 	{
-		if (UPaperTileMapRenderComponent* RenderComp = Cast<UPaperTileMapRenderComponent>(InComponent))
+		if (UPaperTileMapComponent* RenderComp = Cast<UPaperTileMapComponent>(InComponent))
 		{
 			UPaperTileMap* TileMap = Cast<UPaperTileMap>(InAsset);
 
@@ -33,7 +33,7 @@ public:
 
 	virtual UObject* GetAssetFromComponent(UActorComponent* InComponent) override
 	{
-		if (UPaperTileMapRenderComponent* RenderComp = Cast<UPaperTileMapRenderComponent>(InComponent))
+		if (UPaperTileMapComponent* RenderComp = Cast<UPaperTileMapComponent>(InComponent))
 		{
 			if ((RenderComp->TileMap != nullptr) && (RenderComp->TileMap->IsAsset()))
 			{
diff --git a/Engine/Plugins/2D/Paper2D/Source/Paper2DEditor/Private/TileMapEditing/PaperTileMapDetailsCustomization.cpp b/Engine/Plugins/2D/Paper2D/Source/Paper2DEditor/Private/TileMapEditing/PaperTileMapDetailsCustomization.cpp
--- a/Engine/Plugins/2D/Paper2D/Source/Paper2DEditor/Private/TileMapEditing/PaperTileMapDetailsCustomization.cpp
+++ b/Engine/Plugins/2D/Paper2D/Source/Paper2DEditor/Private/TileMapEditing/PaperTileMapDetailsCustomization.cpp
@@ -35,14 +35,14 @@ void FPaperTileMapDetailsCustomization::CustomizeDetails(IDetailLayoutBuilder& D
 		UObject* TestObject = SelectedObjects[ObjectIndex].Get();
 		if (AActor* CurrentActor = Cast<AActor>(TestObject))
 		{
-			if (UPaperTileMapRenderComponent* CurrentComponent = CurrentActor->FindComponentByClass<UPaperTileMapRenderComponent>())
+			if (UPaperTileMapComponent* CurrentComponent = CurrentActor->FindComponentByClass<UPaperTileMapComponent>())
 			{
 				bEditingActor = true;
 				TileMap = CurrentComponent->TileMap;
 				break;
 			}
 		}
-		else if (UPaperTileMapRenderComponent* TestComponent = Cast<UPaperTileMapRenderComponent>(TestObject))
+		else if (UPaperTileMapComponent* TestComponent = Cast<UPaperTileMapComponent>(TestObject))
 		{
 			TileMap = TestComponent->TileMap;
 			break;
diff --git a/Engine/Plugins/2D/Paper2D/Source/Paper2DEditor/Private/TileMapEditing/STileLayerList.cpp b/Engine/Plugins/2D/Paper2D/Source/Paper2DEditor/Private/TileMapEditing/STileLayerList.cpp
--- a/Engine/Plugins/2D/Paper2D/Source/Paper2DEditor/Private/TileMapEditing/STileLayerList.cpp
+++ b/Engine/Plugins/2D/Paper2D/Source/Paper2DEditor/Private/TileMapEditing/STileLayerList.cpp
@@ -419,7 +419,7 @@ void STileLayerList::PostEditNotfications()
 
 	if (NotifyHook != nullptr)
 	{
-		UProperty* TileMapProperty = FindFieldChecked<UProperty>(UPaperTileMapRenderComponent::StaticClass(), GET_MEMBER_NAME_CHECKED(UPaperTileMapRenderComponent, TileMap));
+		UProperty* TileMapProperty = FindFieldChecked<UProperty>(UPaperTileMapComponent::StaticClass(), GET_MEMBER_NAME_CHECKED(UPaperTileMapComponent, TileMap));
 		NotifyHook->NotifyPreChange(TileMapProperty);
 		NotifyHook->NotifyPostChange(FPropertyChangedEvent(TileMapProperty), TileMapProperty);
 	}
diff --git a/Engine/Plugins/2D/Paper2D/Source/Paper2DEditor/Private/TileMapEditing/TileMapActorFactory.cpp b/Engine/Plugins/2D/Paper2D/Source/Paper2DEditor/Private/TileMapEditing/TileMapActorFactory.cpp
--- a/Engine/Plugins/2D/Paper2D/Source/Paper2DEditor/Private/TileMapEditing/TileMapActorFactory.cpp
+++ b/Engine/Plugins/2D/Paper2D/Source/Paper2DEditor/Private/TileMapEditing/TileMapActorFactory.cpp
@@ -16,7 +16,7 @@ UTileMapActorFactory::UTileMapActorFactory(const FObjectInitializer& ObjectIniti
 void UTileMapActorFactory::PostSpawnActor(UObject* Asset, AActor* NewActor)
 {
 	APaperTileMapActor* TypedActor = CastChecked<APaperTileMapActor>(NewActor);
-	UPaperTileMapRenderComponent* RenderComponent = TypedActor->GetRenderComponent();
+	UPaperTileMapComponent* RenderComponent = TypedActor->GetRenderComponent();
 	check(RenderComponent);
 
 	if (UPaperTileMap* TileMap = Cast<UPaperTileMap>(Asset))
@@ -50,7 +50,7 @@ void UTileMapActorFactory::PostCreateBlueprint(UObject* Asset, AActor* CDO)
 {
 	if (APaperTileMapActor* TypedActor = Cast<APaperTileMapActor>(CDO))
 	{
-		UPaperTileMapRenderComponent* RenderComponent = TypedActor->GetRenderComponent();
+		UPaperTileMapComponent* RenderComponent = TypedActor->GetRenderComponent();
 		check(RenderComponent);
 
 		if (UPaperTileMap* TileMap = Cast<UPaperTileMap>(Asset))
diff --git a/Engine/Plugins/2D/Paper2D/Source/Paper2DEditor/Private/TileMapEditing/TileMapEdModeToolkit.cpp b/Engine/Plugins/2D/Paper2D/Source/Paper2DEditor/Private/TileMapEditing/TileMapEdModeToolkit.cpp
--- a/Engine/Plugins/2D/Paper2D/Source/Paper2DEditor/Private/TileMapEditing/TileMapEdModeToolkit.cpp
+++ b/Engine/Plugins/2D/Paper2D/Source/Paper2DEditor/Private/TileMapEditing/TileMapEdModeToolkit.cpp
@@ -67,7 +67,7 @@ void FTileMapEdModeToolkit::Init(const TSharedPtr<IToolkitHost>& InitToolkitHost
 	BindCommands();
 
 	// Try to determine a good default tile set based on the current selection set
-	if (UPaperTileMapRenderComponent* TargetComponent = TileMapEditor->FindSelectedComponent())
+	if (UPaperTileMapComponent* TargetComponent = TileMapEditor->FindSelectedComponent())
 	{
 		if (TargetComponent->TileMap != nullptr)
 		{
@@ -157,7 +157,7 @@ void FTileMapEdModeToolkit::OnChangeTileSet(UObject* NewAsset)
 			}
 
 			// Save the newly selected tile set in the asset so it can be restored next time we edit it
-			if (UPaperTileMapRenderComponent* TargetComponent = TileMapEditor->FindSelectedComponent())
+			if (UPaperTileMapComponent* TargetComponent = TileMapEditor->FindSelectedComponent())
 			{
 				if (TargetComponent->TileMap != nullptr)
 				{
diff --git a/Engine/Plugins/2D/Paper2D/Source/Paper2DEditor/Private/TileMapEditing/TileMapEditorViewportClient.cpp b/Engine/Plugins/2D/Paper2D/Source/Paper2DEditor/Private/TileMapEditing/TileMapEditorViewportClient.cpp
--- a/Engine/Plugins/2D/Paper2D/Source/Paper2DEditor/Private/TileMapEditing/TileMapEditorViewportClient.cpp
+++ b/Engine/Plugins/2D/Paper2D/Source/Paper2DEditor/Private/TileMapEditing/TileMapEditorViewportClient.cpp
@@ -42,7 +42,7 @@ FTileMapEditorViewportClient::FTileMapEditorViewportClient(TWeakPtr<FTileMapEdit
 
 	// Create a render component for the tile map being edited
 	{
-		RenderTileMapComponent = NewObject<UPaperTileMapRenderComponent>();
+		RenderTileMapComponent = NewObject<UPaperTileMapComponent>();
 		UPaperTileMap* TileMap = GetTileMapBeingEdited();
 		RenderTileMapComponent->TileMap = TileMap;
 		GSelectedAnnotation.Set(RenderTileMapComponent);
@@ -131,7 +131,7 @@ void FTileMapEditorViewportClient::Tick(float DeltaSeconds)
 		FIntPoint Size = Viewport->GetSizeXY();
 		if (bDeferZoomToTileMap && (Size.X > 0) && (Size.Y > 0))
 		{
-			UPaperTileMapRenderComponent* ComponentToFocusOn = RenderTileMapComponent;
+			UPaperTileMapComponent* ComponentToFocusOn = RenderTileMapComponent;
 			FocusViewportOnBox(ComponentToFocusOn->Bounds.GetBox(), true);
 			bDeferZoomToTileMap = false;
 		}		
diff --git a/Engine/Plugins/2D/Paper2D/Source/Paper2DEditor/Private/TileMapEditing/TileMapEditorViewportClient.h b/Engine/Plugins/2D/Paper2D/Source/Paper2DEditor/Private/TileMapEditing/TileMapEditorViewportClient.h
--- a/Engine/Plugins/2D/Paper2D/Source/Paper2DEditor/Private/TileMapEditing/TileMapEditorViewportClient.h
+++ b/Engine/Plugins/2D/Paper2D/Source/Paper2DEditor/Private/TileMapEditing/TileMapEditorViewportClient.h
@@ -51,7 +51,7 @@ private:
 	TWeakPtr<FTileMapEditor> TileMapEditorPtr;
 
 	// Render component for the tile map being edited
-	UPaperTileMapRenderComponent* RenderTileMapComponent;
+	UPaperTileMapComponent* RenderTileMapComponent;
 
 	// Widget mode
 	FWidget::EWidgetMode WidgetMode;
