Fix deprecated warning building RootEditor.

diff --git a/Engine/Plugins/2D/Paper2D/Source/Paper2D/Private/PaperTileSet.cpp b/Engine/Plugins/2D/Paper2D/Source/Paper2D/Private/PaperTileSet.cpp
--- a/Engine/Plugins/2D/Paper2D/Source/Paper2D/Private/PaperTileSet.cpp
+++ b/Engine/Plugins/2D/Paper2D/Source/Paper2D/Private/PaperTileSet.cpp
@@ -111,7 +111,7 @@ void UPaperTileSet::PostEditChangeProperty(FPropertyChangedEvent& PropertyChange
 	TileHeight = FMath::Max<int32>(TileHeight, 1);
 
 	//@TODO: Determine when these are really needed, as they're seriously expensive!
-	TComponentReregisterContext<UPaperTileMapRenderComponent> ReregisterStaticComponents;
+	TComponentReregisterContext<UPaperTileMapComponent> ReregisterStaticComponents;
 
 	Super::PostEditChangeProperty(PropertyChangedEvent);
 }
