Paper2D: Add accessor methods for trim state on UPaperSprite (GetOriginInSourceImageBeforeTrimming, GetSourceImageDimensionBeforeTrimming, and IsTrimmedInSourceImage)

diff --git a/Engine/Plugins/2D/Paper2D/Source/Paper2D/Classes/PaperSprite.h b/Engine/Plugins/2D/Paper2D/Source/Paper2D/Classes/PaperSprite.h
--- a/Engine/Plugins/2D/Paper2D/Source/Paper2D/Classes/PaperSprite.h
+++ b/Engine/Plugins/2D/Paper2D/Source/Paper2D/Classes/PaperSprite.h
@@ -177,7 +177,6 @@ public:
 
 	// World space WRT the sprite editor *only*
 	FVector ConvertTextureSpaceToWorldSpace(const FVector2D& SourcePoint) const;
-	//FVector ConvertTextureSpaceToWorldSpace(const FVector& SourcePoint) const;
 	FTransform GetPivotToWorld() const;
 
 	// Returns the raw pivot position (ignoring pixel snapping)
@@ -214,7 +213,20 @@ public:
 	void SetRotated(bool bRotated);
 	void SetPivotMode(ESpritePivotMode::Type PivotMode, FVector2D CustomTextureSpacePivot);
 	
+	// Returns the Origin within SourceImage, prior to atlasing
+	FVector2D GetOriginInSourceImageBeforeTrimming() const { return OriginInSourceImageBeforeTrimming; }
+
+	// Returns the Dimensions of SourceImage prior to trimming
+	FVector2D GetSourceImageDimensionBeforeTrimming() const { return SourceImageDimensionBeforeTrimming; }
+
+	// Returns true if this sprite is trimmed from the original texture, meaning that the source image
+	// dimensions and origin in the source image may not be the same as the final results for the sprite
+	// (empty alpha=0 pixels were trimmed from the exterior region)
+	bool IsTrimmedInSourceImage() const { return bTrimmedInSourceImage; }
+
+	// This texture is rotated in the atlas
 	bool IsRotatedInSourceImage() const { return bRotatedInSourceImage; }
+
 	ESpritePivotMode::Type GetPivotMode(FVector2D& OutCustomTextureSpacePivot) const { OutCustomTextureSpacePivot = CustomPivotPoint; return PivotMode; }
 
 
