Paper2D: Improved terrain spline collision generation
Paper2D: Added separate filled option to closed - might want to be able to generate a closed spline (eg. circle) that is not filled in the future

diff --git a/Engine/Plugins/2D/Paper2D/Source/Paper2D/Classes/PaperTerrainComponent.h b/Engine/Plugins/2D/Paper2D/Source/Paper2D/Classes/PaperTerrainComponent.h
--- a/Engine/Plugins/2D/Paper2D/Source/Paper2D/Classes/PaperTerrainComponent.h
+++ b/Engine/Plugins/2D/Paper2D/Source/Paper2D/Classes/PaperTerrainComponent.h
@@ -5,12 +5,36 @@
 #include "SpriteDrawCall.h"
 #include "PaperTerrainComponent.generated.h"
 
+struct FPaperTerrainMaterialRule;
+
 struct FPaperTerrainMaterialPair
 {
 	TArray<FSpriteDrawCallRecord> Records;
 	UMaterialInterface* Material;
 };
 
+struct FTerrainSpriteStamp
+{
+	const UPaperSprite* Sprite;
+	float NominalWidth;
+	float Time;
+	float Scale;
+	bool bCanStretch;
+
+	FTerrainSpriteStamp(const UPaperSprite* InSprite, float InTime, bool bIsEndCap);
+};
+
+struct FTerrainSegment
+{
+	float StartTime;
+	float EndTime;
+	const FPaperTerrainMaterialRule* Rule;
+	TArray<FTerrainSpriteStamp> Stamps;
+
+	FTerrainSegment();
+	void RepositionStampsToFillSpace();
+};
+
 /**
  * The terrain visualization component for an associated spline component.
  * This takes a 2D terrain material and instances sprite geometry along the spline path.
@@ -29,6 +53,9 @@ public:
 	UPROPERTY(Category = Sprite, EditAnywhere, BlueprintReadOnly)
 	bool bClosedSpline;
 
+	UPROPERTY(Category = Sprite, EditAnywhere, BlueprintReadOnly)
+	bool bFilledSpline;
+
 	UPROPERTY()
 	class UPaperTerrainSplineComponent* AssociatedSpline;
 
@@ -77,8 +104,11 @@ public:
 	void SetTerrainColor(FLinearColor NewColor);
 
 protected:
-	void SpawnSegment(const class UPaperSprite* NewSprite, float Position, float HorizontalScale, float NominalWidth, const struct FPaperTerrainMaterialRule* Rule);
-	void SpawnFromPoly(const class UPaperSprite* NewSprite, FSpriteDrawCallRecord& FillDrawCall, const FVector2D& TextureSize, const TArray<FVector2D>& SplinePolyVertices2D);
+	void SpawnSegments(const TArray<FTerrainSegment>& TerrainSegments, bool bGenerateSegmentColliders);
+	
+	void GenerateFillRenderDataFromPolygon(const class UPaperSprite* NewSprite, FSpriteDrawCallRecord& FillDrawCall, const FVector2D& TextureSize, const TArray<FVector2D>& TriangulatedPolygonVertices);
+	void GenerateCollisionDataFromPolygon(const TArray<FVector2D>& SplinePolyVertices2D, const TArray<float>& TerrainOffsets, const TArray<FVector2D>& TriangulatedPolygonVertices);
+	void InsertConvexCollisionDataFromPolygon(const TArray<FVector2D>& ClosedPolyVertices2D);
 
 	void OnSplineEdited();
 
diff --git a/Engine/Plugins/2D/Paper2D/Source/Paper2D/Private/PaperGeomTools.cpp b/Engine/Plugins/2D/Paper2D/Source/Paper2D/Private/PaperGeomTools.cpp
--- a/Engine/Plugins/2D/Paper2D/Source/Paper2D/Private/PaperGeomTools.cpp
+++ b/Engine/Plugins/2D/Paper2D/Source/Paper2D/Private/PaperGeomTools.cpp
@@ -16,6 +16,34 @@ bool PaperGeomTools::IsPolygonWindingCCW(const TArray<FVector2D>& Points)
 	return (Sum < 0.0f);
 }
 
+// Note: We need to simplify non-simple polygons before this
+static bool IsPolygonConvex(const TArray<FVector2D>& Points)
+{
+	const int PointCount = Points.Num();
+	float Sign = 0;
+	for (int32 PointIndex = 0; PointIndex < PointCount; ++PointIndex)
+	{
+		const FVector2D& A = Points[PointIndex];
+		const FVector2D& B = Points[(PointIndex + 1) % PointCount];
+		const FVector2D& C = Points[(PointIndex + 2) % PointCount];
+		float Det = (B.X - A.X) * (C.Y - B.Y) - (B.Y - A.Y) * (C.X - B.X);
+		float DetSign = FMath::Sign(Det);
+		if (DetSign != 0)
+		{
+			if (Sign == 0)
+			{
+				Sign = DetSign;
+			}
+			else if (Sign != DetSign)
+			{
+				return false;
+			}
+		}
+	}
+
+	return true;
+}
+
 static bool IsPointInPolygon(const FVector2D& TestPoint, const TArray<FVector2D>& PolygonPoints)
 {
 	const int NumPoints = PolygonPoints.Num();
@@ -187,9 +215,10 @@ static void JoinMutuallyVisible(TArray<FVector2D>& AdditivePoly, const TArray<FV
 
 /** Util that tries to combine two triangles if possible. */
 // Unlike the implementation in GeomTools, we only care about 2D points and not any other vertex attributes
+// If bConvex == true, the triangles are only merged if the resulting polygon is convex
 static bool MergeTriangleIntoPolygon(
 	TArray<FVector2D>& PolygonVertices,
-	const FVector2D& TriangleVertexA, const FVector2D& TriangleVertexB, const FVector2D& TriangleVertexC)
+	const FVector2D& TriangleVertexA, const FVector2D& TriangleVertexB, const FVector2D& TriangleVertexC, bool bConvex = false)
 {
 	// Create indexable copy
 	FVector2D TriangleVertices[3] = { TriangleVertexA, TriangleVertexB, TriangleVertexC };
@@ -209,11 +238,23 @@ static bool MergeTriangleIntoPolygon(
 			if (PolygonVertices[PolygonEdgeVertex0].Equals(TriangleVertices[TriangleEdgeVertex1], THRESH_POINTS_ARE_SAME) &&
 				PolygonVertices[PolygonEdgeVertex1].Equals(TriangleVertices[TriangleEdgeVertex0], THRESH_POINTS_ARE_SAME))
 			{
-				// Add the triangle's vertex that isn't in the adjacent edge to the polygon in between the vertices of the adjacent edge.
-				const int32 TriangleOppositeVertexIndex = (TriangleEdgeIndex + 2) % 3;
-				PolygonVertices.Insert(TriangleVertices[TriangleOppositeVertexIndex], PolygonEdgeVertex1);
+				bool bMergeTriangle = true;
+				if (bConvex)
+				{
+					TArray<FVector2D> TmpPolygonVertcies = PolygonVertices;
+					const int32 TriangleOppositeVertexIndex = (TriangleEdgeIndex + 2) % 3;
+					TmpPolygonVertcies.Insert(TriangleVertices[TriangleOppositeVertexIndex], PolygonEdgeVertex1);
+					bMergeTriangle = IsPolygonConvex(TmpPolygonVertcies);
+				}
 
-				return true;
+				if (bMergeTriangle)
+				{
+					// Add the triangle's vertex that isn't in the adjacent edge to the polygon in between the vertices of the adjacent edge.
+					const int32 TriangleOppositeVertexIndex = (TriangleEdgeIndex + 2) % 3;
+					PolygonVertices.Insert(TriangleVertices[TriangleOppositeVertexIndex], PolygonEdgeVertex1);
+
+					return true;
+				}
 			}
 		}
 	}
@@ -601,3 +642,145 @@ void PaperGeomTools::RemoveRedundantTriangles(TArray<FVector2D>& OutTriangles, c
 		OutTriangles.Append(TriangulatedPoly);
 	}
 }
+
+// Find convex polygons from triangle soup
+void PaperGeomTools::GenerateConvexPolygonsFromTriangles(TArray<TArray<FVector2D>>& OutPolygons, const TArray<FVector2D>& InTriangleVertices)
+{
+	// Test: Return input triangles in exactly the same format
+	//for (int32 TriangleVertexIndex = 0; TriangleVertexIndex < InTriangleVertices.Num(); TriangleVertexIndex += 3)
+	//{
+	//	TArray<FVector2D>& Poly = *(new(OutPolygons)TArray<FVector2D>());
+	//	Poly.Add(InTriangleVertices[TriangleVertexIndex + 0]);
+	//	Poly.Add(InTriangleVertices[TriangleVertexIndex + 1]);
+	//	Poly.Add(InTriangleVertices[TriangleVertexIndex + 2]);
+	//}
+	//return;
+
+	struct FLocalTriangle
+	{
+		int VertexA, VertexB, VertexC;
+	};
+
+	TArray<FLocalTriangle> Triangles;
+	for (int32 TriangleVertexIndex = 0; TriangleVertexIndex < InTriangleVertices.Num(); TriangleVertexIndex += 3)
+	{
+		FLocalTriangle* NewTriangle = new(Triangles)FLocalTriangle();
+		NewTriangle->VertexA = TriangleVertexIndex + 0;
+		NewTriangle->VertexB = TriangleVertexIndex + 1;
+		NewTriangle->VertexC = TriangleVertexIndex + 2;
+	}
+
+	while (Triangles.Num() > 0)
+	{
+		TArray<FVector2D> PolygonVertices;
+
+		const FLocalTriangle InitialTriangle = Triangles.Pop();
+		PolygonVertices.Add(InTriangleVertices[InitialTriangle.VertexA]);
+		PolygonVertices.Add(InTriangleVertices[InitialTriangle.VertexB]);
+		PolygonVertices.Add(InTriangleVertices[InitialTriangle.VertexC]);
+
+		// Find triangles that can be merged into the polygon.
+		for (int32 CandidateTriangleIndex = 0; CandidateTriangleIndex < Triangles.Num(); ++CandidateTriangleIndex)
+		{
+			const FLocalTriangle& MergeCandidateTriangle = Triangles[CandidateTriangleIndex];
+			if (MergeTriangleIntoPolygon(PolygonVertices, InTriangleVertices[MergeCandidateTriangle.VertexA], InTriangleVertices[MergeCandidateTriangle.VertexB], InTriangleVertices[MergeCandidateTriangle.VertexC], true))
+			{
+				// Remove the merged triangle from the array.
+				Triangles.RemoveAtSwap(CandidateTriangleIndex);
+
+				// Restart the search for mergeable triangles from the start of the array.
+				CandidateTriangleIndex = -1;
+			}
+		}
+
+		OutPolygons.Add(PolygonVertices);
+	}
+}
+
+// Creates a convex hull that contains the given points
+void PaperGeomTools::GenerateConvexHullFromPoints(TArray<FVector2D>& OutConvexHull, TArray<FVector2D>& SourcePoints)
+{
+	int SourcePointsCount = SourcePoints.Num();
+	if (SourcePointsCount < 3)
+	{
+		return;
+	}
+
+	// Find lowest point. If multiple points have the same y, find leftmost
+	int LowestPointIndex = 0;
+	FVector2D& LowestPoint = SourcePoints[0];
+	for (int I = 1; I < SourcePointsCount; ++I)
+	{
+		if ((SourcePoints[I].Y < LowestPoint.Y) || (SourcePoints[I].Y == LowestPoint.Y && SourcePoints[I].X < LowestPoint.X))
+		{
+			LowestPointIndex = I;
+			LowestPoint = SourcePoints[I];
+		}
+	}
+
+	// Get indices sorted by angle they and the lowest point make with the x-axis (excluding the lowest point)
+	FVector2D& P = SourcePoints[LowestPointIndex];
+	TArray<int32> SortedPoints;
+	SortedPoints.Empty(SourcePointsCount);
+	for (int I = 0; I < SourcePointsCount; ++I)
+	{
+		if (I != LowestPointIndex)
+		{
+			SortedPoints.Add(I);
+		}
+	}
+	int SortedPointsCount = SortedPoints.Num();
+	for (int I = 0; I < SortedPointsCount; ++I)
+	{
+		for (int J = I + 1; J < SortedPointsCount; ++J)
+		{
+			FVector2D DI = SourcePoints[SortedPoints[I]] - P;
+			FVector2D DJ = SourcePoints[SortedPoints[J]] - P;
+			if (FMath::Atan2(DI.Y, DI.X) > FMath::Atan2(DJ.Y, DJ.X))
+			{
+				int32 Temp = SortedPoints[I];
+				SortedPoints[I] = SortedPoints[J];
+				SortedPoints[J] = Temp;
+			}
+		}
+	}
+
+	// Initialize hull points and add lowest point to end of sorted points for the algorithm
+	TArray<int> Hull;
+	Hull.Add(LowestPointIndex);
+	SortedPoints.Add(LowestPointIndex);
+	++SortedPointsCount;
+	// Traverse sorted points, removing all prior points that would make a right turn, before adding each new point (Graham scan)
+	for (int I = 0; I < SortedPointsCount; ++I)
+	{
+		int NewPointIndex = SortedPoints[I];
+		FVector2D& C = SourcePoints[NewPointIndex];
+		while (Hull.Num() > 1)
+		{
+			FVector2D& A = SourcePoints[Hull[Hull.Num() - 2]];
+			FVector2D& B = SourcePoints[Hull[Hull.Num() - 1]];
+			if ((B.X - A.X) * (C.Y - A.Y) - (B.Y - A.Y) * (C.X - A.X) < 0.0f)
+			{
+				// Remove last entry
+				Hull.RemoveAt(Hull.Num() - 1);
+			}
+			else
+			{
+				// Left turn, exit loop
+				break;
+			}
+		}
+		Hull.Add(NewPointIndex);
+	}
+
+	// We added the starting point to the end for the algorithm, so remove the duplicate
+	Hull.RemoveAt(Hull.Num() - 1);
+
+	// Result
+	int HullCount = Hull.Num();
+	OutConvexHull.Empty(HullCount);
+	for (int I = 0; I < HullCount; ++I)
+	{
+		OutConvexHull.Add(SourcePoints[Hull[I]]);
+	}
+}
diff --git a/Engine/Plugins/2D/Paper2D/Source/Paper2D/Private/PaperGeomTools.h b/Engine/Plugins/2D/Paper2D/Source/Paper2D/Private/PaperGeomTools.h
--- a/Engine/Plugins/2D/Paper2D/Source/Paper2D/Private/PaperGeomTools.h
+++ b/Engine/Plugins/2D/Paper2D/Source/Paper2D/Private/PaperGeomTools.h
@@ -26,5 +26,11 @@ public:
 
 	// 2D version of RemoveRedundantTriangles from GeomTools.cpp
 	static void RemoveRedundantTriangles(TArray<FVector2D>& OutTriangles, const TArray<FVector2D>& InTriangles);
+
+	// Generate convex shapes
+	static void GenerateConvexPolygonsFromTriangles(TArray<TArray<FVector2D>>& OutPolygons, const TArray<FVector2D>& InTriangles);
+
+	// Generate convex hull from points
+	static void GenerateConvexHullFromPoints(TArray<FVector2D>& OutConvexHull, TArray<FVector2D>& Points);
 };
 
diff --git a/Engine/Plugins/2D/Paper2D/Source/Paper2D/Private/Terrain/PaperTerrainComponent.cpp b/Engine/Plugins/2D/Paper2D/Source/Paper2D/Private/Terrain/PaperTerrainComponent.cpp
--- a/Engine/Plugins/2D/Paper2D/Source/Paper2D/Private/Terrain/PaperTerrainComponent.cpp
+++ b/Engine/Plugins/2D/Paper2D/Source/Paper2D/Private/Terrain/PaperTerrainComponent.cpp
@@ -8,9 +8,10 @@
 
 #define PAPER_USE_MATERIAL_SLOPES 1
 
+
 //////////////////////////////////////////////////////////////////////////
 
-FBox2D GetSpriteRenderDataBounds2D(const TArray<FVector4>& Data)
+static FBox2D GetSpriteRenderDataBounds2D(const TArray<FVector4>& Data)
 {
 	FBox2D Bounds(ForceInit);
 	for (const FVector4& XYUV : Data)
@@ -21,6 +22,29 @@ FBox2D GetSpriteRenderDataBounds2D(const TArray<FVector4>& Data)
 	return Bounds;
 }
 
+
+//////////////////////////////////////////////////////////////////////////
+
+FTerrainSpriteStamp::FTerrainSpriteStamp(const UPaperSprite* InSprite, float InTime, bool bIsEndCap) : Sprite(InSprite)
+, Time(InTime)
+, Scale(1.0f)
+, bCanStretch(!bIsEndCap)
+{
+	const FBox2D Bounds2D = GetSpriteRenderDataBounds2D(InSprite->BakedRenderData);
+	NominalWidth = FMath::Max<float>(Bounds2D.GetSize().X, 1.0f);
+}
+
+
+FTerrainSegment::FTerrainSegment() : Rule(nullptr)
+{
+
+}
+
+void FTerrainSegment::RepositionStampsToFillSpace()
+{
+
+}
+
 //////////////////////////////////////////////////////////////////////////
 // FPaperTerrainSceneProxy
 
@@ -70,6 +94,7 @@ UPaperTerrainComponent::UPaperTerrainComponent(const FObjectInitializer& ObjectI
 	, ReparamStepsPerSegment(8)
 	, SpriteCollisionDomain(ESpriteCollisionMode::Use3DPhysics)
 	, CollisionThickness(200.0f)
+	, bFilledSpline(true)
 {
 	bCanEverAffectNavigation = true;
 
@@ -219,6 +244,87 @@ FTransform UPaperTerrainComponent::GetTransformAtDistance(float InDistance) cons
 	return LocalTransform;
 }
 
+// Simplify spline by removing vertices from nearly flat edges
+// Currently don't allow merging points when the offset is different
+// SplineEdgeOffsetAmounts is ignored if empty
+static void SimplifyPolygon(TArray<FVector2D> &SplinePolyVertices2D, TArray<float> &SplineEdgeOffsetAmounts, float FlatEdgeDistance = 10)
+{
+	const bool bHasSplineEdgeOffsets = SplineEdgeOffsetAmounts.Num() == SplinePolyVertices2D.Num();
+	const float FlatEdgeDistanceThreshold = FlatEdgeDistance * FlatEdgeDistance;
+	for (int i = 0; i < SplinePolyVertices2D.Num() && SplinePolyVertices2D.Num() > 3; ++i)
+	{
+		int StartRemoveIndex = (i + 1) % SplinePolyVertices2D.Num();
+		int EndRemoveIndex = StartRemoveIndex;
+		FVector2D& A = SplinePolyVertices2D[i];
+		float SplineEdgeOffsetA = bHasSplineEdgeOffsets ? SplineEdgeOffsetAmounts[i] : 0;
+
+		// Keep searching to find if any of the vector rejections fail in subsequent points on the polygon
+		// A B C D E F (eg. when testing A B C, test rejection for BA, CA)
+		// When testing A E F, test rejection for AB-AF, AC-AF, AD-AF, AE-AF
+		// When one of these fails we discard all verts between A and one before the current vertex being tested
+		for (int j = i; j < SplinePolyVertices2D.Num(); ++j)
+		{
+			int IndexC = (j + 2) % SplinePolyVertices2D.Num();
+			FVector2D& C = SplinePolyVertices2D[IndexC];
+			float SplineEdgeOffsetC = bHasSplineEdgeOffsets ? SplineEdgeOffsetAmounts[IndexC] : 0;
+			bool bSmallOffsetFailed = SplineEdgeOffsetA != SplineEdgeOffsetC;
+
+			for (int k = i; k <= j && !bSmallOffsetFailed; ++k)
+			{
+				int IndexB = (k + 1) % SplinePolyVertices2D.Num();
+				FVector2D& B = SplinePolyVertices2D[IndexB];
+				float SplineEdgeOffsetB = bHasSplineEdgeOffsets ? SplineEdgeOffsetAmounts[IndexB] : 0;
+				if (SplineEdgeOffsetA != SplineEdgeOffsetB)
+				{
+					bSmallOffsetFailed = true;
+					break;
+				}
+
+				FVector2D CA = C - A;
+				FVector2D BA = B - A;
+				FVector2D Rejection_BA_CA = BA - (FVector2D::DotProduct(BA, CA) / FVector2D::DotProduct(CA, CA)) * CA;
+				float RejectionLengthSquared = Rejection_BA_CA.SizeSquared();
+				if (RejectionLengthSquared > FlatEdgeDistanceThreshold)
+				{
+					bSmallOffsetFailed = true;
+					break;
+				}
+			}
+
+			if (bSmallOffsetFailed)
+			{
+				break;
+			}
+			else
+			{
+				EndRemoveIndex = (EndRemoveIndex + 1) % SplinePolyVertices2D.Num();
+			}
+		}
+
+		// Remove the vertices that we deemed "too flat"
+		if (EndRemoveIndex > StartRemoveIndex)
+		{
+			SplinePolyVertices2D.RemoveAt(StartRemoveIndex, EndRemoveIndex - StartRemoveIndex);
+			if (bHasSplineEdgeOffsets)
+			{
+				SplineEdgeOffsetAmounts.RemoveAt(StartRemoveIndex, EndRemoveIndex - StartRemoveIndex);
+			}
+		}
+		else if (EndRemoveIndex < StartRemoveIndex)
+		{
+			SplinePolyVertices2D.RemoveAt(StartRemoveIndex, SplinePolyVertices2D.Num() - StartRemoveIndex);
+			SplinePolyVertices2D.RemoveAt(0, EndRemoveIndex);
+			if (bHasSplineEdgeOffsets)
+			{
+				SplineEdgeOffsetAmounts.RemoveAt(StartRemoveIndex, SplineEdgeOffsetAmounts.Num() - StartRemoveIndex);
+				SplineEdgeOffsetAmounts.RemoveAt(0, EndRemoveIndex);
+			}
+			// The search has wrapped around, no more vertices to test
+			break;
+		}
+	}
+}
+
 void UPaperTerrainComponent::OnSplineEdited()
 {
 	// Ensure we have the data structure for the desired collision method
@@ -251,43 +357,6 @@ void UPaperTerrainComponent::OnSplineEdited()
 		float SplineLength = AssociatedSpline->GetSplineLength();
 
 
-		struct FSpriteStamp
-		{
-			const UPaperSprite* Sprite;
-			float NominalWidth;
-			float Time;
-			float Scale;
-			bool bCanStretch;
-
-			FSpriteStamp(const UPaperSprite* InSprite, float InTime, bool bIsEndCap)
-				: Sprite(InSprite)
-				, Time(InTime)
-				, Scale(1.0f)
-				, bCanStretch(!bIsEndCap)
-			{
-				const FBox2D Bounds2D = GetSpriteRenderDataBounds2D(InSprite->BakedRenderData);
-				NominalWidth = FMath::Max<float>(Bounds2D.GetSize().X, 1.0f);
-			}
-		};
-
-		struct FTerrainSegment
-		{
-			float StartTime;
-			float EndTime;
-			const FPaperTerrainMaterialRule* Rule;
-			TArray<FSpriteStamp> Stamps;
-
-			FTerrainSegment()
-				: Rule(nullptr)
-			{
-			}
-
-			void RepositionStampsToFillSpace()
-			{
-			}
-
-		};
-
 		struct FTerrainRuleHelper
 		{
 		public:
@@ -403,15 +472,8 @@ void UPaperTerrainComponent::OnSplineEdited()
 			// Add the start cap
 			if (RuleHelper.StartWidth > 0.0f)
 			{
-				new (Segment.Stamps) FSpriteStamp(Segment.Rule->StartCap, Segment.StartTime + RuleHelper.StartWidth * 0.5f, /*bIsEndCap=*/ bUseBodySegments);
-			}
-
-			// Add the end cap
-			if (RuleHelper.EndWidth > 0.0f)
-			{
-				new (Segment.Stamps) FSpriteStamp(Segment.Rule->EndCap, Segment.EndTime - RuleHelper.EndWidth * 0.5f, /*bIsEndCap=*/ bUseBodySegments);
+				new (Segment.Stamps) FTerrainSpriteStamp(Segment.Rule->StartCap, Segment.StartTime + RuleHelper.StartWidth * 0.5f, /*bIsEndCap=*/ bUseBodySegments);
 			}
-			
 
 			// Add body segments
 			if (bUseBodySegments)
@@ -429,7 +491,7 @@ void UPaperTerrainComponent::OnSplineEdited()
  					{
  						break;
  					}
-					new (Segment.Stamps) FSpriteStamp(Sprite, Position + (Width * 0.5f), /*bIsEndCap=*/ false);
+					new (Segment.Stamps) FTerrainSpriteStamp(Sprite, Position + (Width * 0.5f), /*bIsEndCap=*/ false);
 
 					DistanceBudget -= Width;
 					Position += Width;
@@ -443,7 +505,7 @@ void UPaperTerrainComponent::OnSplineEdited()
 				float PositionCorrectionSum = 0.0f;
  				for (int32 Index = 0; Index < NumSegments; ++Index)
  				{
-					FSpriteStamp& Stamp = Segment.Stamps[Index + (Segment.Stamps.Num() - NumSegments)];
+					FTerrainSpriteStamp& Stamp = Segment.Stamps[Index + (Segment.Stamps.Num() - NumSegments)];
 					
 					const float WidthChange = (OverallScaleFactor - 1.0f) * Stamp.NominalWidth;
 					const float FirstGapIsSmallerFactor = (Index == 0) ? 0.5f : 1.0f;
@@ -458,19 +520,23 @@ void UPaperTerrainComponent::OnSplineEdited()
 				// Stretch endcaps
 			}
 
-			// Convert stamps into geometry
-			for (const FSpriteStamp& Stamp : Segment.Stamps)
+			// Add the end cap
+			if (RuleHelper.EndWidth > 0.0f)
 			{
-				SpawnSegment(Stamp.Sprite, Stamp.Time, Stamp.Scale, Stamp.NominalWidth, Segment.Rule);
+				new (Segment.Stamps) FTerrainSpriteStamp(Segment.Rule->EndCap, Segment.EndTime - RuleHelper.EndWidth * 0.5f, /*bIsEndCap=*/ bUseBodySegments);
 			}
 		}
 
+		// Convert stamps into geometry
+		SpawnSegments(Segments, !bClosedSpline || (bClosedSpline && !bFilledSpline));
+
 		// Generate the background if the spline is closed
-		if (bClosedSpline && (TerrainMaterial->InteriorFill != nullptr))
+		if (bClosedSpline && bFilledSpline)
 		{
 			// Create a polygon from the spline
 			FBox2D SplineBounds(ForceInit);
 			TArray<FVector2D> SplinePolyVertices2D;
+			TArray<float> SplineEdgeOffsetAmounts;
 			{
 				float CurrentTime = 0.0f;
 				while (CurrentTime < SplineLength)
@@ -482,52 +548,80 @@ void UPaperTerrainComponent::OnSplineEdited()
 					SplineBounds += Position2D;
 					SplinePolyVertices2D.Add(Position2D);
 
+					// Find the collision offset for this sample point
+					float CollisionOffset = 0;
+					for (int SegmentIndex = 0; SegmentIndex < Segments.Num(); ++SegmentIndex)
+					{
+						FTerrainSegment& Segment = Segments[SegmentIndex];
+						if (CurrentTime >= Segment.StartTime && CurrentTime <= Segment.EndTime)
+						{
+							CollisionOffset = (Segment.Rule != nullptr) ? (Segment.Rule->CollisionOffset * 0.25f) : 0;
+							break;
+						}
+					}
+					SplineEdgeOffsetAmounts.Add(CollisionOffset);
+
 					CurrentTime += FillRasterizationTimeRate;
 				}
 			}
 
-			const UPaperSprite* FillSprite = TerrainMaterial->InteriorFill;
-			FPaperTerrainMaterialPair& MaterialBatch = *new (GeneratedSpriteGeometry) FPaperTerrainMaterialPair(); //@TODO: Look up the existing one instead
-			MaterialBatch.Material = FillSprite->GetDefaultMaterial();
+			SimplifyPolygon(SplinePolyVertices2D, SplineEdgeOffsetAmounts);
 
-			FSpriteDrawCallRecord& FillDrawCall = *new (MaterialBatch.Records) FSpriteDrawCallRecord();
-			FillDrawCall.BuildFromSprite(FillSprite);
-			FillDrawCall.RenderVerts.Empty();
-			FillDrawCall.Color = TerrainColor;
-			FillDrawCall.Destination = PaperAxisZ * 0.1f;
+			// Always CCW and facing forward regardless of spline winding
+			TArray<FVector2D> CorrectedSplineVertices;
+			PaperGeomTools::CorrectPolygonWinding(CorrectedSplineVertices, SplinePolyVertices2D, false);
 
-			const FVector2D TextureSize = GetSpriteRenderDataBounds2D(FillSprite->BakedRenderData).GetSize();
-			const FVector2D SplineSize = SplineBounds.GetSize();
-			
-			SpawnFromPoly(FillSprite, FillDrawCall, TextureSize, SplinePolyVertices2D);
+			TArray<FVector2D> TriangulatedPolygonVertices;
+			PaperGeomTools::TriangulatePoly(/*out*/TriangulatedPolygonVertices, CorrectedSplineVertices, false);
 
-			//@TODO: Add support for the fill sprite being smaller than the entire texture
-#if NOT_WORKING
-			const float StartingDivisionPointX = FMath::CeilToFloat(SplineBounds.Min.X / TextureSize.X);
-			const float StartingDivisionPointY = FMath::CeilToFloat(SplineBounds.Min.Y / TextureSize.Y);
+			GenerateCollisionDataFromPolygon(SplinePolyVertices2D, SplineEdgeOffsetAmounts, TriangulatedPolygonVertices);
 
-			FPoly VerticalRemainder = SplineAsPolygon;
-			for (float Y = StartingDivisionPointY; VerticalRemainder.Vertices.Num() > 0; Y += TextureSize.Y)
+			if (TerrainMaterial->InteriorFill != nullptr)
 			{
-				FPoly Top;
-				FPoly Bottom;
-				const FVector SplitBaseOuter = (Y * PaperAxisY);
-				VerticalRemainder.SplitWithPlane(SplitBaseOuter, -PaperAxisY, &Top, &Bottom, 1);
-				VerticalRemainder = Bottom;
-
-				FPoly HorizontalRemainder = Top;
-				for (float X = StartingDivisionPointX; HorizontalRemainder.Vertices.Num() > 0; X += TextureSize.X)
-				{
-					FPoly Left;
-					FPoly Right;
-					const FVector SplitBaseInner = (X * PaperAxisX) + (Y * PaperAxisY);
-					HorizontalRemainder.SplitWithPlane(SplitBaseInner, -PaperAxisX, &Left, &Right, 1);
-					HorizontalRemainder = Right;
+				const UPaperSprite* FillSprite = TerrainMaterial->InteriorFill;
+				FPaperTerrainMaterialPair& MaterialBatch = *new (GeneratedSpriteGeometry)FPaperTerrainMaterialPair(); //@TODO: Look up the existing one instead
+				MaterialBatch.Material = FillSprite->GetDefaultMaterial();
+
+				FSpriteDrawCallRecord& FillDrawCall = *new (MaterialBatch.Records) FSpriteDrawCallRecord();
+				FillDrawCall.BuildFromSprite(FillSprite);
+				FillDrawCall.RenderVerts.Empty();
+				FillDrawCall.Color = TerrainColor;
+				FillDrawCall.Destination = PaperAxisZ * 0.1f;
+
+				const FVector2D TextureSize = GetSpriteRenderDataBounds2D(FillSprite->BakedRenderData).GetSize();
+				const FVector2D SplineSize = SplineBounds.GetSize();
 
-					SpawnFromPoly(FillSprite, FillDrawCall, TextureSize, Left);
+				GenerateFillRenderDataFromPolygon(FillSprite, FillDrawCall, TextureSize, TriangulatedPolygonVertices);
+
+				//@TODO: Add support for the fill sprite being smaller than the entire texture
+#if NOT_WORKING
+				const float StartingDivisionPointX = FMath::CeilToFloat(SplineBounds.Min.X / TextureSize.X);
+				const float StartingDivisionPointY = FMath::CeilToFloat(SplineBounds.Min.Y / TextureSize.Y);
+
+				FPoly VerticalRemainder = SplineAsPolygon;
+				for (float Y = StartingDivisionPointY; VerticalRemainder.Vertices.Num() > 0; Y += TextureSize.Y)
+				{
+					FPoly Top;
+					FPoly Bottom;
+					const FVector SplitBaseOuter = (Y * PaperAxisY);
+					VerticalRemainder.SplitWithPlane(SplitBaseOuter, -PaperAxisY, &Top, &Bottom, 1);
+					VerticalRemainder = Bottom;
+
+					FPoly HorizontalRemainder = Top;
+					for (float X = StartingDivisionPointX; HorizontalRemainder.Vertices.Num() > 0; X += TextureSize.X)
+					{
+						FPoly Left;
+						FPoly Right;
+						const FVector SplitBaseInner = (X * PaperAxisX) + (Y * PaperAxisY);
+						HorizontalRemainder.SplitWithPlane(SplitBaseInner, -PaperAxisX, &Left, &Right, 1);
+						HorizontalRemainder = Right;
+
+						//BROKEN, function no longer exists (split into 2 parts)
+						SpawnFromPoly(Segments, SplineEdgeOffsetAmounts, FillSprite, FillDrawCall, TextureSize, Left);
+					}
 				}
-			}
 #endif
+			}
 		}
 
 		// Draw debug frames at the start and end of the spline
@@ -549,75 +643,261 @@ void UPaperTerrainComponent::OnSplineEdited()
 	RecreateRenderState_Concurrent();
 }
 
-void UPaperTerrainComponent::SpawnSegment(const UPaperSprite* NewSprite, float Position, float HorizontalScale, float NominalWidth, const FPaperTerrainMaterialRule* Rule)
+#define USE_SIMPLIFIED_POLYGON_COLLIDERS_FOR_SEGMENTS
+
+void UPaperTerrainComponent::SpawnSegments(const TArray<FTerrainSegment>& TerrainSegments, bool bGenerateSegmentColliders)
 {
-	if ((Rule->bEnableCollision) && (CachedBodySetup != nullptr))
+#ifdef USE_SIMPLIFIED_POLYGON_COLLIDERS_FOR_SEGMENTS
+	TArray<FVector2D> CollisionPolygonPoints;
+#endif
+
+	for (const FTerrainSegment& Segment : TerrainSegments)
 	{
-		const FTransform LocalTransformAtCenter(GetTransformAtDistance(Position));
+		for (const FTerrainSpriteStamp& Stamp : Segment.Stamps)
+		{
+			const class UPaperSprite* NewSprite = Stamp.Sprite;
+			float Position = Stamp.Time;
+			float HorizontalScale = Stamp.Scale;
+			float NominalWidth = Stamp.NominalWidth;
+			const struct FPaperTerrainMaterialRule* Rule = Segment.Rule;
 
-		FKBoxElem Box;
+			if (bGenerateSegmentColliders && (Rule->bEnableCollision) && (CachedBodySetup != nullptr))
+			{
+				const FTransform LocalTransformAtCenter(GetTransformAtDistance(Position));
 
-		if (bClosedSpline)
-		{
-			//@TODO: Add proper support for this!
-			Box.SetTransform(LocalTransformAtCenter);
-			Box.X = NominalWidth * HorizontalScale;
-			Box.Y = CollisionThickness;
-			Box.Z = 30.0f;
+#ifdef USE_SIMPLIFIED_POLYGON_COLLIDERS_FOR_SEGMENTS
+				// Check note be low Re: bClosedSplines
+				FVector2D BoxExtents( 0.5f * NominalWidth * HorizontalScale, 0.5f * FMath::Max<float>(1.0f, FMath::Abs<float>(Rule->CollisionOffset * 0.5f)) );
+				
+				FVector BoxPoints[4];
+				BoxPoints[0] = LocalTransformAtCenter.TransformPosition(FVector(BoxExtents.X, 0, BoxExtents.Y));
+				BoxPoints[1] = LocalTransformAtCenter.TransformPosition(FVector(-BoxExtents.X, 0, BoxExtents.Y));
+				BoxPoints[2] = LocalTransformAtCenter.TransformPosition(FVector(-BoxExtents.X, 0, -BoxExtents.Y));
+				BoxPoints[3] = LocalTransformAtCenter.TransformPosition(FVector(BoxExtents.X, 0, -BoxExtents.Y));
+
+				FVector2D BoxPoints2D[4];
+				BoxPoints2D[0].Set(BoxPoints[0].X, BoxPoints[0].Z);
+				BoxPoints2D[1].Set(BoxPoints[1].X, BoxPoints[1].Z);
+				BoxPoints2D[2].Set(BoxPoints[2].X, BoxPoints[2].Z);
+				BoxPoints2D[3].Set(BoxPoints[3].X, BoxPoints[3].Z);
+
+				// If there is a previous polygon, try to merge
+				if (CollisionPolygonPoints.Num() >= 4)
+				{
+					int InsertPoint = CollisionPolygonPoints.Num() / 2 - 1;
+					float LengthV0 = FVector2D::Distance(CollisionPolygonPoints[InsertPoint], BoxPoints2D[0]);
+					float LengthV1 = FVector2D::Distance(CollisionPolygonPoints[InsertPoint + 1], BoxPoints2D[3]);
+					float MergeThreshold = 10;
+					bool bMergeIntoPolygon = LengthV0 < MergeThreshold && LengthV1 < MergeThreshold;
+					if (bMergeIntoPolygon)
+					{
+						CollisionPolygonPoints.Insert(BoxPoints2D[2], InsertPoint + 1);
+						CollisionPolygonPoints.Insert(BoxPoints2D[1], InsertPoint + 1);
+					}
+					else
+					{
+						InsertConvexCollisionDataFromPolygon(CollisionPolygonPoints);
+						CollisionPolygonPoints.Empty(0);
+						CollisionPolygonPoints.Add(BoxPoints2D[0]);
+						CollisionPolygonPoints.Add(BoxPoints2D[1]);
+						CollisionPolygonPoints.Add(BoxPoints2D[2]);
+						CollisionPolygonPoints.Add(BoxPoints2D[3]);
+					}
+				}
+				else
+				{
+					CollisionPolygonPoints.Add(BoxPoints2D[0]);
+					CollisionPolygonPoints.Add(BoxPoints2D[1]);
+					CollisionPolygonPoints.Add(BoxPoints2D[2]);
+					CollisionPolygonPoints.Add(BoxPoints2D[3]);
+				}
+#else
+				FKBoxElem Box;
+				// The spline is never "closed" properly right now
+				//if (bClosedSpline)
+				//{
+				//	//@TODO: Add proper support for this!
+				//	Box.SetTransform(LocalTransformAtCenter);
+				//	Box.X = NominalWidth * HorizontalScale;
+				//	Box.Y = CollisionThickness;
+				//	Box.Z = 30.0f;
+				//}
+				//else
+				{
+					Box.SetTransform(LocalTransformAtCenter);
+					Box.X = NominalWidth * HorizontalScale;
+					Box.Y = CollisionThickness;
+
+					Box.Z = FMath::Max<float>(1.0f, FMath::Abs<float>(Rule->CollisionOffset * 0.5f));
+				}
+				CachedBodySetup->AggGeom.BoxElems.Add(Box);
+#endif
+			}
+
+			if (NewSprite)
+			{
+				FPaperTerrainMaterialPair& MaterialBatch = *new (GeneratedSpriteGeometry) FPaperTerrainMaterialPair(); //@TODO: Look up the existing one instead
+				MaterialBatch.Material = NewSprite->GetDefaultMaterial();
+
+				FSpriteDrawCallRecord& Record = *new (MaterialBatch.Records) FSpriteDrawCallRecord();
+				Record.BuildFromSprite(NewSprite);
+				Record.Color = TerrainColor;
+
+				for (FVector4& XYUV : Record.RenderVerts)
+				{
+					const FTransform LocalTransformAtX(GetTransformAtDistance(Position + (XYUV.X * HorizontalScale)));
+					const FVector SourceVector = (XYUV.Y * PaperAxisY);
+					const FVector NewVector = LocalTransformAtX.TransformPosition(SourceVector);
+
+					const float NewX = FVector::DotProduct(NewVector, PaperAxisX);
+					const float NewY = FVector::DotProduct(NewVector, PaperAxisY);
+
+					XYUV.X = NewX;
+					XYUV.Y = NewY;
+				}
+			}
 		}
-		else
+	}
+
+#ifdef USE_SIMPLIFIED_POLYGON_COLLIDERS_FOR_SEGMENTS
+	// For whatever is remaining
+	if (CollisionPolygonPoints.Num() > 0)
+	{
+		InsertConvexCollisionDataFromPolygon(CollisionPolygonPoints);
+	}
+#endif
+}
+
+// Create an extruded convex hull resulting from extruding edges by the amount defined in EdgeExtrudeAmount
+// Each edge is extruded in the normal direction to the edge
+static void CreateExtrudedConvexHull(TArray<FVector2D>& OutConvexHull, TArray<FVector2D>& SourcePoints, TArray<float> EdgeExtrudeAmount)
+{
+	TArray<FVector2D> ExtrudedPoints;
+	int SourcePointsCount = SourcePoints.Num();
+	for (int i = 0; i < SourcePointsCount; ++i)
+	{
+		FVector2D& A = SourcePoints[i];
+		FVector2D& B = SourcePoints[(i + 1) % SourcePointsCount];
+		FVector2D N(B.Y - A.Y, A.X - B.X);
+		N.Normalize();
+		float extrude = EdgeExtrudeAmount[i];
+
+		// Each edge is pushed forwards and backwards, and the points added
+		ExtrudedPoints.Add(A + N * extrude);
+		ExtrudedPoints.Add(A - N * extrude);
+		ExtrudedPoints.Add(B + N * extrude);
+		ExtrudedPoints.Add(B - N * extrude);
+	}
+
+	PaperGeomTools::GenerateConvexHullFromPoints(OutConvexHull, ExtrudedPoints);
+}
+
+void UPaperTerrainComponent::GenerateFillRenderDataFromPolygon(const class UPaperSprite* NewSprite, FSpriteDrawCallRecord& FillDrawCall, const FVector2D& TextureSize, const TArray<FVector2D>& TriangulatedPolygonVertices)
+{
+	const FVector2D TextureSizeInUnits = TextureSize * NewSprite->GetUnrealUnitsPerPixel();
+
+	// Pack vertex data
+	if (TriangulatedPolygonVertices.Num() >= 3)
+	{
+		for (int32 TriangleVertexIndex = 0; TriangleVertexIndex < TriangulatedPolygonVertices.Num(); ++TriangleVertexIndex)
 		{
-			Box.SetTransform(LocalTransformAtCenter);
-			Box.X = NominalWidth * HorizontalScale;
-			Box.Y = CollisionThickness;
-			Box.Z = FMath::Max<float>(1.0f, FMath::Abs<float>(Rule->CollisionOffset * 0.5f));
+			const FVector2D& TriangleVertex = TriangulatedPolygonVertices[TriangleVertexIndex];
+			new (FillDrawCall.RenderVerts) FVector4(TriangleVertex.X, TriangleVertex.Y, TriangleVertex.X / TextureSizeInUnits.X, -TriangleVertex.Y / TextureSizeInUnits.Y);
 		}
-		CachedBodySetup->AggGeom.BoxElems.Add(Box);
 	}
+}
 
-	if (NewSprite)
+void UPaperTerrainComponent::GenerateCollisionDataFromPolygon(const TArray<FVector2D>& SplinePolyVertices2D, const TArray<float>& TerrainOffsets, const TArray<FVector2D>& TriangulatedPolygonVertices)
+{
+	if (CachedBodySetup != nullptr && TriangulatedPolygonVertices.Num() >= 3)
 	{
-		FPaperTerrainMaterialPair& MaterialBatch = *new (GeneratedSpriteGeometry) FPaperTerrainMaterialPair(); //@TODO: Look up the existing one instead
-		MaterialBatch.Material = NewSprite->GetDefaultMaterial();
-
-		FSpriteDrawCallRecord& Record = *new (MaterialBatch.Records) FSpriteDrawCallRecord();
-		Record.BuildFromSprite(NewSprite);
-		Record.Color = TerrainColor;
+		// Generate polygon collider
+		TArray<TArray<FVector2D>> ConvexHulls;
+		PaperGeomTools::GenerateConvexPolygonsFromTriangles(ConvexHulls, TriangulatedPolygonVertices);
+		TArray<float> ConvexHullEdgeExtrusionAmount;
 
-		for (FVector4& XYUV : Record.RenderVerts)
+		for (TArray<FVector2D> ConvexHull : ConvexHulls)
 		{
-			const FTransform LocalTransformAtX(GetTransformAtDistance(Position + (XYUV.X * HorizontalScale)));
-			const FVector SourceVector = (XYUV.Y * PaperAxisY);
-			const FVector NewVector = LocalTransformAtX.TransformPosition(SourceVector);
+			ConvexHullEdgeExtrusionAmount.Empty(ConvexHull.Num());
 
-			const float NewX = FVector::DotProduct(NewVector, PaperAxisX);
-			const float NewY = FVector::DotProduct(NewVector, PaperAxisY);
+			// Find distances for each edge in this convex hull from the pair of points forming the edge
+			// Not all edges will match edges in the original concave geometry, eg. newly created internal edges
+			//TODO: Speed this up by using indices / vertex & edge identifiers instead of brute force search
+			for (int ConvexHullPoint = 0; ConvexHullPoint < ConvexHull.Num(); ++ConvexHullPoint)
+			{
+				const FVector2D& A = ConvexHull[ConvexHullPoint];
+				const FVector2D& B = ConvexHull[(ConvexHullPoint + 1) % ConvexHull.Num()];
+				bool bFound = false;
+				for (int Vertex = 0; Vertex < SplinePolyVertices2D.Num(); ++Vertex)
+				{
+					// The winding is might be different to the source polygon, compare both ways
+					int NextVertexIndex = (Vertex + 1) % SplinePolyVertices2D.Num();
+					if ((A.Equals(SplinePolyVertices2D[Vertex], THRESH_POINTS_ARE_SAME) && B.Equals(SplinePolyVertices2D[NextVertexIndex], THRESH_POINTS_ARE_SAME)) ||
+						(B.Equals(SplinePolyVertices2D[Vertex], THRESH_POINTS_ARE_SAME) && A.Equals(SplinePolyVertices2D[NextVertexIndex], THRESH_POINTS_ARE_SAME)))
+					{
+						// Found an edge that matches the 2 vertex points
+						ConvexHullEdgeExtrusionAmount.Add(TerrainOffsets[Vertex]);
+						bFound = true;
+						break;
+					}
+				}
+				if (!bFound)
+				{
+					// Couldn't find this edge in the original polygon 
+					ConvexHullEdgeExtrusionAmount.Add(0);
+				}
+			}
+				
+			TArray<FVector2D> ExtrudedConvexHull;
+			CreateExtrudedConvexHull(ExtrudedConvexHull, ConvexHull, ConvexHullEdgeExtrusionAmount);
+//			ExtrudedConvexHull = ConvexHull;
+
+			// Generate convex hull
+			FKConvexElem Convex;
+			for (int J = 0; J < ExtrudedConvexHull.Num(); ++J)
+			{
+				FVector2D& Vert = ExtrudedConvexHull[J];
+				new(Convex.VertexData) FVector(Vert.X, -0.5f * CollisionThickness, Vert.Y);
+				new(Convex.VertexData) FVector(Vert.X, 0.5f * CollisionThickness, Vert.Y);
+			}
+			Convex.UpdateElemBox();
 
-			XYUV.X = NewX;
-			XYUV.Y = NewY;
+			CachedBodySetup->AggGeom.ConvexElems.Add(Convex);
 		}
 	}
 }
 
-void UPaperTerrainComponent::SpawnFromPoly(const class UPaperSprite* NewSprite, FSpriteDrawCallRecord& Batch, const FVector2D& TextureSize, const TArray<FVector2D>& SplinePolyVertices2D)
+void UPaperTerrainComponent::InsertConvexCollisionDataFromPolygon(const TArray<FVector2D>& ClosedPolyVertices2D)
 {
-	const FVector2D TextureSizeInUnits = TextureSize * NewSprite->GetUnrealUnitsPerPixel();
+	if (CachedBodySetup != nullptr && ClosedPolyVertices2D.Num() >= 3)
+	{
+		// Simplify polygon
+		TArray<float> EmptyOffsetsList;
+		TArray<FVector2D> LocalPolyVertices = ClosedPolyVertices2D;
+		SimplifyPolygon(LocalPolyVertices, EmptyOffsetsList);
 
-	// Correct spline vertices so they are always CCW
-	//@TODO: Handle overlapping and crossing polygon edges during triangulation
-	TArray<FVector2D> CorrectedSpriteVertices;
-	PaperGeomTools::CorrectPolygonWinding(/*out*/CorrectedSpriteVertices, SplinePolyVertices2D, false);
+		// Always CCW and facing forward regardless of spline winding
+		TArray<FVector2D> CorrectedSplineVertices;
+		PaperGeomTools::CorrectPolygonWinding(CorrectedSplineVertices, LocalPolyVertices, false);
 
-	if (SplinePolyVertices2D.Num() >= 3)
-	{
 		TArray<FVector2D> TriangulatedPolygonVertices;
-		PaperGeomTools::TriangulatePoly(/*out*/TriangulatedPolygonVertices, CorrectedSpriteVertices, false);
+		PaperGeomTools::TriangulatePoly(/*out*/TriangulatedPolygonVertices, CorrectedSplineVertices, false);
 
-		// Pack vertex data
-		for (int32 TriangleVertexIndex = 0; TriangleVertexIndex < TriangulatedPolygonVertices.Num(); ++TriangleVertexIndex)
+		TArray<TArray<FVector2D>> ConvexHulls;
+		PaperGeomTools::GenerateConvexPolygonsFromTriangles(ConvexHulls, TriangulatedPolygonVertices);
+
+		for (TArray<FVector2D> ConvexHull : ConvexHulls)
 		{
-			const FVector2D& TriangleVertex = TriangulatedPolygonVertices[TriangleVertexIndex];
-			new (Batch.RenderVerts) FVector4(TriangleVertex.X, TriangleVertex.Y, TriangleVertex.X / TextureSizeInUnits.X, -TriangleVertex.Y / TextureSizeInUnits.Y);
+			FKConvexElem Convex;
+			for (int J = 0; J < ConvexHull.Num(); ++J)
+			{
+				FVector2D& Vert = ConvexHull[J];
+				new(Convex.VertexData) FVector(Vert.X, -0.5f * CollisionThickness, Vert.Y);
+				new(Convex.VertexData) FVector(Vert.X, 0.5f * CollisionThickness, Vert.Y);
+			}
+			Convex.UpdateElemBox();
+
+			CachedBodySetup->AggGeom.ConvexElems.Add(Convex);
 		}
 	}
 }
