Paper2D: PaperTerrainComponent, bFilledSpline is only enabled when bClosedSpline

diff --git a/Engine/Plugins/2D/Paper2D/Source/Paper2D/Classes/PaperTerrainComponent.h b/Engine/Plugins/2D/Paper2D/Source/Paper2D/Classes/PaperTerrainComponent.h
--- a/Engine/Plugins/2D/Paper2D/Source/Paper2D/Classes/PaperTerrainComponent.h
+++ b/Engine/Plugins/2D/Paper2D/Source/Paper2D/Classes/PaperTerrainComponent.h
@@ -53,7 +53,7 @@ public:
 	UPROPERTY(Category = Sprite, EditAnywhere, BlueprintReadOnly)
 	bool bClosedSpline;
 
-	UPROPERTY(Category = Sprite, EditAnywhere, BlueprintReadOnly)
+	UPROPERTY(Category = Sprite, EditAnywhere, BlueprintReadOnly, meta = (EditCondition = "bClosedSpline"))
 	bool bFilledSpline;
 
 	UPROPERTY()
