Paper2D: Add stub mesh paint adapter for sprite components

diff --git a/Engine/Plugins/2D/Paper2D/Source/Paper2DEditor/Paper2DEditor.Build.cs b/Engine/Plugins/2D/Paper2D/Source/Paper2DEditor/Paper2DEditor.Build.cs
--- a/Engine/Plugins/2D/Paper2D/Source/Paper2DEditor/Paper2DEditor.Build.cs
+++ b/Engine/Plugins/2D/Paper2D/Source/Paper2DEditor/Paper2DEditor.Build.cs
@@ -6,10 +6,6 @@ public class Paper2DEditor : ModuleRules
 {
 	public Paper2DEditor(TargetInfo Target)
 	{
-		//Example path:
-		//D:\BaseUE4\Engine\Source\Editor\Persona\Private\Persona.cpp
-		//PrivateIncludePaths.Add("Editor/Persona/Private");	// For PCH includes (because they don't work with relative paths, yet)
-
 		PrivateIncludePaths.Add("Paper2DEditor/Private");
 
 		PrivateDependencyModuleNames.AddRange(
@@ -39,7 +35,8 @@ public class Paper2DEditor : ModuleRules
 		PrivateIncludePathModuleNames.AddRange(
 			new string[] {
 				"Settings",
-				"IntroTutorials"
+				"IntroTutorials",
+				"MeshPaint"
 			}
 		);
 	}
diff --git a/Engine/Plugins/2D/Paper2D/Source/Paper2DEditor/Private/MeshPainting/MeshPaintSpriteAdapter.cpp b/Engine/Plugins/2D/Paper2D/Source/Paper2DEditor/Private/MeshPainting/MeshPaintSpriteAdapter.cpp
new file mode 100644
--- /dev/null
+++ b/Engine/Plugins/2D/Paper2D/Source/Paper2DEditor/Private/MeshPainting/MeshPaintSpriteAdapter.cpp
@@ -0,0 +1,50 @@
+// Copyright 1998-2015 Epic Games, Inc. All Rights Reserved.
+
+#include "Paper2DEditorPrivatePCH.h"
+#include "MeshPaintSpriteAdapter.h"
+
+//////////////////////////////////////////////////////////////////////////
+// FMeshPaintSpriteAdapter
+
+bool FMeshPaintSpriteAdapter::Construct(UMeshComponent* InComponent, int32 InPaintingMeshLODIndex, int32 InUVChannelIndex)
+{
+	SpriteComponent = CastChecked<UPaperSpriteComponent>(InComponent);
+	return SpriteComponent->GetSprite() != nullptr;
+}
+
+int32 FMeshPaintSpriteAdapter::GetNumTexCoords() const
+{
+	return 1;
+}
+
+void FMeshPaintSpriteAdapter::GetTriangleInfo(int32 TriIndex, struct FTexturePaintTriangleInfo& OutTriInfo) const
+{
+	//@TODO: MESHPAINT: Write this
+}
+
+
+bool FMeshPaintSpriteAdapter::LineTraceComponent(struct FHitResult& OutHit, const FVector Start, const FVector End, const struct FCollisionQueryParams& Params) const
+{
+	//@TODO: MESHPAINT: Write this
+	return false;
+}
+
+//////////////////////////////////////////////////////////////////////////
+// FMeshPaintSpriteAdapterFactory
+
+TSharedPtr<IMeshPaintGeometryAdapter> FMeshPaintSpriteAdapterFactory::Construct(class UMeshComponent* InComponent, int32 InPaintingMeshLODIndex, int32 InUVChannelIndex) const
+{
+	if (UPaperSpriteComponent* SpriteComponent = Cast<UPaperSpriteComponent>(InComponent))
+	{
+		if (SpriteComponent->GetSprite() != nullptr)
+		{
+			TSharedRef<FMeshPaintSpriteAdapter> Result = MakeShareable(new FMeshPaintSpriteAdapter());
+			if (Result->Construct(InComponent, InPaintingMeshLODIndex, InUVChannelIndex))
+			{
+				return Result;
+			}
+		}
+	}
+
+	return nullptr;
+}
diff --git a/Engine/Plugins/2D/Paper2D/Source/Paper2DEditor/Private/MeshPainting/MeshPaintSpriteAdapter.h b/Engine/Plugins/2D/Paper2D/Source/Paper2DEditor/Private/MeshPainting/MeshPaintSpriteAdapter.h
new file mode 100644
--- /dev/null
+++ b/Engine/Plugins/2D/Paper2D/Source/Paper2DEditor/Private/MeshPainting/MeshPaintSpriteAdapter.h
@@ -0,0 +1,31 @@
+// Copyright 1998-2015 Epic Games, Inc. All Rights Reserved.
+
+#pragma once
+
+#include "MeshPaintModule.h"
+
+//////////////////////////////////////////////////////////////////////////
+// FMeshPaintSpriteAdapter
+
+class FMeshPaintSpriteAdapter : public IMeshPaintGeometryAdapter
+{
+public:
+	virtual bool Construct(UMeshComponent* InComponent, int32 InPaintingMeshLODIndex, int32 InUVChannelIndex) override;
+	virtual int32 GetNumTexCoords() const override;
+	virtual void GetTriangleInfo(int32 TriIndex, struct FTexturePaintTriangleInfo& OutTriInfo) const override;
+	virtual bool SupportsTexturePaint() const override { return true; }
+	virtual bool SupportsVertexPaint() const override { return false; }
+	virtual bool LineTraceComponent(struct FHitResult& OutHit, const FVector Start, const FVector End, const struct FCollisionQueryParams& Params) const override;
+
+protected:
+	UPaperSpriteComponent* SpriteComponent;
+};
+
+//////////////////////////////////////////////////////////////////////////
+// FMeshPaintSpriteAdapterFactory
+
+class FMeshPaintSpriteAdapterFactory : public IMeshPaintGeometryAdapterFactory
+{
+public:
+	virtual TSharedPtr<IMeshPaintGeometryAdapter> Construct(class UMeshComponent* InComponent, int32 InPaintingMeshLODIndex, int32 InUVChannelIndex) const override;
+};
diff --git a/Engine/Plugins/2D/Paper2D/Source/Paper2DEditor/Private/Paper2DEditorModule.cpp b/Engine/Plugins/2D/Paper2D/Source/Paper2DEditor/Private/Paper2DEditorModule.cpp
--- a/Engine/Plugins/2D/Paper2D/Source/Paper2DEditor/Private/Paper2DEditorModule.cpp
+++ b/Engine/Plugins/2D/Paper2D/Source/Paper2DEditor/Private/Paper2DEditorModule.cpp
@@ -42,6 +42,9 @@
 // Intro tutorials
 #include "Editor/IntroTutorials/Public/IIntroTutorials.h"
 
+// Mesh paint adapters
+#include "MeshPainting/MeshPaintSpriteAdapter.h"
+
 DEFINE_LOG_CATEGORY(LogPaper2DEditor);
 
 #define LOCTEXT_NAMESPACE "Paper2DEditor"
@@ -74,6 +77,7 @@ private:
 	TSharedPtr<IComponentAssetBroker> PaperFlipbookBroker;
 	TSharedPtr<IComponentAssetBroker> PaperTileMapBroker;
 
+	TSharedPtr<IMeshPaintGeometryAdapterFactory> SpriteMeshPaintAdapterFactory;
 	FCoreUObjectDelegates::FOnObjectPropertyChanged::FDelegate OnPropertyChangedHandle;
 
 public:
@@ -137,6 +141,14 @@ public:
 		// Integrate Paper2D actions associated with existing engine types (e.g., Texture2D) into the content browser
 		FPaperContentBrowserExtensions::InstallHooks();
 
+		// Register with the mesh paint module
+		if (IMeshPaintModule* MeshPaintModule = FModuleManager::LoadModulePtr<IMeshPaintModule>("MeshPaint"))
+		{
+			SpriteMeshPaintAdapterFactory = MakeShareable(new FMeshPaintSpriteAdapterFactory());
+			MeshPaintModule->RegisterGeometryAdapterFactory(SpriteMeshPaintAdapterFactory.ToSharedRef());
+		}
+
+		//
 		RegisterSettings();
 	}
 
@@ -152,6 +164,13 @@ public:
 		{
 			UnregisterSettings();
 
+			// Unregister from the mesh paint module
+			if (IMeshPaintModule* MeshPaintModule = FModuleManager::GetModulePtr<IMeshPaintModule>("MeshPaint"))
+			{
+				MeshPaintModule->UnregisterGeometryAdapterFactory(SpriteMeshPaintAdapterFactory.ToSharedRef());
+				SpriteMeshPaintAdapterFactory.Reset();
+			}
+
 			FPaperContentBrowserExtensions::RemoveHooks();
 
 			FComponentAssetBrokerage::UnregisterBroker(PaperTileMapBroker);
