Paper2D: Implement FMeshPaintSpriteAdapter::LineTraceComponent

diff --git a/Engine/Plugins/2D/Paper2D/Source/Paper2DEditor/Private/MeshPainting/MeshPaintSpriteAdapter.cpp b/Engine/Plugins/2D/Paper2D/Source/Paper2DEditor/Private/MeshPainting/MeshPaintSpriteAdapter.cpp
--- a/Engine/Plugins/2D/Paper2D/Source/Paper2DEditor/Private/MeshPainting/MeshPaintSpriteAdapter.cpp
+++ b/Engine/Plugins/2D/Paper2D/Source/Paper2DEditor/Private/MeshPainting/MeshPaintSpriteAdapter.cpp
@@ -25,7 +25,61 @@ void FMeshPaintSpriteAdapter::GetTriangleInfo(int32 TriIndex, struct FTexturePai
 
 bool FMeshPaintSpriteAdapter::LineTraceComponent(struct FHitResult& OutHit, const FVector Start, const FVector End, const struct FCollisionQueryParams& Params) const
 {
-	//@TODO: MESHPAINT: Write this
+	UPaperSprite* Sprite = SpriteComponent->GetSprite();
+
+	const FTransform& ComponentToWorld = SpriteComponent->ComponentToWorld;
+
+	// Can we possibly intersect with the sprite?
+	const FBoxSphereBounds& Bounds = SpriteComponent->Bounds;
+	if (FMath::PointDistToSegment(Bounds.Origin, Start, End) <= Bounds.SphereRadius)
+	{
+		const FVector LocalStart = ComponentToWorld.InverseTransformPosition(Start);
+		const FVector LocalEnd = ComponentToWorld.InverseTransformPosition(End);
+
+		FPlane LocalSpacePlane(FVector::ZeroVector, PaperAxisX, PaperAxisY);
+
+		FVector Intersection;
+		if (FMath::SegmentPlaneIntersection(LocalStart, LocalEnd, LocalSpacePlane, /*out*/ Intersection))
+		{
+			const float LocalX = FVector::DotProduct(Intersection, PaperAxisX);
+			const float LocalY = FVector::DotProduct(Intersection, PaperAxisY);
+			const FVector LocalPoint(LocalX, LocalY, 0.0f);
+
+			const TArray<FVector4>& BakedPoints = Sprite->BakedRenderData;
+			check((BakedPoints.Num() % 3) == 0);
+
+			for (int32 VertexIndex = 0; VertexIndex < BakedPoints.Num(); VertexIndex += 3)
+			{
+				const FVector& A = BakedPoints[VertexIndex + 0];
+				const FVector& B = BakedPoints[VertexIndex + 1];
+				const FVector& C = BakedPoints[VertexIndex + 2];
+				const FVector Q = FMath::GetBaryCentric2D(LocalPoint, A, B, C);
+
+				if ((Q.X >= 0.0f) && (Q.Y >= 0.0f) && (Q.Z >= 0.0f) && FMath::IsNearlyEqual(Q.X + Q.Y + Q.Z, 1.0f))
+				{
+					const FVector WorldIntersection = ComponentToWorld.TransformPosition(Intersection);
+
+					const FVector WorldNormalFront = ComponentToWorld.TransformVectorNoScale(PaperAxisZ);
+					const FVector WorldNormal = (LocalSpacePlane.PlaneDot(LocalStart) >= 0.0f) ? WorldNormalFront : -WorldNormalFront;
+
+					OutHit.bBlockingHit = true;
+					OutHit.Time = (WorldIntersection - Start).Size() / (End - Start).Size();
+					OutHit.Location = WorldIntersection;
+					OutHit.Normal = WorldNormal;
+					OutHit.ImpactPoint = WorldIntersection;
+					OutHit.ImpactNormal = WorldNormal;
+					OutHit.TraceStart = Start;
+					OutHit.TraceEnd = End;
+					OutHit.Actor = SpriteComponent->GetOwner();
+					OutHit.Component = SpriteComponent;
+					OutHit.FaceIndex = VertexIndex / 3;
+
+					return true;
+				}
+			}
+		}
+	}
+
 	return false;
 }
 
