Paper2D: Improved shrink wrap quality and performance significantly

diff --git a/Engine/Plugins/2D/Paper2D/Source/Paper2D/Classes/PaperSprite.h b/Engine/Plugins/2D/Paper2D/Source/Paper2D/Classes/PaperSprite.h
--- a/Engine/Plugins/2D/Paper2D/Source/Paper2D/Classes/PaperSprite.h
+++ b/Engine/Plugins/2D/Paper2D/Source/Paper2D/Classes/PaperSprite.h
@@ -195,7 +195,7 @@ public:
 	void ExtractSourceRegionFromTexturePoint(const FVector2D& Point);
 
 	void FindTextureBoundingBox(float AlphaThreshold, /*out*/ FVector2D& OutBoxPosition, /*out*/ FVector2D& OutBoxSize);
-	static void FindContours(const FIntPoint& ScanPos, const FIntPoint& ScanSize, float AlphaThreshold, UTexture2D* Texture, TArray< TArray<FIntPoint> >& OutPoints);
+	static void FindContours(const FIntPoint& ScanPos, const FIntPoint& ScanSize, float AlphaThreshold, float Detail, UTexture2D* Texture, TArray< TArray<FIntPoint> >& OutPoints);
 	static void ExtractRectsFromTexture(UTexture2D* Texture, TArray<FIntRect>& OutRects);
 	void BuildGeometryFromContours(FSpritePolygonCollection& GeomOwner);
 
diff --git a/Engine/Plugins/2D/Paper2D/Source/Paper2D/Classes/SpriteEditorOnlyTypes.h b/Engine/Plugins/2D/Paper2D/Source/Paper2D/Classes/SpriteEditorOnlyTypes.h
--- a/Engine/Plugins/2D/Paper2D/Source/Paper2D/Classes/SpriteEditorOnlyTypes.h
+++ b/Engine/Plugins/2D/Paper2D/Source/Paper2D/Classes/SpriteEditorOnlyTypes.h
@@ -107,6 +107,10 @@ struct FSpritePolygonCollection
 	UPROPERTY(Category=PolygonData, EditAnywhere, AdvancedDisplay)
 	float AlphaThreshold;
 
+	// Amount to detail to consider when shrink-wrapping (range 0..1, 0 = low detail, 1 = high detail)
+	UPROPERTY(Category = PolygonData, EditAnywhere, AdvancedDisplay)
+	float DetailAmount;
+
 	// This is the threshold below which multiple vertices will be merged together when doing shrink-wrapping.  Higher values result in fewer vertices.
 	UPROPERTY(Category=PolygonData, EditAnywhere, AdvancedDisplay)
 	float SimplifyEpsilon;
@@ -118,6 +122,7 @@ struct FSpritePolygonCollection
 		, bAvoidVertexMerging(false)
 		, AlphaThreshold(0.0f)
 		, SimplifyEpsilon(2.0f)
+		, DetailAmount(0.5f)
 	{
 	}
 
diff --git a/Engine/Plugins/2D/Paper2D/Source/Paper2D/Private/PaperGeomTools.cpp b/Engine/Plugins/2D/Paper2D/Source/Paper2D/Private/PaperGeomTools.cpp
--- a/Engine/Plugins/2D/Paper2D/Source/Paper2D/Private/PaperGeomTools.cpp
+++ b/Engine/Plugins/2D/Paper2D/Source/Paper2D/Private/PaperGeomTools.cpp
@@ -16,6 +16,19 @@ bool PaperGeomTools::IsPolygonWindingCCW(const TArray<FVector2D>& Points)
 	return (Sum < 0.0f);
 }
 
+bool PaperGeomTools::IsPolygonWindingCCW(const TArray<FIntPoint>& Points)
+{
+	int32 Sum = 0;
+	const int PointCount = Points.Num();
+	for (int PointIndex = 0; PointIndex < PointCount; ++PointIndex)
+	{
+		const FIntPoint& A = Points[PointIndex];
+		const FIntPoint& B = Points[(PointIndex + 1) % PointCount];
+		Sum += (B.X - A.X) * (B.Y + A.Y);
+	}
+	return (Sum < 0);
+}
+
 // Note: We need to simplify non-simple polygons before this
 static bool IsPolygonConvex(const TArray<FVector2D>& Points)
 {
@@ -58,7 +71,7 @@ static bool IsPointInPolygon(const FVector2D& TestPoint, const TArray<FVector2D>
 	return (FMath::Abs(AngleSum) > 0.001f);
 }
 
-static bool IsPointInTriangle(const FVector2D& TestPoint, const FVector2D& A, const FVector2D& B, const FVector2D& C)
+static bool IsAdditivePointInTriangle(const FVector2D& TestPoint, const FVector2D& A, const FVector2D& B, const FVector2D& C)
 {
 	const FVector2D AP = TestPoint - A;
 	const FVector2D BP = TestPoint - B;
@@ -75,6 +88,37 @@ static bool IsPointInTriangle(const FVector2D& TestPoint, const FVector2D& A, co
 	return true;
 }
 
+static float VectorSign(const FVector2D& Vec, const FVector2D& A, const FVector2D& B)
+{
+	return FMath::Sign((B.X - A.X) * (Vec.Y - A.Y) - (B.Y - A.Y) * (Vec.X - A.X));
+}
+
+// Returns true when the point is inside the triangle
+// Should not return true when the point is on one of the edges
+static bool IsPointInTriangle(const FVector2D& TestPoint, const FVector2D& A, const FVector2D& B, const FVector2D& C)
+{
+	float BA = VectorSign(B, A, TestPoint);
+	float CB = VectorSign(C, B, TestPoint);
+	float AC = VectorSign(A, C, TestPoint);
+
+	// point is in the same direction of all 3 tri edge lines
+	// must be inside, regardless of tri winding
+	return BA == CB && CB == AC;
+}
+
+// Returns true when the point is on the line segment limited by A and B
+static bool IsPointOnLineSegment(const FVector2D& TestPoint, const FVector2D& A, const FVector2D& B)
+{
+	FVector2D BA = B - A;
+	FVector2D PA = TestPoint - A;
+	float SizeSquaredBA = FVector2D::DotProduct(BA, BA);
+	float AreaCompareThreshold = 0.01f * SizeSquaredBA;
+	float ParallelogramArea = BA.X * PA.Y - BA.Y * PA.X;
+
+	return  TestPoint.X >= FMath::Min(A.X, B.X) && TestPoint.X <= FMath::Max(A.X, B.X) && // X within AB.X, including ON A or B
+			TestPoint.Y >= FMath::Min(A.Y, B.Y) && TestPoint.Y <= FMath::Max(A.Y, B.Y) && // Y within AB.Y, including ON A or B
+			FMath::Abs(ParallelogramArea) < AreaCompareThreshold; // Area is smaller than allowed epsilon = point on line
+}
 
 static void JoinSubtractiveToAdditive(TArray<FVector2D>& AdditivePoly, const TArray<FVector2D>& SubtractivePoly, const int AdditiveJoinIndex, const int SubtractiveJoinIndex)
 {
@@ -196,7 +240,7 @@ static void JoinMutuallyVisible(TArray<FVector2D>& AdditivePoly, const TArray<FV
 			continue;
 		}
 
-		if (IsPointInTriangle(AdditivePoly[AdditiveIndex], TriA, TriB, TriC))
+		if (IsAdditivePointInTriangle(AdditivePoly[AdditiveIndex], TriA, TriB, TriC))
 		{
 			const FVector2D MR = AdditivePoly[AdditiveIndex] - PointMaxX;
 			const float CosAngle = MR.X / MR.Size();
@@ -558,9 +602,11 @@ bool PaperGeomTools::TriangulatePoly(TArray<FVector2D>& OutTris, const TArray<FV
 					const FVector2D& CurrentVertex = PolyVerts[VertexIndex];
 					// Test indices first, then make sure we arent comparing identical points
 					// These might have been added by the convex / concave splitting
+					// If a point is not in the triangle, it may be on the new edge we're adding, which isn't allowed as 
+					// it will create a partition in the polygon
 					if (VertexIndex != AIndex && VertexIndex != BIndex && VertexIndex != CIndex &&
 						CurrentVertex != PolyVerts[AIndex] && CurrentVertex != PolyVerts[BIndex] && CurrentVertex != PolyVerts[CIndex] &&
-						IsPointInTriangle(CurrentVertex, PolyVerts[AIndex], PolyVerts[BIndex], PolyVerts[CIndex]))
+						(IsPointInTriangle(CurrentVertex, PolyVerts[AIndex], PolyVerts[BIndex], PolyVerts[CIndex]) || IsPointOnLineSegment(CurrentVertex, PolyVerts[CIndex], PolyVerts[AIndex])))
 					{
 						bFoundVertInside = true;
 						break;
@@ -586,6 +632,7 @@ bool PaperGeomTools::TriangulatePoly(TArray<FVector2D>& OutTris, const TArray<FV
 			// If we couldn't find an 'ear' it indicates something is bad with this polygon - discard triangles and return.
 			if (!bFoundEar)
 			{
+				// If we couldn't find an 'ear' it indicates something is bad with this polygon - discard triangles and return.
 				UE_LOG(LogPaper2D, Log, TEXT("Triangulation of poly failed."));
 				OutTris.Empty();
 				return false;
diff --git a/Engine/Plugins/2D/Paper2D/Source/Paper2D/Private/PaperGeomTools.h b/Engine/Plugins/2D/Paper2D/Source/Paper2D/Private/PaperGeomTools.h
--- a/Engine/Plugins/2D/Paper2D/Source/Paper2D/Private/PaperGeomTools.h
+++ b/Engine/Plugins/2D/Paper2D/Source/Paper2D/Private/PaperGeomTools.h
@@ -13,6 +13,7 @@ public:
 	// Returns true if the points forming a polygon have CCW winding
 	// Returns true if the polygon isn't valid
 	static bool IsPolygonWindingCCW(const TArray<FVector2D>& Points);
+	static bool IsPolygonWindingCCW(const TArray<FIntPoint>& Points);
 
 	// Checks that these polygons can be successfully triangulated	
 	static bool ArePolygonsValid(const TArray<TArray<FVector2D>>& Polygons);
diff --git a/Engine/Plugins/2D/Paper2D/Source/Paper2D/Private/PaperSprite.cpp b/Engine/Plugins/2D/Paper2D/Source/Paper2D/Private/PaperSprite.cpp
--- a/Engine/Plugins/2D/Paper2D/Source/Paper2D/Private/PaperSprite.cpp
+++ b/Engine/Plugins/2D/Paper2D/Source/Paper2D/Private/PaperSprite.cpp
@@ -28,11 +28,12 @@ void RemoveCollinearPoints(TArray<FIntPoint>& PointList)
 	{
 		return;
 	}
-
-	for (int32 VertexIndex = 1; VertexIndex < PointList.Num(); )
+	
+	// Wrap around to get the final pair of vertices (N-1, 0, 1)
+	for (int32 VertexIndex = 1; VertexIndex <= PointList.Num() && PointList.Num() >= 3; )
 	{
 		const FVector2D A(PointList[VertexIndex-1]);
-		const FVector2D B(PointList[VertexIndex]);
+		const FVector2D B(PointList[VertexIndex % PointList.Num()]);
 		const FVector2D C(PointList[(VertexIndex+1) % PointList.Num()]);
 
 		// Determine if the area of the triangle ABC is zero (if so, they're collinear)
@@ -41,7 +42,7 @@ void RemoveCollinearPoints(TArray<FIntPoint>& PointList)
 		if (FMath::Abs(AreaABC) < KINDA_SMALL_NUMBER)
 		{
 			// Remove B
-			PointList.RemoveAt(VertexIndex);
+			PointList.RemoveAt(VertexIndex % PointList.Num());
 		}
 		else
 		{
@@ -141,10 +142,76 @@ struct FDouglasPeuckerSimplifier
 	}
 };
 
+static void BruteForceSimplifier(TArray<FIntPoint>& Points, float Epsilon)
+{
+	float FlatEdgeDistanceThreshold = (int)(Epsilon * Epsilon);
+
+	// Run through twice to remove remnants from staircase artifacts
+	for (int Pass = 0; Pass < 2; ++Pass)
+	{
+		for (int I = 0; I < Points.Num() && Points.Num() > 3; ++I)
+		{
+			int StartRemoveIndex = (I + 1) % Points.Num();
+			int EndRemoveIndex = StartRemoveIndex;
+			FIntPoint& A = Points[I];
+			// Keep searching to find if any of the vector rejections fail in subsequent points on the polygon
+			// A B C D E F (eg. when testing A B C, test rejection for BA, CA)
+			// When testing A E F, test rejection for AB-AF, AC-AF, AD-AF, AE-AF
+			// When one of these fails we discard all verts between A and one before the current vertex being tested
+			for (int J = I; J < Points.Num(); ++J)
+			{
+				int IndexC = (J + 2) % Points.Num();
+				FIntPoint& C = Points[IndexC];
+				bool bSmallOffsetFailed = false;
+
+				for (int K = I; K <= J && !bSmallOffsetFailed; ++K)
+				{
+					int IndexB = (K + 1) % Points.Num();
+					FIntPoint& B = Points[IndexB];
+
+					FVector2D CA = C - A;
+					FVector2D BA = B - A;
+					FVector2D Rejection_BA_CA = BA - (FVector2D::DotProduct(BA, CA) / FVector2D::DotProduct(CA, CA)) * CA;
+					float RejectionLengthSquared = Rejection_BA_CA.SizeSquared();
+					// If any of the points is behind the polyline up till now, it gets rejected. Staircase artefacts are handled in a second pass
+					if (RejectionLengthSquared > FlatEdgeDistanceThreshold || FVector2D::CrossProduct(CA, BA) < 0) 
+					{
+						bSmallOffsetFailed = true;
+						break;
+					}
+				}
+
+				if (bSmallOffsetFailed)
+				{
+					break;
+				}
+				else
+				{
+					EndRemoveIndex = (EndRemoveIndex + 1) % Points.Num();
+				}
+			}
+
+			// Remove the vertices that we deemed "too flat"
+			if (EndRemoveIndex > StartRemoveIndex)
+			{
+				Points.RemoveAt(StartRemoveIndex, EndRemoveIndex - StartRemoveIndex);
+			}
+			else if (EndRemoveIndex < StartRemoveIndex)
+			{
+				Points.RemoveAt(StartRemoveIndex, Points.Num() - StartRemoveIndex);
+				Points.RemoveAt(0, EndRemoveIndex);
+				// The search has wrapped around, no more vertices to test
+				break;
+			}
+		}
+	} // Pass
+}
+
 void SimplifyPoints(TArray<FIntPoint>& Points, float Epsilon)
 {
-	FDouglasPeuckerSimplifier Simplifier(Points, Epsilon);
-	Simplifier.Execute(Points);
+//	FDouglasPeuckerSimplifier Simplifier(Points, Epsilon);
+//	Simplifier.Execute(Points);
+	BruteForceSimplifier(Points, Epsilon);
 }
 
 	
@@ -964,6 +1031,15 @@ void UPaperSprite::FindTextureBoundingBox(float AlphaThreshold, /*out*/ FVector2
 	OutBoxPosition.Y = TopBound;
 }
 
+// Get a divisor ("pixel" size) from the "detail" parameter
+// Size is fed in for possible changes later
+static int32 GetDivisorFromDetail(const FIntPoint& Size, float Detail)
+{
+	//@TODO: Consider MaxSize somehow when deciding divisor
+	//int32 MaxSize = FMath::Max(Size.X, Size.Y);
+	return  FMath::Lerp(8, 1, FMath::Clamp(Detail, 0.0f, 1.0f));
+}
+
 void UPaperSprite::BuildGeometryFromContours(FSpritePolygonCollection& GeomOwner)
 {
 	// First trim the image to the tight fitting bounding box (the other pixels can't matter)
@@ -976,18 +1052,20 @@ void UPaperSprite::BuildGeometryFromContours(FSpritePolygonCollection& GeomOwner
 
 	// Find the contours
 	TArray< TArray<FIntPoint> > Contours;
-	FindContours(InitialPos, InitialSize, GeomOwner.AlphaThreshold, SourceTexture, /*out*/ Contours);
 
-	//@TODO: Remove fully enclosed contours - may not be needed if we're tracing only outsides
+	// DK: FindContours only returns positive contours, i.e. outsides
+	// Contour generation is simplified in FindContours by downscaling the detail prior to generating contour data
+	FindContours(InitialPos, InitialSize, GeomOwner.AlphaThreshold, GeomOwner.DetailAmount, SourceTexture, /*out*/ Contours);
 
 	// Convert the contours into geometry
 	GeomOwner.Polygons.Empty();
 	for (int32 ContourIndex = 0; ContourIndex < Contours.Num(); ++ContourIndex)
 	{
-		//@TODO: Optimize contours
 		TArray<FIntPoint>& Contour = Contours[ContourIndex];
 
-		SimplifyPoints(Contour, GeomOwner.SimplifyEpsilon);
+		// Scale the simplification epsilon by the size we know the pixels will be
+		int Divisor = GetDivisorFromDetail(InitialSize, GeomOwner.DetailAmount);
+		SimplifyPoints(Contour, GeomOwner.SimplifyEpsilon * Divisor);
 
 		if (Contour.Num() > 0)
 		{
@@ -1006,9 +1084,100 @@ void UPaperSprite::BuildGeometryFromContours(FSpritePolygonCollection& GeomOwner
 }
 
 
-// findme
+static void TraceContour(TArray<FIntPoint>& Result, const TArray<FIntPoint>& Points)
+{
+	const int PointCount = Points.Num();
+	if (PointCount < 2)
+	{
+		return;
+	}
 
-void UPaperSprite::FindContours(const FIntPoint& ScanPos, const FIntPoint& ScanSize, float AlphaThreshold, UTexture2D* Texture, TArray< TArray<FIntPoint> >& OutPoints)
+	int CurrentX = (int)Points[0].X;
+	int CurrentY = (int)Points[0].Y;
+	int CurrentDirection = 0;
+	int FirstDx = (int)Points[1].X - CurrentX;
+	int FirstDy = (int)Points[1].Y - CurrentY;
+	
+	if (FirstDx == 1 && FirstDy == 0) CurrentDirection = 0;
+	else if (FirstDx == 1 && FirstDy == 1) CurrentDirection = 1;
+	else if (FirstDx == 0 && FirstDy == 1) CurrentDirection = 1;
+	else if (FirstDx == -1 && FirstDy == 1) CurrentDirection = 2;
+	else if (FirstDx == -1 && FirstDy == 0) CurrentDirection = 2;
+	else if (FirstDx == -1 && FirstDy == -1) CurrentDirection = 3;
+	else if (FirstDx == 0 && FirstDy == -1) CurrentDirection = 3;
+	else if (FirstDx == 1 && FirstDy == -1) CurrentDirection = 0;
+
+	int CurrentPointIndex = 0;
+
+	const int StartX = CurrentX;
+	const int StartY = CurrentY;
+	const int StartDirection = CurrentDirection;
+
+	static const int DirectionDx[] = { 1, 0, -1, 0 };
+	static const int DirectionDy[] = { 0, 1, 0, -1 };
+
+	bool bFinished = false;
+	while (!bFinished)
+	{
+		const FIntPoint& NextPoint = Points[(CurrentPointIndex + 1) % PointCount];
+		const int NextDx = (int)NextPoint.X - CurrentX;
+		const int NextDy = (int)NextPoint.Y - CurrentY;
+
+		int LeftDirection = (CurrentDirection + 3) % 4;
+		int CurrentDx = DirectionDx[CurrentDirection];
+		int CurrentDy = DirectionDy[CurrentDirection];
+		int LeftDx = DirectionDx[LeftDirection];
+		int LeftDy = DirectionDy[LeftDirection];
+		bool bDidMove = true;
+		if (NextDx != 0 || NextDy != 0)
+		{
+			if (NextDx == LeftDx && NextDy == LeftDy)
+			{
+				// Space to the left, turn left and move forwards
+				CurrentDirection = LeftDirection;
+				CurrentX += LeftDx;
+				CurrentY += LeftDy;
+			}
+			else
+			{
+				// Wall to the left. Add the corner vertex to our output.
+				Result.Add(FIntPoint((int)((float)CurrentX + 0.5f + (float)(CurrentDx + LeftDx) * 0.5f), (int)((float)CurrentY + 0.5f + (float)(CurrentDy + LeftDy) * 0.5f)));
+				if (NextDx == CurrentDx && NextDy == CurrentDy)
+				{
+					// Move forward
+					CurrentX += CurrentDx;
+					CurrentY += CurrentDy;
+				}
+				else if (NextDx == CurrentDx + LeftDx && NextDy == CurrentDy + LeftDy)
+				{
+					// Move forward, turn left, move forwards again
+					CurrentX += CurrentDx;
+					CurrentY += CurrentDy;
+					CurrentDirection = LeftDirection;
+					CurrentX += LeftDx;
+					CurrentY += LeftDy;
+				}
+				else
+				{
+					// Turn right
+					CurrentDirection = (CurrentDirection + 1) % 4;
+					bDidMove = false;
+				}
+			}
+		}
+		if (bDidMove)
+		{
+			++CurrentPointIndex;
+		}
+
+		if (CurrentX == StartX && CurrentY == StartY && CurrentDirection == StartDirection)
+		{
+			bFinished = true;
+		}
+	}
+}
+
+void UPaperSprite::FindContours(const FIntPoint& ScanPos, const FIntPoint& ScanSize, float AlphaThreshold, float Detail, UTexture2D* Texture, TArray< TArray<FIntPoint> >& OutPoints)
 {
 	OutPoints.Empty();
 
@@ -1017,11 +1186,6 @@ void UPaperSprite::FindContours(const FIntPoint& ScanPos, const FIntPoint& ScanS
 		return;
 	}
 
-	const int32 LeftBound = ScanPos.X;
-	const int32 RightBound = ScanPos.X + ScanSize.X - 1;
-	const int32 TopBound = ScanPos.Y;
-	const int32 BottomBound = ScanPos.Y + ScanSize.Y - 1;
-
 	// Neighborhood array (clockwise starting at -X,-Y; assuming prev is at -X)
 	const int32 NeighborX[] = {-1, 0,+1,+1,+1, 0,-1,-1};
 	const int32 NeighborY[] = {-1,-1,-1, 0,+1,+1,+1, 0};
@@ -1041,13 +1205,30 @@ void UPaperSprite::FindContours(const FIntPoint& ScanPos, const FIntPoint& ScanS
 	};
 
 	int32 AlphaThresholdInt = FMath::Clamp<int32>(AlphaThreshold * 255, 0, 255);
-	FBitmap SourceBitmap(Texture, AlphaThresholdInt);
-	if (SourceBitmap.IsValid())
+
+	FBitmap FullSizeBitmap(Texture, AlphaThresholdInt);
+	if (FullSizeBitmap.IsValid())
 	{
-		checkSlow((LeftBound >= 0) && (TopBound >= 0) && (RightBound < SourceBitmap.Width) && (BottomBound < SourceBitmap.Height));
+		const int32 DownsampleAmount = GetDivisorFromDetail(ScanSize, Detail);
+
+		FBitmap SourceBitmap((ScanSize.X + DownsampleAmount - 1) / DownsampleAmount, (ScanSize.Y + DownsampleAmount - 1) / DownsampleAmount, 0);
+		for (int32 Y = 0; Y < ScanSize.Y; ++Y)
+		{
+			for (int32 X = 0; X < ScanSize.X; ++X)
+			{
+				SourceBitmap.SetPixel(X / DownsampleAmount, Y / DownsampleAmount, SourceBitmap.GetPixel(X / DownsampleAmount, Y / DownsampleAmount) | FullSizeBitmap.GetPixel(ScanPos.X + X, ScanPos.Y + Y));
+			}
+		}
+
+		const int32 LeftBound = 0;
+		const int32 RightBound = SourceBitmap.Width - 1;
+		const int32 TopBound = 0;
+		const int32 BottomBound = SourceBitmap.Height - 1;
+
+		//checkSlow((LeftBound >= 0) && (TopBound >= 0) && (RightBound < SourceBitmap.Width) && (BottomBound < SourceBitmap.Height));
 
 		// Create the 'output' boundary image
-		FBoundaryImage BoundaryImage(ScanPos, ScanSize);
+		FBoundaryImage BoundaryImage(FIntPoint(0, 0), FIntPoint(SourceBitmap.Width, SourceBitmap.Height));
 
 		bool bInsideBoundary = false;
 
@@ -1077,11 +1258,12 @@ void UPaperSprite::FindContours(const FIntPoint& ScanPos, const FIntPoint& ScanS
 					else if (bIsFilledPixel)
 					{
 						// Create the output chain we'll build from the boundary image
-						TArray<FIntPoint>& Contour = *new (OutPoints) TArray<FIntPoint>();
+						//TArray<FIntPoint>& Contour = *new (OutPoints) TArray<FIntPoint>();
+						TArray<FIntPoint> Contour;
 
 						// Moving into an undiscovered boundary
 						BoundaryImage.SetPixel(X, Y, 1);
-						new (Contour) FIntPoint(X, Y);
+						new (Contour)FIntPoint(X, Y);
 
 						// Current pixel
 						int32 NeighborPhase = 0;
@@ -1099,8 +1281,8 @@ void UPaperSprite::FindContours(const FIntPoint& ScanPos, const FIntPoint& ScanS
 							const bool bTestPixelInsideBounds = (CX >= LeftBound && CX <= RightBound && CY >= TopBound && CY <= BottomBound);
 							const bool bTestPixelPasses = bTestPixelInsideBounds && SourceBitmap.GetPixel(CX, CY) != 0;
 
-							UE_LOG(LogPaper2D, Log, TEXT("Outer P(%d,%d), C(%d,%d) Ph%d %s"),
-								PX, PY, CX, CY, NeighborPhase, bTestPixelPasses ? TEXT("[BORDER]") : TEXT("[]"));
+							//UE_LOG(LogPaper2D, Log, TEXT("Outer P(%d,%d), C(%d,%d) Ph%d %s"),
+							//	PX, PY, CX, CY, NeighborPhase, bTestPixelPasses ? TEXT("[BORDER]") : TEXT("[]"));
 
 							if (bTestPixelPasses)
 							{
@@ -1121,7 +1303,7 @@ void UPaperSprite::FindContours(const FIntPoint& ScanPos, const FIntPoint& ScanS
 // 								}
 
 								BoundaryImage.SetPixel(CX, CY, NeighborPhase+1);
-								new (Contour) FIntPoint(CX, CY);
+								new (Contour)FIntPoint(CX, CY);
 
 								PX = CX;
 								PY = CY;
@@ -1142,9 +1324,33 @@ void UPaperSprite::FindContours(const FIntPoint& ScanPos, const FIntPoint& ScanS
 								}
 							}
 						}
-					
+						
+						// Trace the contour shape creating polygon edges
+						TArray<FIntPoint>& ContourPoly = *new (OutPoints)TArray<FIntPoint>();
+						TraceContour(/*out*/ContourPoly, Contour);
+
 						// Remove collinear points from the result
-						RemoveCollinearPoints(/*inout*/ Contour);
+						RemoveCollinearPoints(/*inout*/ ContourPoly);
+
+						if (!PaperGeomTools::IsPolygonWindingCCW(ContourPoly))
+						{
+							// Remove newly added polygon, we don't support holes just yet
+							OutPoints.RemoveAt(OutPoints.Num() - 1);
+						}
+						else
+						{
+							for (int ContourPolyIndex = 0; ContourPolyIndex < ContourPoly.Num(); ++ContourPolyIndex)
+							{
+								// Rescale and recenter contour poly
+								FIntPoint RescaledPoint = ScanPos + ContourPoly[ContourPolyIndex] * DownsampleAmount;
+
+								// Make sure rescaled point doesn't exceed the original max bounds
+								RescaledPoint.X = FMath::Min(RescaledPoint.X, ScanPos.X + ScanSize.X);
+								RescaledPoint.Y = FMath::Min(RescaledPoint.Y, ScanPos.Y + ScanSize.Y);
+
+								ContourPoly[ContourPolyIndex] = RescaledPoint;
+							}
+						}
 					}
 				}
 			}
