Paper2D: PaperTerrainComponent, Added support for overlapping terrain sections and material draw order
Paper2D: PaperTerrainComponent, Improved overlap in tight corners by unwrapping sprite when its twisted (> 90 deg)

diff --git a/Engine/Plugins/2D/Paper2D/Source/Paper2D/Classes/PaperTerrainComponent.h b/Engine/Plugins/2D/Paper2D/Source/Paper2D/Classes/PaperTerrainComponent.h
--- a/Engine/Plugins/2D/Paper2D/Source/Paper2D/Classes/PaperTerrainComponent.h
+++ b/Engine/Plugins/2D/Paper2D/Source/Paper2D/Classes/PaperTerrainComponent.h
@@ -7,10 +7,11 @@
 
 struct FPaperTerrainMaterialRule;
 
-struct FPaperTerrainMaterialPair
+struct FPaperTerrainSpriteGeometry
 {
 	TArray<FSpriteDrawCallRecord> Records;
 	UMaterialInterface* Material;
+	int32 DrawOrder;
 };
 
 struct FTerrainSpriteStamp
@@ -63,6 +64,10 @@ public:
 	UPROPERTY(Category=Sprite, EditAnywhere)
 	int32 RandomSeed;
 
+	/** The overlap amount between segments */
+	UPROPERTY(Category=Sprite, EditAnywhere)
+	float SegmentOverlapAmount;
+
 protected:
 	/** The color of the terrain (passed to the sprite material as a vertex color) */
 	UPROPERTY(Category=Sprite, BlueprintReadOnly, Interp)
@@ -116,7 +121,7 @@ protected:
 	UPROPERTY(Transient, DuplicateTransient)
 	class UBodySetup* CachedBodySetup;
 
-	TArray<FPaperTerrainMaterialPair> GeneratedSpriteGeometry;
+	TArray<FPaperTerrainSpriteGeometry> GeneratedSpriteGeometry;
 	
 	FTransform GetTransformAtDistance(float InDistance) const;
 };
diff --git a/Engine/Plugins/2D/Paper2D/Source/Paper2D/Classes/PaperTerrainMaterial.h b/Engine/Plugins/2D/Paper2D/Source/Paper2D/Classes/PaperTerrainMaterial.h
--- a/Engine/Plugins/2D/Paper2D/Source/Paper2D/Classes/PaperTerrainMaterial.h
+++ b/Engine/Plugins/2D/Paper2D/Source/Paper2D/Classes/PaperTerrainMaterial.h
@@ -38,6 +38,10 @@ struct FPaperTerrainMaterialRule
 	UPROPERTY(Category=Sprite, EditAnywhere)
 	float CollisionOffset;
 
+	// Specify a draw order for different materials in a spline. Smaller draw orders are drawn first, negative values are allowed.
+	UPROPERTY(Category=Sprite, EditAnywhere)
+	int32 DrawOrder;
+
 #if WITH_EDITORONLY_DATA
 	// Readable description for the rule (unused anywhere, just for clarity when editing the material)
 	UPROPERTY(Category = Sprite, EditAnywhere)
diff --git a/Engine/Plugins/2D/Paper2D/Source/Paper2D/Private/Terrain/PaperTerrainComponent.cpp b/Engine/Plugins/2D/Paper2D/Source/Paper2D/Private/Terrain/PaperTerrainComponent.cpp
--- a/Engine/Plugins/2D/Paper2D/Source/Paper2D/Private/Terrain/PaperTerrainComponent.cpp
+++ b/Engine/Plugins/2D/Paper2D/Source/Paper2D/Private/Terrain/PaperTerrainComponent.cpp
@@ -51,23 +51,23 @@ void FTerrainSegment::RepositionStampsToFillSpace()
 class FPaperTerrainSceneProxy : public FPaperRenderSceneProxy
 {
 public:
-	FPaperTerrainSceneProxy(const UPaperTerrainComponent* InComponent, const TArray<FPaperTerrainMaterialPair>& InDrawingData);
+	FPaperTerrainSceneProxy(const UPaperTerrainComponent* InComponent, const TArray<FPaperTerrainSpriteGeometry>& InDrawingData);
 
 protected:
-	TArray<FPaperTerrainMaterialPair> DrawingData;
+	TArray<FPaperTerrainSpriteGeometry> DrawingData;
 protected:
 	// FPaperRenderSceneProxy interface
 	virtual void GetDynamicMeshElementsForView(const FSceneView* View, int32 ViewIndex, bool bUseOverrideColor, const FLinearColor& OverrideColor, FMeshElementCollector& Collector) const override;
 	// End of FPaperRenderSceneProxy interface
 };
 
-FPaperTerrainSceneProxy::FPaperTerrainSceneProxy(const UPaperTerrainComponent* InComponent, const TArray<FPaperTerrainMaterialPair>& InDrawingData)
+FPaperTerrainSceneProxy::FPaperTerrainSceneProxy(const UPaperTerrainComponent* InComponent, const TArray<FPaperTerrainSpriteGeometry>& InDrawingData)
 	: FPaperRenderSceneProxy(InComponent)
 {
 	DrawingData = InDrawingData;
 
 	// Combine the material relevance for all materials
-	for (const FPaperTerrainMaterialPair& Batch : DrawingData)
+	for (const FPaperTerrainSpriteGeometry& Batch : DrawingData)
 	{
 		const UMaterialInterface* MaterialInterface = (Batch.Material != nullptr) ? Batch.Material : UMaterial::GetDefaultMaterial(MD_Surface);
 		MaterialRelevance |= MaterialInterface->GetRelevance_Concurrent(GetScene().GetFeatureLevel());
@@ -76,7 +76,7 @@ FPaperTerrainSceneProxy::FPaperTerrainSceneProxy(const UPaperTerrainComponent* I
 
 void FPaperTerrainSceneProxy::GetDynamicMeshElementsForView(const FSceneView* View, int32 ViewIndex, bool bUseOverrideColor, const FLinearColor& OverrideColor, FMeshElementCollector& Collector) const
 {
-	for (const FPaperTerrainMaterialPair& Batch : DrawingData)
+	for (const FPaperTerrainSpriteGeometry& Batch : DrawingData)
 	{
 		if (Batch.Material != nullptr)
 		{
@@ -91,6 +91,7 @@ void FPaperTerrainSceneProxy::GetDynamicMeshElementsForView(const FSceneView* Vi
 UPaperTerrainComponent::UPaperTerrainComponent(const FObjectInitializer& ObjectInitializer)
 	: Super(ObjectInitializer)
 	, bFilledSpline(true)
+	, SegmentOverlapAmount(100.0f)
 	, TerrainColor(FLinearColor::White)
 	, ReparamStepsPerSegment(8)
 	, SpriteCollisionDomain(ESpriteCollisionMode::Use3DPhysics)
@@ -162,7 +163,7 @@ FBoxSphereBounds UPaperTerrainComponent::CalcBounds(const FTransform& LocalToWor
 	{
 		FBox BoundingBox(ForceInit);
 
-		for (const FPaperTerrainMaterialPair& DrawCall : GeneratedSpriteGeometry)
+		for (const FPaperTerrainSpriteGeometry& DrawCall : GeneratedSpriteGeometry)
 		{
 			for (const FSpriteDrawCallRecord& Record : DrawCall.Records)
 			{
@@ -444,6 +445,12 @@ void UPaperTerrainComponent::OnSplineEdited()
 					else
 					{
 						ActiveSegment->EndTime = CurrentTime;
+						
+						// Segment is too small, delete it
+						if (ActiveSegment->EndTime < ActiveSegment->StartTime + 2.0f * SegmentOverlapAmount)
+						{
+							Segments.Pop(false);
+						}
 
 						ActiveSegment = new (Segments)FTerrainSegment();
 						ActiveSegment->StartTime = CurrentTime;
@@ -456,6 +463,13 @@ void UPaperTerrainComponent::OnSplineEdited()
 			}
 		}
 
+		// Account for overlap
+		for (FTerrainSegment& Segment : Segments)
+		{
+			Segment.StartTime -= SegmentOverlapAmount;
+			Segment.EndTime += SegmentOverlapAmount;
+		}
+
 		// Convert those segments to actual geometry
 		for (FTerrainSegment& Segment : Segments)
 		{
@@ -579,7 +593,7 @@ void UPaperTerrainComponent::OnSplineEdited()
 			if (TerrainMaterial->InteriorFill != nullptr)
 			{
 				const UPaperSprite* FillSprite = TerrainMaterial->InteriorFill;
-				FPaperTerrainMaterialPair& MaterialBatch = *new (GeneratedSpriteGeometry)FPaperTerrainMaterialPair(); //@TODO: Look up the existing one instead
+				FPaperTerrainSpriteGeometry& MaterialBatch = *new (GeneratedSpriteGeometry)FPaperTerrainSpriteGeometry(); //@TODO: Look up the existing one instead
 				MaterialBatch.Material = FillSprite->GetDefaultMaterial();
 
 				FSpriteDrawCallRecord& FillDrawCall = *new (MaterialBatch.Records) FSpriteDrawCallRecord();
@@ -736,16 +750,32 @@ void UPaperTerrainComponent::SpawnSegments(const TArray<FTerrainSegment>& Terrai
 
 			if (NewSprite)
 			{
-				FPaperTerrainMaterialPair& MaterialBatch = *new (GeneratedSpriteGeometry) FPaperTerrainMaterialPair(); //@TODO: Look up the existing one instead
+				FPaperTerrainSpriteGeometry& MaterialBatch = *new (GeneratedSpriteGeometry) FPaperTerrainSpriteGeometry(); //@TODO: Look up the existing one instead
 				MaterialBatch.Material = NewSprite->GetDefaultMaterial();
+				MaterialBatch.DrawOrder = Rule->DrawOrder;
 
 				FSpriteDrawCallRecord& Record = *new (MaterialBatch.Records) FSpriteDrawCallRecord();
 				Record.BuildFromSprite(NewSprite);
 				Record.Color = TerrainColor;
 
+				// Work out if the sprite is likely to be bent > X deg (folded over itself)
+				const FVector ForwardVector(1, 0, 0);
+				const FTransform LocalTransformAtBack(GetTransformAtDistance(Position - 0.5f * NominalWidth * HorizontalScale));
+				FVector StartForwardVector = LocalTransformAtBack.TransformVector(ForwardVector).GetSafeNormal();
+				const FTransform LocalTransformAtFront(GetTransformAtDistance(Position + 0.5f * NominalWidth * HorizontalScale));
+				FVector EndForwardVector = LocalTransformAtFront.TransformVector(ForwardVector).GetSafeNormal();
+				bool bIsSpriteBent = FVector::DotProduct(StartForwardVector, EndForwardVector) < 0.0f;// 0.7071f; // (45deg looks worse)
+
 				for (FVector4& XYUV : Record.RenderVerts)
 				{
-					const FTransform LocalTransformAtX(GetTransformAtDistance(Position + (XYUV.X * HorizontalScale)));
+					FTransform LocalTransformAtX(GetTransformAtDistance(Position + (XYUV.X * HorizontalScale)));
+
+					// When the quad is overly bent, inherit rotation from the start of the quad to unfold it
+					if (bIsSpriteBent)
+					{
+						LocalTransformAtX.SetRotation(LocalTransformAtFront.GetRotation());
+					}
+
 					const FVector SourceVector = (XYUV.Y * PaperAxisY);
 					const FVector NewVector = LocalTransformAtX.TransformPosition(SourceVector);
 
@@ -759,6 +789,9 @@ void UPaperTerrainComponent::SpawnSegments(const TArray<FTerrainSegment>& Terrai
 		}
 	}
 
+	//@TODO: Sort by draw order first, materials next - Merge batches with the same material
+	GeneratedSpriteGeometry.Sort([](const FPaperTerrainSpriteGeometry& A, const FPaperTerrainSpriteGeometry& B) { return B.DrawOrder > A.DrawOrder; });
+
 #ifdef USE_SIMPLIFIED_POLYGON_COLLIDERS_FOR_SEGMENTS
 	// For whatever is remaining
 	if (CollisionPolygonPoints.Num() > 0)
@@ -910,7 +943,7 @@ void UPaperTerrainComponent::SetTerrainColor(FLinearColor NewColor)
 		TerrainColor = NewColor;
 
 		// Update the color in the game-thread copy of the render geometry
-		for (FPaperTerrainMaterialPair& Batch : GeneratedSpriteGeometry)
+		for (FPaperTerrainSpriteGeometry& Batch : GeneratedSpriteGeometry)
 		{
 			for (FSpriteDrawCallRecord& DrawCall : Batch.Records)
 			{
