Fixed bug where opening the sprite editor closes the property matrix window if the asset is open in it.

diff --git a/Engine/Plugins/2D/Paper2D/Source/Paper2DEditor/Private/SpriteEditor/SpriteEditor.cpp b/Engine/Plugins/2D/Paper2D/Source/Paper2DEditor/Private/SpriteEditor/SpriteEditor.cpp
--- a/Engine/Plugins/2D/Paper2D/Source/Paper2DEditor/Private/SpriteEditor/SpriteEditor.cpp
+++ b/Engine/Plugins/2D/Paper2D/Source/Paper2DEditor/Private/SpriteEditor/SpriteEditor.cpp
@@ -361,7 +361,6 @@ void FSpriteEditor::UnregisterTabSpawners(const TSharedRef<class FTabManager>& T
 
 void FSpriteEditor::InitSpriteEditor(const EToolkitMode::Type Mode, const TSharedPtr< class IToolkitHost >& InitToolkitHost, class UPaperSprite* InitSprite)
 {
-	FAssetEditorManager::Get().CloseOtherEditors(InitSprite, this);
 	SpriteBeingEdited = InitSprite;
 
 	FSpriteEditorCommands::Register();
