Paper2D: Allow tile maps to be created always (removing the bEnableTileMapEditing setting)

diff --git a/Engine/Plugins/2D/Paper2D/Source/Paper2D/Classes/PaperRuntimeSettings.h b/Engine/Plugins/2D/Paper2D/Source/Paper2D/Classes/PaperRuntimeSettings.h
--- a/Engine/Plugins/2D/Paper2D/Source/Paper2D/Classes/PaperRuntimeSettings.h
+++ b/Engine/Plugins/2D/Paper2D/Source/Paper2D/Classes/PaperRuntimeSettings.h
@@ -30,10 +30,6 @@ class PAPER2D_API UPaperRuntimeSettings : public UObject
 	UPROPERTY(EditAnywhere, config, Category=Experimental)
 	bool bEnableSpriteAtlasGroups;
 
-	// Enables experimental *incomplete and unsupported* tile map editing. Note: You need to restart the editor when enabling this setting for the change to fully take effect.
-	UPROPERTY(EditAnywhere, config, Category=Experimental)
-	bool bEnableTileMapEditing;
-
 	// Enables experimental *incomplete and unsupported* 2D terrain spline editing. Note: You need to restart the editor when enabling this setting for the change to fully take effect.
 	UPROPERTY(EditAnywhere, config, Category=Experimental)
 	bool bEnableTerrainSplineEditing;
diff --git a/Engine/Plugins/2D/Paper2D/Source/Paper2D/Private/PaperRuntimeSettings.cpp b/Engine/Plugins/2D/Paper2D/Source/Paper2D/Private/PaperRuntimeSettings.cpp
--- a/Engine/Plugins/2D/Paper2D/Source/Paper2D/Private/PaperRuntimeSettings.cpp
+++ b/Engine/Plugins/2D/Paper2D/Source/Paper2D/Private/PaperRuntimeSettings.cpp
@@ -10,7 +10,6 @@ UPaperRuntimeSettings::UPaperRuntimeSettings(const FObjectInitializer& ObjectIni
 	: Super(ObjectInitializer)
 	, DefaultPixelsPerUnrealUnit(2.56f)
 	, bEnableSpriteAtlasGroups(false)
-	, bEnableTileMapEditing(false)
 	, bEnableTerrainSplineEditing(false)
 {
 #if WITH_EDITORONLY_DATA
diff --git a/Engine/Plugins/2D/Paper2D/Source/Paper2DEditor/Private/Paper2DEditorModule.cpp b/Engine/Plugins/2D/Paper2D/Source/Paper2DEditor/Private/Paper2DEditorModule.cpp
--- a/Engine/Plugins/2D/Paper2D/Source/Paper2DEditor/Private/Paper2DEditorModule.cpp
+++ b/Engine/Plugins/2D/Paper2D/Source/Paper2DEditor/Private/Paper2DEditorModule.cpp
@@ -137,7 +137,11 @@ public:
 		//@TODO: PAPER2D: UThumbnailManager::Get().RegisterCustomRenderer(UPaperTileMap::StaticClass(), UPaperTileMapThumbnailRenderer::StaticClass());
 
 		// Register the editor modes
-		UpdateTileMapEditorModeInstallation();
+		FEditorModeRegistry::Get().RegisterMode<FEdModeTileMap>(
+			FEdModeTileMap::EM_TileMap,
+			LOCTEXT("TileMapEditMode", "Tile Map Editor"),
+			FSlateIcon(),
+			false);
 
 		// Integrate Paper2D actions associated with existing engine types (e.g., Texture2D) into the content browser
 		FPaperContentBrowserExtensions::InstallHooks();
@@ -227,7 +231,7 @@ private:
 		}
 		else if (UPaperRuntimeSettings* Settings = Cast<UPaperRuntimeSettings>(ObjectBeingModified))
 		{
-			UpdateTileMapEditorModeInstallation();
+			// Handle changes to experimental flags here
 		}
 	}
 
@@ -250,25 +254,6 @@ private:
 			SettingsModule->UnregisterSettings("Project", "Plugins", "Paper2D");
 		}
 	}
-
-	// Installs or uninstalls the tile map editing mode depending on settings
-	void UpdateTileMapEditorModeInstallation()
-	{
-		const bool bAlreadyRegistered = FEditorModeRegistry::Get().GetFactoryMap().Contains(FEdModeTileMap::EM_TileMap);
-		const bool bShouldBeRegistered = GetDefault<UPaperRuntimeSettings>()->bEnableTileMapEditing;
-		if (bAlreadyRegistered && !bShouldBeRegistered)
-		{
-			FEditorModeRegistry::Get().UnregisterMode(FEdModeTileMap::EM_TileMap);
-		}
-		else if (!bAlreadyRegistered && bShouldBeRegistered)
-		{
-			FEditorModeRegistry::Get().RegisterMode<FEdModeTileMap>(
-				FEdModeTileMap::EM_TileMap,
-				LOCTEXT("TileMapEditMode", "Tile Map Editor"),
-				FSlateIcon(),
-				false);
-		}
-	}
 };
 
 //////////////////////////////////////////////////////////////////////////
diff --git a/Engine/Plugins/2D/Paper2D/Source/Paper2DEditor/Private/TileMapEditing/TileMapActorFactory.cpp b/Engine/Plugins/2D/Paper2D/Source/Paper2DEditor/Private/TileMapEditing/TileMapActorFactory.cpp
--- a/Engine/Plugins/2D/Paper2D/Source/Paper2DEditor/Private/TileMapEditing/TileMapActorFactory.cpp
+++ b/Engine/Plugins/2D/Paper2D/Source/Paper2DEditor/Private/TileMapEditing/TileMapActorFactory.cpp
@@ -75,29 +75,21 @@ void UTileMapActorFactory::PostCreateBlueprint(UObject* Asset, AActor* CDO)
 
 bool UTileMapActorFactory::CanCreateActorFrom(const FAssetData& AssetData, FText& OutErrorMsg)
 {
-	if (GetDefault<UPaperRuntimeSettings>()->bEnableTileMapEditing)
+	if (AssetData.IsValid())
 	{
-		if (AssetData.IsValid())
+		UClass* AssetClass = AssetData.GetClass();
+		if ((AssetClass != nullptr) && (AssetClass->IsChildOf(UPaperTileMap::StaticClass()) || AssetClass->IsChildOf(UPaperTileSet::StaticClass())))
 		{
-			UClass* AssetClass = AssetData.GetClass();
-			if ((AssetClass != nullptr) && (AssetClass->IsChildOf(UPaperTileMap::StaticClass()) || AssetClass->IsChildOf(UPaperTileSet::StaticClass())))
-			{
-				return true;
-			}
-			else
-			{
-				OutErrorMsg = NSLOCTEXT("Paper2D", "CanCreateActorFrom_NoTileMap", "No tile map was specified.");
-				return false;
-			}
+			return true;
 		}
 		else
 		{
-			return true;
+			OutErrorMsg = NSLOCTEXT("Paper2D", "CanCreateActorFrom_NoTileMap", "No tile map was specified.");
+			return false;
 		}
 	}
 	else
 	{
-		OutErrorMsg = NSLOCTEXT("Paper2D", "CanCreateActorFrom_Disabled", "Tile map support is disabled in the Paper2D settings.");
-		return false;
+		return true;
 	}
 }
diff --git a/Engine/Plugins/2D/Paper2D/Source/PaperTiledImporter/Private/PaperTiledImporterFactory.cpp b/Engine/Plugins/2D/Paper2D/Source/PaperTiledImporter/Private/PaperTiledImporterFactory.cpp
--- a/Engine/Plugins/2D/Paper2D/Source/PaperTiledImporter/Private/PaperTiledImporterFactory.cpp
+++ b/Engine/Plugins/2D/Paper2D/Source/PaperTiledImporter/Private/PaperTiledImporterFactory.cpp
@@ -91,19 +91,16 @@ FText UPaperTiledImporterFactory::GetToolTip() const
 
 bool UPaperTiledImporterFactory::FactoryCanImport(const FString& Filename)
 {
-	if (GetDefault<UPaperRuntimeSettings>()->bEnableTileMapEditing)
+	FString FileContent;
+	if (FFileHelper::LoadFileToString(/*out*/ FileContent, *Filename))
 	{
-		FString FileContent;
-		if (FFileHelper::LoadFileToString(/*out*/ FileContent, *Filename))
+		TSharedPtr<FJsonObject> DescriptorObject = ParseJSON(FileContent, FString(), /*bSilent=*/ true);
+		if (DescriptorObject.IsValid())
 		{
-			TSharedPtr<FJsonObject> DescriptorObject = ParseJSON(FileContent, FString(), /*bSilent=*/ true);
-			if (DescriptorObject.IsValid())
-			{
-				FTileMapFromTiled GlobalInfo;
-				ParseGlobalInfoFromJSON(DescriptorObject, /*out*/ GlobalInfo, FString(), /*bSilent=*/ true);
+			FTileMapFromTiled GlobalInfo;
+			ParseGlobalInfoFromJSON(DescriptorObject, /*out*/ GlobalInfo, FString(), /*bSilent=*/ true);
 
-				return GlobalInfo.IsValid();
-			}
+			return GlobalInfo.IsValid();
 		}
 	}
 
