Paper2D: Add a factory for tile sets

diff --git a/Engine/Plugins/2D/Paper2D/Source/Paper2D/Classes/PaperTileSet.h b/Engine/Plugins/2D/Paper2D/Source/Paper2D/Classes/PaperTileSet.h
--- a/Engine/Plugins/2D/Paper2D/Source/Paper2D/Classes/PaperTileSet.h
+++ b/Engine/Plugins/2D/Paper2D/Source/Paper2D/Classes/PaperTileSet.h
@@ -8,7 +8,7 @@
 #include "PaperTileSet.generated.h"
 
 UCLASS()
-class PAPER2D_API UPaperTileSet : public UDataAsset //@TODO: Just to make it easy to spawn for now
+class PAPER2D_API UPaperTileSet : public UObject
 {
 	GENERATED_UCLASS_BODY()
 
diff --git a/Engine/Plugins/2D/Paper2D/Source/Paper2DEditor/Classes/PaperTileSetFactory.h b/Engine/Plugins/2D/Paper2D/Source/Paper2DEditor/Classes/PaperTileSetFactory.h
new file mode 100644
--- /dev/null
+++ b/Engine/Plugins/2D/Paper2D/Source/Paper2DEditor/Classes/PaperTileSetFactory.h
@@ -0,0 +1,22 @@
+// Copyright 1998-2015 Epic Games, Inc. All Rights Reserved.
+
+#pragma once
+
+/**
+ * Factory for tile sets
+ */
+
+#include "PaperTileSetFactory.generated.h"
+
+UCLASS()
+class UPaperTileSetFactory : public UFactory
+{
+	GENERATED_UCLASS_BODY()
+
+	// Initial texture to create the tile set from (Can be nullptr)
+	class UTexture2D* InitialTexture;
+
+	// UFactory interface
+	virtual UObject* FactoryCreateNew(UClass* Class, UObject* InParent, FName Name, EObjectFlags Flags, UObject* Context, FFeedbackContext* Warn) override;
+	// End of UFactory interface
+};
diff --git a/Engine/Plugins/2D/Paper2D/Source/Paper2DEditor/Private/PaperTileSetFactory.cpp b/Engine/Plugins/2D/Paper2D/Source/Paper2DEditor/Private/PaperTileSetFactory.cpp
new file mode 100644
--- /dev/null
+++ b/Engine/Plugins/2D/Paper2D/Source/Paper2DEditor/Private/PaperTileSetFactory.cpp
@@ -0,0 +1,26 @@
+// Copyright 1998-2015 Epic Games, Inc. All Rights Reserved.
+
+#include "Paper2DEditorPrivatePCH.h"
+
+#define LOCTEXT_NAMESPACE "Paper2D"
+
+/////////////////////////////////////////////////////
+// UPaperTileSetFactory
+
+UPaperTileSetFactory::UPaperTileSetFactory(const FObjectInitializer& ObjectInitializer)
+	: Super(ObjectInitializer)
+{
+	bCreateNew = true;
+	bEditAfterNew = true;
+	SupportedClass = UPaperTileSet::StaticClass();
+}
+
+UObject* UPaperTileSetFactory::FactoryCreateNew(UClass* Class, UObject* InParent, FName Name, EObjectFlags Flags, UObject* Context, FFeedbackContext* Warn)
+{
+	UPaperTileSet* NewTileSet = ConstructObject<UPaperTileSet>(Class, InParent, Name, Flags | RF_Transactional);
+	NewTileSet->TileSheet = InitialTexture;
+
+	return NewTileSet;
+}
+
+#undef LOCTEXT_NAMESPACE
\ No newline at end of file
