Paper2D: Add a context menu command for texture assets to create a tile set from them

diff --git a/Engine/Plugins/2D/Paper2D/Source/Paper2DEditor/Private/ContentBrowserExtensions/ContentBrowserExtensions.cpp b/Engine/Plugins/2D/Paper2D/Source/Paper2DEditor/Private/ContentBrowserExtensions/ContentBrowserExtensions.cpp
--- a/Engine/Plugins/2D/Paper2D/Source/Paper2DEditor/Private/ContentBrowserExtensions/ContentBrowserExtensions.cpp
+++ b/Engine/Plugins/2D/Paper2D/Source/Paper2DEditor/Private/ContentBrowserExtensions/ContentBrowserExtensions.cpp
@@ -50,12 +50,8 @@ struct FCreateSpriteFromTextureExtension : public FContentBrowserSelectedAssetEx
 		FAssetToolsModule& AssetToolsModule = FModuleManager::Get().LoadModuleChecked<FAssetToolsModule>("AssetTools");
 		FContentBrowserModule& ContentBrowserModule = FModuleManager::LoadModuleChecked<FContentBrowserModule>("ContentBrowser");
 
-		const bool bOneTextureSelected = Textures.Num() == 1;
 		TArray<UObject*> ObjectsToSync;
 
-		FString Name;
-		FString PackageName;
-
 		for (auto TextureIt = Textures.CreateConstIterator(); TextureIt; ++TextureIt)
 		{
 			UTexture2D* Texture = *TextureIt;
@@ -65,25 +61,18 @@ struct FCreateSpriteFromTextureExtension : public FContentBrowserSelectedAssetEx
 			SpriteFactory->InitialTexture = Texture;
 
 			// Create the sprite
+			FString Name;
+			FString PackageName;
+
 			if (!bExtractSprites)
 			{
-				if (bOneTextureSelected)
-				{
-					// Get a unique name for the sprite
-					AssetToolsModule.Get().CreateUniqueAssetName(Texture->GetOutermost()->GetName(), DefaultSuffix, /*out*/ PackageName, /*out*/ Name);
-					const FString PackagePath = FPackageName::GetLongPackagePath(PackageName);
-					ContentBrowserModule.Get().CreateNewAsset(Name, PackagePath, UPaperSprite::StaticClass(), SpriteFactory);
-				}
-				else
-				{
-					// Get a unique name for the sprite
-					AssetToolsModule.Get().CreateUniqueAssetName(Texture->GetOutermost()->GetName(), DefaultSuffix, /*out*/ PackageName, /*out*/ Name);
-					const FString PackagePath = FPackageName::GetLongPackagePath(PackageName);
+				// Get a unique name for the sprite
+				AssetToolsModule.Get().CreateUniqueAssetName(Texture->GetOutermost()->GetName(), DefaultSuffix, /*out*/ PackageName, /*out*/ Name);
+				const FString PackagePath = FPackageName::GetLongPackagePath(PackageName);
 
-					if (UObject* NewAsset = AssetToolsModule.Get().CreateAsset(Name, PackagePath, UPaperSprite::StaticClass(), SpriteFactory))
-					{
-						ObjectsToSync.Add(NewAsset);
-					}
+				if (UObject* NewAsset = AssetToolsModule.Get().CreateAsset(Name, PackagePath, UPaperSprite::StaticClass(), SpriteFactory))
+				{
+					ObjectsToSync.Add(NewAsset);
 				}
 			}
 			else
@@ -209,6 +198,60 @@ struct FConfigureTexturesForSpriteUsageExtension : public FContentBrowserSelecte
 	}
 };
 
+struct FCreateTileSetFromTextureExtension : public FContentBrowserSelectedAssetExtensionBase
+{
+	void CreateTileSetsFromTextures(TArray<UTexture2D*>& Textures)
+	{
+		const FString DefaultSuffix = TEXT("_TileSet");
+
+		FAssetToolsModule& AssetToolsModule = FModuleManager::Get().LoadModuleChecked<FAssetToolsModule>("AssetTools");
+		FContentBrowserModule& ContentBrowserModule = FModuleManager::LoadModuleChecked<FContentBrowserModule>("ContentBrowser");
+
+		TArray<UObject*> ObjectsToSync;
+
+		for (auto TextureIt = Textures.CreateConstIterator(); TextureIt; ++TextureIt)
+		{
+			UTexture2D* Texture = *TextureIt;
+
+			// Create the factory used to generate the tile set
+			UPaperTileSetFactory* TileSetFactory = ConstructObject<UPaperTileSetFactory>(UPaperTileSetFactory::StaticClass());
+			TileSetFactory->InitialTexture = Texture;
+
+			// Get a unique name for the tile set
+			FString Name;
+			FString PackageName;
+			AssetToolsModule.Get().CreateUniqueAssetName(Texture->GetOutermost()->GetName(), DefaultSuffix, /*out*/ PackageName, /*out*/ Name);
+			const FString PackagePath = FPackageName::GetLongPackagePath(PackageName);
+
+			if (UObject* NewAsset = AssetToolsModule.Get().CreateAsset(Name, PackagePath, UPaperTileSet::StaticClass(), TileSetFactory))
+			{
+				ObjectsToSync.Add(NewAsset);
+			}
+		}
+
+		if (ObjectsToSync.Num() > 0)
+		{
+			ContentBrowserModule.Get().SyncBrowserToAssets(ObjectsToSync);
+		}
+	}
+
+	virtual void Execute() override
+	{
+		// Create sprites for any selected textures
+		TArray<UTexture2D*> Textures;
+		for (auto AssetIt = SelectedAssets.CreateConstIterator(); AssetIt; ++AssetIt)
+		{
+			const FAssetData& AssetData = *AssetIt;
+			if (UTexture2D* Texture = Cast<UTexture2D>(AssetData.GetAsset()))
+			{
+				Textures.Add(Texture);
+			}
+		}
+
+		CreateTileSetsFromTextures(Textures);
+	}
+};
+
 //////////////////////////////////////////////////////////////////////////
 // FPaperContentBrowserExtensions_Impl
 
@@ -220,22 +263,25 @@ public:
 		SelectedAssetFunctor->Execute();
 	}
 
-	static void CreateSpriteActionsSubMenu(FMenuBuilder& MenuBuilder, TSharedPtr<FCreateSpriteFromTextureExtension> SpriteCreator, TSharedPtr<FCreateSpriteFromTextureExtension> SpriteExtractor, TSharedPtr<FConfigureTexturesForSpriteUsageExtension> ConfigFunctor)
+	static void CreateSpriteActionsSubMenu(FMenuBuilder& MenuBuilder, TArray<FAssetData> SelectedAssets)
 	{
 		MenuBuilder.AddSubMenu(
 			LOCTEXT("SpriteActionsSubMenuLabel", "Sprite Actions"),
 			LOCTEXT("SpriteActionsSubMenuToolTip", "Sprite-related actions for this texture."),
-			FNewMenuDelegate::CreateStatic(&FPaperContentBrowserExtensions_Impl::PopulateSpriteActionsMenu, SpriteCreator, SpriteExtractor, ConfigFunctor),
+			FNewMenuDelegate::CreateStatic(&FPaperContentBrowserExtensions_Impl::PopulateSpriteActionsMenu, SelectedAssets),
 			false,
 			FSlateIcon(FEditorStyle::GetStyleSetName(), "ClassIcon.PaperSprite")
 		);
 	}
 
-	static void PopulateSpriteActionsMenu(FMenuBuilder& MenuBuilder, TSharedPtr<FCreateSpriteFromTextureExtension> SpriteCreator, TSharedPtr<FCreateSpriteFromTextureExtension> SpriteExtractor, TSharedPtr<FConfigureTexturesForSpriteUsageExtension> ConfigFunctor)
+	static void PopulateSpriteActionsMenu(FMenuBuilder& MenuBuilder, TArray<FAssetData> SelectedAssets)
 	{
 		// Create sprites
+		TSharedPtr<FCreateSpriteFromTextureExtension> SpriteCreatorFunctor = MakeShareable(new FCreateSpriteFromTextureExtension());
+		SpriteCreatorFunctor->SelectedAssets = SelectedAssets;
+
 		FUIAction Action_CreateSpritesFromTextures(
-			FExecuteAction::CreateStatic(&FPaperContentBrowserExtensions_Impl::ExecuteSelectedContentFunctor, StaticCastSharedPtr<FContentBrowserSelectedAssetExtensionBase>(SpriteCreator)));
+			FExecuteAction::CreateStatic(&FPaperContentBrowserExtensions_Impl::ExecuteSelectedContentFunctor, StaticCastSharedPtr<FContentBrowserSelectedAssetExtensionBase>(SpriteCreatorFunctor)));
 		
 		MenuBuilder.AddMenuEntry(
 			LOCTEXT("CB_Extension_Texture_CreateSprite", "Create Sprite"),
@@ -246,8 +292,12 @@ public:
 			EUserInterfaceActionType::Button);
 
 		// Extract Sprites
+		TSharedPtr<FCreateSpriteFromTextureExtension> SpriteExtractorFunctor = MakeShareable(new FCreateSpriteFromTextureExtension());
+		SpriteExtractorFunctor->SelectedAssets = SelectedAssets;
+		SpriteExtractorFunctor->bExtractSprites = true;
+
 		FUIAction Action_ExtractSpritesFromTextures(
-			FExecuteAction::CreateStatic(&FPaperContentBrowserExtensions_Impl::ExecuteSelectedContentFunctor, StaticCastSharedPtr<FContentBrowserSelectedAssetExtensionBase>(SpriteExtractor)));
+			FExecuteAction::CreateStatic(&FPaperContentBrowserExtensions_Impl::ExecuteSelectedContentFunctor, StaticCastSharedPtr<FContentBrowserSelectedAssetExtensionBase>(SpriteExtractorFunctor)));
 
 		MenuBuilder.AddMenuEntry(
 			LOCTEXT("CB_Extension_Texture_ExtractSprites", "Extract Sprites"),
@@ -258,8 +308,11 @@ public:
 			EUserInterfaceActionType::Button);
 
 		// Configure for retro sprites
+		TSharedPtr<FConfigureTexturesForSpriteUsageExtension> TextureConfigFunctor = MakeShareable(new FConfigureTexturesForSpriteUsageExtension());
+		TextureConfigFunctor->SelectedAssets = SelectedAssets;
+
 		FUIAction Action_ConfigureTexturesForSprites(
-			FExecuteAction::CreateStatic(&FPaperContentBrowserExtensions_Impl::ExecuteSelectedContentFunctor, StaticCastSharedPtr<FContentBrowserSelectedAssetExtensionBase>(ConfigFunctor)));
+			FExecuteAction::CreateStatic(&FPaperContentBrowserExtensions_Impl::ExecuteSelectedContentFunctor, StaticCastSharedPtr<FContentBrowserSelectedAssetExtensionBase>(TextureConfigFunctor)));
 
 		MenuBuilder.AddMenuEntry(
 			LOCTEXT("CB_Extension_Texture_ConfigureTextureForSprites", "Configure For Retro Sprites"),
@@ -268,6 +321,21 @@ public:
 			Action_ConfigureTexturesForSprites,
 			NAME_None,
 			EUserInterfaceActionType::Button);
+
+		// Create sprites
+		TSharedPtr<FCreateTileSetFromTextureExtension> TileSetCreatorFunctor = MakeShareable(new FCreateTileSetFromTextureExtension());
+		TileSetCreatorFunctor->SelectedAssets = SelectedAssets;
+
+		FUIAction Action_CreateTileSetFromTextures(
+			FExecuteAction::CreateStatic(&FPaperContentBrowserExtensions_Impl::ExecuteSelectedContentFunctor, StaticCastSharedPtr<FContentBrowserSelectedAssetExtensionBase>(TileSetCreatorFunctor)));
+
+		MenuBuilder.AddMenuEntry(
+			LOCTEXT("CB_Extension_Texture_CreateTileSet", "Create Tile Set"),
+			LOCTEXT("CB_Extension_Texture_CreateTileSet_Tooltip", "Create tile set from selected texture"),
+			FSlateIcon(FEditorStyle::GetStyleSetName(), "ClassIcon.PaperTileSet"),
+			Action_CreateTileSetFromTextures,
+			NAME_None,
+			EUserInterfaceActionType::Button);
 	}
 
 	static TSharedRef<FExtender> OnExtendContentBrowserAssetSelectionMenu(const TArray<FAssetData>& SelectedAssets)
@@ -284,26 +352,12 @@ public:
 
 		if (bAnyTextures)
 		{
-			// Create Sprite
-			TSharedPtr<FCreateSpriteFromTextureExtension> SpriteCreatorFunctor = MakeShareable(new FCreateSpriteFromTextureExtension());
-			SpriteCreatorFunctor->SelectedAssets = SelectedAssets;
-
-			////@TODO: Already getting nasty, need to refactor the data to be independent of the functor
-			// Extract Sprites
-			TSharedPtr<FCreateSpriteFromTextureExtension> SpriteExtractorFunctor = MakeShareable(new FCreateSpriteFromTextureExtension());
-			SpriteExtractorFunctor->SelectedAssets = SelectedAssets;
-			SpriteExtractorFunctor->bExtractSprites = true;
-
-			// Configure for retro sprites
-			TSharedPtr<FConfigureTexturesForSpriteUsageExtension> TextureConfigFunctor = MakeShareable(new FConfigureTexturesForSpriteUsageExtension());
-			TextureConfigFunctor->SelectedAssets = SelectedAssets;
-
 			// Add the sprite actions sub-menu extender
 			Extender->AddMenuExtension(
 				"GetAssetActions",
 				EExtensionHook::After,
 				nullptr,
-				FMenuExtensionDelegate::CreateStatic(&FPaperContentBrowserExtensions_Impl::CreateSpriteActionsSubMenu, SpriteCreatorFunctor, SpriteExtractorFunctor, TextureConfigFunctor));
+				FMenuExtensionDelegate::CreateStatic(&FPaperContentBrowserExtensions_Impl::CreateSpriteActionsSubMenu, SelectedAssets));
 		}
 
 		return Extender;
