Paper2D: Change default value for UPaperTileMap::SeparationPerLayer to 4.0

diff --git a/Engine/Plugins/2D/Paper2D/Source/Paper2D/Private/PaperTileMap.cpp b/Engine/Plugins/2D/Paper2D/Source/Paper2D/Private/PaperTileMap.cpp
--- a/Engine/Plugins/2D/Paper2D/Source/Paper2D/Private/PaperTileMap.cpp
+++ b/Engine/Plugins/2D/Paper2D/Source/Paper2D/Private/PaperTileMap.cpp
@@ -19,7 +19,7 @@ UPaperTileMap::UPaperTileMap(const FObjectInitializer& ObjectInitializer)
 	PixelsPerUnit = 1.0f;
 	SeparationPerTileX = 0.0f;
 	SeparationPerTileY = 0.0f;
-	SeparationPerLayer = 64.0f;
+	SeparationPerLayer = 4.0f;
 	SpriteCollisionDomain = ESpriteCollisionMode::None;
 
 #if WITH_EDITORONLY_DATA
