Paper2D: Revamp UPaperTileMapComponent asset reference to work similar to ambient audio actors, allowing for either an internal owned instance or an externally referenced asset
- Add "New" and "Promote" buttons to details panel customization, and changed "Edit" behavior to open the asset editor for non-owned maps
- Changed TileMap to no longer be Instanced; details customization now shows or hides inner properties and the layer list depending on whether the component owns the data or is referencing an asset
- Promoted tile map components from experimental to early access preview
- Added UPaperTileMapComponent::SetTileMap to allow changing the tile map of a component at runtime

diff --git a/Engine/Plugins/2D/Paper2D/Source/Paper2D/Classes/PaperTileMapComponent.h b/Engine/Plugins/2D/Paper2D/Source/Paper2D/Classes/PaperTileMapComponent.h
--- a/Engine/Plugins/2D/Paper2D/Source/Paper2D/Classes/PaperTileMapComponent.h
+++ b/Engine/Plugins/2D/Paper2D/Source/Paper2D/Classes/PaperTileMapComponent.h
@@ -6,7 +6,7 @@
 
 #include "PaperTileMapComponent.generated.h"
 
-UCLASS(hideCategories=Object, ClassGroup=Paper2D, Experimental, meta=(BlueprintSpawnableComponent))
+UCLASS(hideCategories=Object, ClassGroup=Paper2D, EarlyAccessPreview, meta=(BlueprintSpawnableComponent))
 class PAPER2D_API UPaperTileMapComponent : public UPrimitiveComponent
 {
 	GENERATED_UCLASS_BODY()
@@ -39,7 +39,7 @@ private:
 
 public:
 	// The tile map used by this component
-	UPROPERTY(Category=Setup, EditAnywhere, Instanced, BlueprintReadOnly)
+	UPROPERTY(Category=Setup, EditAnywhere, BlueprintReadOnly)
 	class UPaperTileMap* TileMap;
 
 protected:
@@ -67,8 +67,15 @@ public:
 	virtual class UBodySetup* GetBodySetup() override;
 	// End of UPrimitiveComponent interface
 
+	// Creates a new tile map internally, replacing the TileMap reference (or dropping the previous owned one)
+	void CreateNewOwnedTileMap();
+
 	// Does this component own the tile map (is it instanced instead of being an asset reference)?
 	bool OwnsTileMap() const;
+
+	/** Change the PaperTileMap used by this instance. */
+	UFUNCTION(BlueprintCallable, Category="Sprite")
+	virtual bool SetTileMap(class UPaperTileMap* NewTileMap);
 };
 
 // Allow the old name to continue to work for one release
diff --git a/Engine/Plugins/2D/Paper2D/Source/Paper2D/Private/PaperTileMapComponent.cpp b/Engine/Plugins/2D/Paper2D/Source/Paper2D/Private/PaperTileMapComponent.cpp
--- a/Engine/Plugins/2D/Paper2D/Source/Paper2D/Private/PaperTileMapComponent.cpp
+++ b/Engine/Plugins/2D/Paper2D/Source/Paper2D/Private/PaperTileMapComponent.cpp
@@ -318,9 +318,46 @@ void UPaperTileMapComponent::RebuildRenderData(FPaperTileMapRenderSceneProxy* Pr
 	Proxy->SetBatchesHack(BatchedSprites);
 }
 
+void UPaperTileMapComponent::CreateNewOwnedTileMap()
+{
+	TGuardValue<TEnumAsByte<EComponentMobility::Type>> MobilitySaver(Mobility, EComponentMobility::Movable);
+
+	UPaperTileMap* NewTileMap = NewObject<UPaperTileMap>(this);
+	NewTileMap->SetFlags(RF_Transactional);
+	NewTileMap->AddNewLayer();
+
+	SetTileMap(NewTileMap);
+}
+
 bool UPaperTileMapComponent::OwnsTileMap() const
 {
 	return (TileMap != nullptr) && (TileMap->GetOuter() == this);
 }
 
+bool UPaperTileMapComponent::SetTileMap(class UPaperTileMap* NewTileMap)
+{
+	if (NewTileMap != TileMap)
+	{
+		// Don't allow changing the tile map if we are "static".
+		AActor* Owner = GetOwner();
+		if (!IsRegistered() || (Owner == nullptr) || (Mobility != EComponentMobility::Static))
+		{
+			TileMap = NewTileMap;
+
+			// Need to send this to render thread at some point
+			MarkRenderStateDirty();
+
+			// Update physics representation right away
+			RecreatePhysicsState();
+
+			// Since we have new mesh, we need to update bounds
+			UpdateBounds();
+
+			return true;
+		}
+	}
+
+	return false;
+}
+
 #undef LOCTEXT_NAMESPACE
\ No newline at end of file
diff --git a/Engine/Plugins/2D/Paper2D/Source/Paper2DEditor/Classes/PaperTileMapPromotionFactory.h b/Engine/Plugins/2D/Paper2D/Source/Paper2DEditor/Classes/PaperTileMapPromotionFactory.h
new file mode 100644
--- /dev/null
+++ b/Engine/Plugins/2D/Paper2D/Source/Paper2DEditor/Classes/PaperTileMapPromotionFactory.h
@@ -0,0 +1,23 @@
+// Copyright 1998-2015 Epic Games, Inc. All Rights Reserved.
+
+#pragma once
+
+/**
+ * Factory used to promote an instance-specific object from instance to asset by renaming it into an asset package
+ */
+
+#include "PaperTileMapPromotionFactory.generated.h"
+
+UCLASS()
+class UPaperTileMapPromotionFactory : public UFactory
+{
+	GENERATED_UCLASS_BODY()
+
+	// Object being promoted to an asset
+	UPROPERTY()
+	class UPaperTileMap* AssetToRename;
+
+	// UFactory interface
+	virtual UObject* FactoryCreateNew(UClass* Class, UObject* InParent, FName Name, EObjectFlags Flags, UObject* Context, FFeedbackContext* Warn) override;
+	// End of UFactory interface
+};
diff --git a/Engine/Plugins/2D/Paper2D/Source/Paper2DEditor/Paper2DEditor.Build.cs b/Engine/Plugins/2D/Paper2D/Source/Paper2DEditor/Paper2DEditor.Build.cs
--- a/Engine/Plugins/2D/Paper2D/Source/Paper2DEditor/Paper2DEditor.Build.cs
+++ b/Engine/Plugins/2D/Paper2D/Source/Paper2DEditor/Paper2DEditor.Build.cs
@@ -20,6 +20,7 @@ public class Paper2DEditor : ModuleRules
 				"AssetTools",
 				"UnrealEd", // for FAssetEditorManager
 				"KismetWidgets",
+				"EditorWidgets", // for SCreateNewAssetFromFactory
 				"GraphEditor",
 				"Kismet",  // for FWorkflowCentricApplication
 				"PropertyEditor",
diff --git a/Engine/Plugins/2D/Paper2D/Source/Paper2DEditor/Private/PaperTileMapPromotionFactory.cpp b/Engine/Plugins/2D/Paper2D/Source/Paper2DEditor/Private/PaperTileMapPromotionFactory.cpp
new file mode 100644
--- /dev/null
+++ b/Engine/Plugins/2D/Paper2D/Source/Paper2DEditor/Private/PaperTileMapPromotionFactory.cpp
@@ -0,0 +1,26 @@
+// Copyright 1998-2015 Epic Games, Inc. All Rights Reserved.
+
+#include "Paper2DEditorPrivatePCH.h"
+
+#define LOCTEXT_NAMESPACE "Paper2D"
+
+/////////////////////////////////////////////////////
+// UPaperTileMapPromotionFactory
+
+UPaperTileMapPromotionFactory::UPaperTileMapPromotionFactory(const FObjectInitializer& ObjectInitializer)
+	: Super(ObjectInitializer)
+{
+	bCreateNew = false;
+	bEditAfterNew = true;
+	SupportedClass = UPaperTileMap::StaticClass();
+}
+
+UObject* UPaperTileMapPromotionFactory::FactoryCreateNew(UClass* Class, UObject* InParent, FName Name, EObjectFlags Flags, UObject* Context, FFeedbackContext* Warn)
+{
+	AssetToRename->Rename(*Name.ToString(), InParent);
+	AssetToRename->SetFlags(Flags | RF_Transactional);
+
+	return AssetToRename;
+}
+
+#undef LOCTEXT_NAMESPACE
\ No newline at end of file
diff --git a/Engine/Plugins/2D/Paper2D/Source/Paper2DEditor/Private/TileMapEditing/PaperTileMapDetailsCustomization.cpp b/Engine/Plugins/2D/Paper2D/Source/Paper2DEditor/Private/TileMapEditing/PaperTileMapDetailsCustomization.cpp
--- a/Engine/Plugins/2D/Paper2D/Source/Paper2DEditor/Private/TileMapEditing/PaperTileMapDetailsCustomization.cpp
+++ b/Engine/Plugins/2D/Paper2D/Source/Paper2DEditor/Private/TileMapEditing/PaperTileMapDetailsCustomization.cpp
@@ -6,6 +6,7 @@
 #include "DetailCategoryBuilder.h"
 #include "EdModeTileMap.h"
 #include "PropertyEditing.h"
+#include "SCreateNewAssetFromFactory.h"
 
 #include "STileLayerList.h"
 #include "ScopedTransaction.h"
@@ -24,12 +25,14 @@ TSharedRef<IDetailCustomization> FPaperTileMapDetailsCustomization::MakeInstance
 void FPaperTileMapDetailsCustomization::CustomizeDetails(IDetailLayoutBuilder& DetailLayout)
 {
 	const TArray< TWeakObjectPtr<UObject> >& SelectedObjects = DetailLayout.GetDetailsView().GetSelectedObjects();
+	MyDetailLayout = &DetailLayout;
 	
 	FNotifyHook* NotifyHook = DetailLayout.GetPropertyUtilities()->GetNotifyHook();
 
 	bool bEditingActor = false;
 
 	UPaperTileMap* TileMap = nullptr;
+	UPaperTileMapComponent* TileComponent = nullptr;
 	for (int32 ObjectIndex = 0; ObjectIndex < SelectedObjects.Num(); ++ObjectIndex)
 	{
 		UObject* TestObject = SelectedObjects[ObjectIndex].Get();
@@ -38,12 +41,14 @@ void FPaperTileMapDetailsCustomization::CustomizeDetails(IDetailLayoutBuilder& D
 			if (UPaperTileMapComponent* CurrentComponent = CurrentActor->FindComponentByClass<UPaperTileMapComponent>())
 			{
 				bEditingActor = true;
+				TileComponent = CurrentComponent;
 				TileMap = CurrentComponent->TileMap;
 				break;
 			}
 		}
 		else if (UPaperTileMapComponent* TestComponent = Cast<UPaperTileMapComponent>(TestObject))
 		{
+			TileComponent = TestComponent;
 			TileMap = TestComponent->TileMap;
 			break;
 		}
@@ -54,46 +59,177 @@ void FPaperTileMapDetailsCustomization::CustomizeDetails(IDetailLayoutBuilder& D
 		}
 	}
 	TileMapPtr = TileMap;
+	TileMapComponentPtr = TileComponent;
 
-	IDetailCategoryBuilder& TileMapCategory = DetailLayout.EditCategory("TileMap");
+	IDetailCategoryBuilder& TileMapCategory = DetailLayout.EditCategory("Tile Map");
 
-	if (bEditingActor)
+	TAttribute<EVisibility> InternalInstanceVis = TAttribute<EVisibility>::Create(TAttribute<EVisibility>::FGetter::CreateSP(this, &FPaperTileMapDetailsCustomization::GetVisibilityForInstancedOnlyProperties));
+
+	if (TileComponent != nullptr)
 	{
 		TileMapCategory
-		.AddCustomRow(LOCTEXT( "EnterTileMapEditMode", "Enter Edit Mode" ))
+		.AddCustomRow(LOCTEXT( "TileMapInstancingControlsSearchText", "Edit New Promote Asset"))
 		[
 			SNew(SVerticalBox)
-			+SVerticalBox::Slot()
+			+ SVerticalBox::Slot()
+			.Padding(0.0f, 2.0f, 0.0f, 0.0f)
+			.FillHeight(1.0f)
+			.VAlign(VAlign_Center)
 			[
-				SNew(SButton)
-				.ContentPadding(3)
+				SNew(SHorizontalBox)
+
+				// Edit button
+				+SHorizontalBox::Slot()
+				.AutoWidth()
+				.Padding( 2.0f, 0.0f )
+				.VAlign(VAlign_Center)
+				.HAlign(HAlign_Left)
+				[
+					SNew(SButton)
+					.VAlign(VAlign_Center)
+					.OnClicked(this, &FPaperTileMapDetailsCustomization::EnterTileMapEditingMode)
+					.Visibility(this, &FPaperTileMapDetailsCustomization::GetNonEditModeVisibility)
+					.Text( LOCTEXT("EditAsset", "Edit") )
+					.ToolTipText( LOCTEXT("EditAssetToolTip", "Edit this tile map") )
+				]
+
+				// Create new tile map button
+				+SHorizontalBox::Slot()
+				.AutoWidth()
+				.Padding( 2.0f, 0.0f )
+				.VAlign(VAlign_Center)
+				.HAlign(HAlign_Left)
+				[
+					SNew(SButton)
+					.VAlign(VAlign_Center)
+					.OnClicked(this, &FPaperTileMapDetailsCustomization::OnNewButtonClicked)
+					.Visibility(this, &FPaperTileMapDetailsCustomization::GetNewButtonVisiblity)
+					.Text(LOCTEXT("CreateNewInstancedMap", "New"))
+					.ToolTipText( LOCTEXT("CreateNewInstancedMapToolTip", "Create a new tile map") )
+				]
+
+				// Promote to asset button
+				+SHorizontalBox::Slot()
+				.AutoWidth()
+				.Padding( 2.0f, 0.0f )
 				.VAlign(VAlign_Center)
-				.HAlign(HAlign_Center)
-				.OnClicked(this, &FPaperTileMapDetailsCustomization::EnterTileMapEditingMode)
-				.Visibility(this, &FPaperTileMapDetailsCustomization::GetNonEditModeVisibility)
-				.Text( LOCTEXT( "EnterTileMapEditMode", "Enter Edit Mode" ) )
+				.HAlign(HAlign_Left)
+				[
+					SNew(SButton)
+					.VAlign(VAlign_Center)
+					.OnClicked(this, &FPaperTileMapDetailsCustomization::OnPromoteButtonClicked)
+					.Visibility(this, &FPaperTileMapDetailsCustomization::GetVisibilityForInstancedOnlyProperties)
+					.Text(LOCTEXT("PromoteToAsset", "Promote to asset"))
+					.ToolTipText(LOCTEXT("PromoteToAssetToolTip", "Save this tile map as a reusable asset"))
+				]
 			]
 		];
+
+		TileMapCategory.AddProperty(GET_MEMBER_NAME_CHECKED(UPaperTileMapComponent, TileMap));
 	}
 
-	//@TODO: Handle showing layers when multiple tile maps are selected
+	// Add the layer browser
 	if (TileMap != nullptr)
 	{
-		IDetailCategoryBuilder& LayersCategory = DetailLayout.EditCategory("Tile Layers");
+		TAttribute<EVisibility> LayerBrowserVis;
+		LayerBrowserVis.Set(EVisibility::Visible);
+		if (TileComponent != nullptr)
+		{
+			LayerBrowserVis = InternalInstanceVis;
+		}
 
-		LayersCategory.AddCustomRow(LOCTEXT("TileLayerList", "Tile layer list"))
+		FText TileLayerListText = LOCTEXT("TileLayerList", "Tile layer list");
+		TileMapCategory.AddCustomRow(TileLayerListText)
+		.Visibility(LayerBrowserVis)
 		[
-			SNew(STileLayerList, TileMap, NotifyHook)
+			SNew(SVerticalBox)
+			+SVerticalBox::Slot()
+			.AutoHeight()
+			[
+				SNew(STextBlock)
+				.Font(DetailLayout.GetDetailFont())
+				.Text(TileLayerListText)
+			]
+			+SVerticalBox::Slot()
+			[
+				SNew(STileLayerList, TileMap, NotifyHook)
+			]
 		];
 	}
 
+	// Add all of the properties from the inline tilemap
+	if ((TileComponent != nullptr) && (TileComponent->OwnsTileMap()))
+	{
+		TArray<UObject*> ListOfTileMaps;
+		ListOfTileMaps.Add(TileMap);
+
+		for (TFieldIterator<UProperty> PropIt(UPaperTileMap::StaticClass()); PropIt; ++PropIt)
+		{
+			UProperty* TestProperty = *PropIt;
+
+			if (TestProperty->HasAnyPropertyFlags(CPF_Edit))
+			{
+				FName CategoryName(*TestProperty->GetMetaData(TEXT("Category")));
+				IDetailCategoryBuilder& Category = DetailLayout.EditCategory(CategoryName);
+
+				if (IDetailPropertyRow* ExternalRow = Category.AddExternalProperty(ListOfTileMaps, TestProperty->GetFName()))
+				{
+					ExternalRow->Visibility(InternalInstanceVis);
+				}
+			}
+		}
+	}
+
 	// Make sure the setup category is near the top
 	DetailLayout.EditCategory("Setup");
 }
 
 FReply FPaperTileMapDetailsCustomization::EnterTileMapEditingMode()
 {
-	GLevelEditorModeTools().ActivateMode(FEdModeTileMap::EM_TileMap);
+	if (UPaperTileMapComponent* TileMapComponent = TileMapComponentPtr.Get())
+	{
+		if (TileMapComponent->OwnsTileMap())
+		{
+			GLevelEditorModeTools().ActivateMode(FEdModeTileMap::EM_TileMap);
+		}
+		else
+		{
+			FAssetEditorManager::Get().OpenEditorForAsset(TileMapComponent->TileMap);
+		}
+	}
+	return FReply::Handled();
+}
+
+FReply FPaperTileMapDetailsCustomization::OnNewButtonClicked()
+{
+	if (UPaperTileMapComponent* TileMapComponent = TileMapComponentPtr.Get())
+	{
+		const FScopedTransaction Transaction(LOCTEXT("CreateNewTileMap", "New Tile Map"));
+		TileMapComponent->Modify();
+		TileMapComponent->CreateNewOwnedTileMap();
+
+		MyDetailLayout->ForceRefreshDetails();
+	}
+
+	return FReply::Handled();
+}
+
+FReply FPaperTileMapDetailsCustomization::OnPromoteButtonClicked()
+{
+	if (UPaperTileMapComponent* TileMapComponent = TileMapComponentPtr.Get())
+	{
+		if (TileMapComponent->OwnsTileMap())
+		{
+			if (TileMapComponent->TileMap != nullptr)
+			{
+				// Try promoting the tile map to be an asset (prompts for a name&path, creates a package and then calls the factory, which renames the existing asset and sets RF_Public)
+				UPaperTileMapPromotionFactory* PromotionFactory = NewObject<UPaperTileMapPromotionFactory>();
+				PromotionFactory->AssetToRename = TileMapComponent->TileMap;
+				SCreateNewAssetFromFactory::Create(PromotionFactory);
+			}
+		}
+	}
+
 	return FReply::Handled();
 }
 
@@ -102,6 +238,24 @@ EVisibility FPaperTileMapDetailsCustomization::GetNonEditModeVisibility() const
 	return GLevelEditorModeTools().IsModeActive(FEdModeTileMap::EM_TileMap) ? EVisibility::Collapsed : EVisibility::Visible;
 }
 
+EVisibility FPaperTileMapDetailsCustomization::GetNewButtonVisiblity() const
+{
+	return (TileMapComponentPtr.Get() != nullptr) ? EVisibility::Visible : EVisibility::Collapsed;
+}
+
+EVisibility FPaperTileMapDetailsCustomization::GetVisibilityForInstancedOnlyProperties() const
+{
+	if (UPaperTileMapComponent* TileMapComponent = TileMapComponentPtr.Get())
+	{
+		if (TileMapComponent->OwnsTileMap())
+		{
+			return EVisibility::Visible;
+		}
+	}
+
+	return EVisibility::Hidden;
+}
+
 //////////////////////////////////////////////////////////////////////////
 
 #undef LOCTEXT_NAMESPACE
\ No newline at end of file
diff --git a/Engine/Plugins/2D/Paper2D/Source/Paper2DEditor/Private/TileMapEditing/PaperTileMapDetailsCustomization.h b/Engine/Plugins/2D/Paper2D/Source/Paper2DEditor/Private/TileMapEditing/PaperTileMapDetailsCustomization.h
--- a/Engine/Plugins/2D/Paper2D/Source/Paper2DEditor/Private/TileMapEditing/PaperTileMapDetailsCustomization.h
+++ b/Engine/Plugins/2D/Paper2D/Source/Paper2DEditor/Private/TileMapEditing/PaperTileMapDetailsCustomization.h
@@ -19,8 +19,18 @@ public:
 
 private:
 	TWeakObjectPtr<class UPaperTileMap> TileMapPtr;
+	TWeakObjectPtr<class UPaperTileMapComponent> TileMapComponentPtr;
+
+	IDetailLayoutBuilder* MyDetailLayout;
 
 private:
 	FReply EnterTileMapEditingMode();
+	FReply OnNewButtonClicked();
+	FReply OnPromoteButtonClicked();
+
 	EVisibility GetNonEditModeVisibility() const;
+
+	EVisibility GetVisibilityForInstancedOnlyProperties() const;
+
+	EVisibility GetNewButtonVisiblity() const;
 };
