Paper2D: Add FAssetEditorModeTools::GetWorld override, pulling from an associated preview scene if possible and falling back to the GEditor world otherwise

diff --git a/Engine/Plugins/2D/Paper2D/Source/Paper2DEditor/Private/PaperEditorViewportClient.cpp b/Engine/Plugins/2D/Paper2D/Source/Paper2DEditor/Private/PaperEditorViewportClient.cpp
--- a/Engine/Plugins/2D/Paper2D/Source/Paper2DEditor/Private/PaperEditorViewportClient.cpp
+++ b/Engine/Plugins/2D/Paper2D/Source/Paper2DEditor/Private/PaperEditorViewportClient.cpp
@@ -5,11 +5,13 @@
 #include "PaperEditorViewportClient.h"
 #include "CanvasTypes.h"
 #include "SEditorViewport.h"
+#include "PreviewScene.h"
 
 //////////////////////////////////////////////////////////////////////////
 // FAssetEditorModeTools
 
 FAssetEditorModeTools::FAssetEditorModeTools()
+	: PreviewScene(nullptr)
 {
 	ActorSet = NewObject<USelection>();
 	ActorSet->SetFlags(RF_Transactional);
@@ -38,6 +40,16 @@ USelection* FAssetEditorModeTools::GetSelectedObjects() const
 	return ObjectSet;
 }
 
+UWorld* FAssetEditorModeTools::GetWorld() const
+{
+	return (PreviewScene != nullptr) ? PreviewScene->GetWorld() : GEditor->GetEditorWorldContext().World();
+}
+
+void FAssetEditorModeTools::SetPreviewScene(class FPreviewScene* NewPreviewScene)
+{
+	PreviewScene = NewPreviewScene;
+}
+
 //////////////////////////////////////////////////////////////////////////
 // FPaperEditorViewportClient
 
diff --git a/Engine/Plugins/2D/Paper2D/Source/Paper2DEditor/Private/PaperEditorViewportClient.h b/Engine/Plugins/2D/Paper2D/Source/Paper2DEditor/Private/PaperEditorViewportClient.h
--- a/Engine/Plugins/2D/Paper2D/Source/Paper2DEditor/Private/PaperEditorViewportClient.h
+++ b/Engine/Plugins/2D/Paper2D/Source/Paper2DEditor/Private/PaperEditorViewportClient.h
@@ -23,11 +23,14 @@ public:
 	// FEditorModeTools interface
 	virtual class USelection* GetSelectedActors() const override;
 	virtual class USelection* GetSelectedObjects() const override;
+	virtual UWorld* GetWorld() const override;
 	// End of FEditorModeTools interface
 
+	void SetPreviewScene(class FPreviewScene* NewPreviewScene);
 protected:
 	class USelection* ActorSet;
 	class USelection* ObjectSet;
+	class FPreviewScene* PreviewScene;
 };
 
 //////////////////////////////////////////////////////////////////////////
diff --git a/Engine/Plugins/2D/Paper2D/Source/Paper2DEditor/Private/TileMapEditing/TileMapEditorViewportClient.cpp b/Engine/Plugins/2D/Paper2D/Source/Paper2DEditor/Private/TileMapEditing/TileMapEditorViewportClient.cpp
--- a/Engine/Plugins/2D/Paper2D/Source/Paper2DEditor/Private/TileMapEditing/TileMapEditorViewportClient.cpp
+++ b/Engine/Plugins/2D/Paper2D/Source/Paper2DEditor/Private/TileMapEditing/TileMapEditorViewportClient.cpp
@@ -25,6 +25,7 @@ FTileMapEditorViewportClient::FTileMapEditorViewportClient(TWeakPtr<FTileMapEdit
 	check(TileMapEditorPtr.IsValid() && TileMapEditorViewportPtr.IsValid());
 
 	PreviewScene = &OwnedPreviewScene;
+	((FAssetEditorModeTools*)ModeTools)->SetPreviewScene(PreviewScene);
 
 	SetRealtime(true);
 
