Paper2D: Change UPaperTileLayer::ResizeMap to avoid doing unnecessary work if the size didn't change

diff --git a/Engine/Plugins/2D/Paper2D/Source/Paper2D/Private/PaperTileLayer.cpp b/Engine/Plugins/2D/Paper2D/Source/Paper2D/Private/PaperTileLayer.cpp
--- a/Engine/Plugins/2D/Paper2D/Source/Paper2D/Private/PaperTileLayer.cpp
+++ b/Engine/Plugins/2D/Paper2D/Source/Paper2D/Private/PaperTileLayer.cpp
@@ -36,9 +36,12 @@ void UPaperTileLayer::DestructiveAllocateMap(int32 NewWidth, int32 NewHeight)
 
 void UPaperTileLayer::ResizeMap(int32 NewWidth, int32 NewHeight)
 {
-	LayerWidth = NewWidth;
-	LayerHeight = NewHeight;
-	ReallocateAndCopyMap();
+	if ((LayerWidth != NewWidth) || (LayerHeight != NewHeight))
+	{
+		LayerWidth = NewWidth;
+		LayerHeight = NewHeight;
+		ReallocateAndCopyMap();
+	}
 }
 
 void UPaperTileLayer::ReallocateAndCopyMap()
