Paper2D: Add a visual preview of the current ink on the tile map painting brush using a preview component

diff --git a/Engine/Plugins/2D/Paper2D/Source/Paper2D/Classes/PaperTileMap.h b/Engine/Plugins/2D/Paper2D/Source/Paper2D/Classes/PaperTileMap.h
--- a/Engine/Plugins/2D/Paper2D/Source/Paper2D/Classes/PaperTileMap.h
+++ b/Engine/Plugins/2D/Paper2D/Source/Paper2D/Classes/PaperTileMap.h
@@ -109,7 +109,7 @@ public:
 	// End of UObject interface
 
 
-	FVector GetTilePositionInLocalSpace(int32 TileX, int32 TileY, int32 LayerIndex = 0) const;
+	FVector GetTilePositionInLocalSpace(float TileX, float TileY, int32 LayerIndex = 0) const;
 
 	FBoxSphereBounds GetRenderBounds() const;
 
diff --git a/Engine/Plugins/2D/Paper2D/Source/Paper2D/Private/PaperTileMap.cpp b/Engine/Plugins/2D/Paper2D/Source/Paper2D/Private/PaperTileMap.cpp
--- a/Engine/Plugins/2D/Paper2D/Source/Paper2D/Private/PaperTileMap.cpp
+++ b/Engine/Plugins/2D/Paper2D/Source/Paper2D/Private/PaperTileMap.cpp
@@ -138,7 +138,7 @@ void UPaperTileMap::UpdateBodySetup()
 	}
 }
 
-FVector UPaperTileMap::GetTilePositionInLocalSpace(int32 TileX, int32 TileY, int32 LayerIndex) const
+FVector UPaperTileMap::GetTilePositionInLocalSpace(float TileX, float TileY, int32 LayerIndex) const
 {
 	//@TODO: Tile pivot issue
 	const FVector PartialX = PaperAxisX * (TileX - 0.5f) * TileWidth;
diff --git a/Engine/Plugins/2D/Paper2D/Source/Paper2DEditor/Private/TileMapEditing/EdModeTileMap.cpp b/Engine/Plugins/2D/Paper2D/Source/Paper2DEditor/Private/TileMapEditing/EdModeTileMap.cpp
--- a/Engine/Plugins/2D/Paper2D/Source/Paper2DEditor/Private/TileMapEditing/EdModeTileMap.cpp
+++ b/Engine/Plugins/2D/Paper2D/Source/Paper2DEditor/Private/TileMapEditing/EdModeTileMap.cpp
@@ -87,9 +87,6 @@ FEdModeTileMap::FEdModeTileMap()
 	, EraseBrushSize(1)
 {
 	bDrawPivot = false;
-
-	SetActiveTool(ETileMapEditorTool::Paintbrush);
-	SetActiveLayerPaintingMode(ETileMapLayerPaintingMode::VisualLayers);
 }
 
 FEdModeTileMap::~FEdModeTileMap()
@@ -110,6 +107,18 @@ void FEdModeTileMap::Enter()
 {
 	FEdMode::Enter();
 
+	UWorld* World = GetWorld();
+
+	CursorPreviewComponent = NewObject<UPaperTileMapComponent>();
+	CursorPreviewComponent->TileMap->AddNewLayer();
+	CursorPreviewComponent->TranslucencySortPriority = 99999;
+	CursorPreviewComponent->UpdateBounds();
+	CursorPreviewComponent->AddToRoot();
+	CursorPreviewComponent->RegisterComponentWithWorld(World);
+
+	SetActiveTool(ETileMapEditorTool::Paintbrush);
+	SetActiveLayerPaintingMode(ETileMapLayerPaintingMode::VisualLayers);
+
 	if (!Toolkit.IsValid())
 	{
 		Toolkit = MakeShareable(new FTileMapEdModeToolkit(this));
@@ -125,6 +134,10 @@ void FEdModeTileMap::Exit()
 		Toolkit.Reset();
 	}
 
+	CursorPreviewComponent->RemoveFromRoot();
+	CursorPreviewComponent->UnregisterComponent();
+	CursorPreviewComponent = nullptr;
+
 	// Call base Exit method to ensure proper cleanup
 	FEdMode::Exit();
 }
@@ -137,6 +150,18 @@ void FEdModeTileMap::ActorSelectionChangeNotify()
 	}
 }
 
+bool FEdModeTileMap::MouseEnter(FEditorViewportClient* ViewportClient, FViewport* Viewport, int32 x, int32 y)
+{
+	RefreshBrushSize();
+	return FEdMode::MouseEnter(ViewportClient, Viewport, x, y);
+}
+
+bool FEdModeTileMap::MouseLeave(FEditorViewportClient* ViewportClient, FViewport* Viewport)
+{
+	CursorPreviewComponent->SetVisibility(false);
+	return FEdMode::MouseLeave(ViewportClient, Viewport);
+}
+
 bool FEdModeTileMap::MouseMove(FEditorViewportClient* InViewportClient, FViewport* InViewport, int32 x, int32 y)
 {
 	if (InViewportClient->EngineShowFlags.ModeWidgets)
@@ -348,7 +373,7 @@ UPaperTileMapComponent* FEdModeTileMap::FindSelectedComponent() const
 	return TileMapComponent;
 }
 
-UPaperTileLayer* FEdModeTileMap::GetSelectedLayerUnderCursor(const FViewportCursorLocation& Ray, int32& OutTileX, int32& OutTileY) const
+UPaperTileLayer* FEdModeTileMap::GetSelectedLayerUnderCursor(const FViewportCursorLocation& Ray, int32& OutTileX, int32& OutTileY, bool bAllowOutOfBounds) const
 {
 	const FVector TraceStart = Ray.GetOrigin();
 	const FVector TraceDir = Ray.GetDirection();
@@ -390,7 +415,8 @@ UPaperTileLayer* FEdModeTileMap::GetSelectedLayerUnderCursor(const FViewportCurs
 					OutTileX = FMath::FloorToInt(NormalizedX * TileMap->MapWidth);
 					OutTileY = FMath::FloorToInt(NormalizedY * TileMap->MapHeight);
 
-					if ((OutTileX > -BrushWidth) && (OutTileX < TileMap->MapWidth) && (OutTileY > -BrushHeight) && (OutTileY < TileMap->MapHeight))
+					const bool bInBounds = (OutTileX > -BrushWidth) && (OutTileX < TileMap->MapWidth) && (OutTileY > -BrushHeight) && (OutTileY < TileMap->MapHeight);
+					if (bInBounds || bAllowOutOfBounds)
 					{
 						return Layer;
 					}
@@ -732,9 +758,66 @@ void FEdModeTileMap::SetActivePaint(UPaperTileSet* TileSet, FIntPoint TopLeft, F
 	PaintSourceTileSet = TileSet;
 	PaintSourceTopLeft = TopLeft;
 	PaintSourceDimensions = Dimensions;
+
+	UPaperTileMap* PreviewMap = CursorPreviewComponent->TileMap;
+	PreviewMap->MapWidth = FMath::Max<int32>(Dimensions.X, 1);
+	PreviewMap->MapHeight = FMath::Max<int32>(Dimensions.Y, 1);
+
+	UPaperTileLayer* PreviewLayer = PreviewMap->TileLayers[0];
+	for (int32 Y = 0; Y < PreviewMap->MapHeight; ++Y)
+	{
+		for (int32 X = 0; X < PreviewMap->MapWidth; ++X)
+		{
+			FPaperTileInfo TileInfo;
+
+			int32 SourceX = X + PaintSourceTopLeft.X;
+			int32 SourceY = Y + PaintSourceTopLeft.Y;
+
+			if ((TileSet != nullptr) && (SourceX < TileSet->GetTileCountX()) && (SourceY < TileSet->GetTileCountY()))
+			{
+				TileInfo.PackedTileIndex = SourceX + (SourceY * TileSet->GetTileCountX());
+				TileInfo.TileSet = TileSet;
+			}
+
+			PreviewLayer->SetCell(X, Y, TileInfo);
+		}
+	}
+
+	CursorPreviewComponent->MarkRenderStateDirty();
+
 	RefreshBrushSize();
 }
 
+void FEdModeTileMap::SynchronizePreviewWithTileMap(UPaperTileMap* NewTileMap)
+{
+	UPaperTileMap* PreviewTileMap = CursorPreviewComponent->TileMap;
+
+	bool bPreviewComponentDirty = false;
+
+#define UE_CHANGE_IF_DIFFERENT(PropertyName) \
+	if (PreviewTileMap->PropertyName != NewTileMap->PropertyName) \
+	{ \
+		PreviewTileMap->PropertyName = NewTileMap->PropertyName; \
+		bPreviewComponentDirty = true; \
+	}
+
+	UE_CHANGE_IF_DIFFERENT(TileWidth);
+	UE_CHANGE_IF_DIFFERENT(TileHeight);
+	UE_CHANGE_IF_DIFFERENT(PixelsPerUnit);
+	UE_CHANGE_IF_DIFFERENT(SeparationPerTileX);
+	UE_CHANGE_IF_DIFFERENT(SeparationPerTileY);
+	UE_CHANGE_IF_DIFFERENT(SeparationPerLayer);
+	UE_CHANGE_IF_DIFFERENT(Material);
+	UE_CHANGE_IF_DIFFERENT(ProjectionMode);
+
+#undef UE_CHANGE_IF_DIFFERENT
+
+	if (bPreviewComponentDirty)
+	{
+		CursorPreviewComponent->MarkRenderStateDirty();
+	}
+}
+
 void FEdModeTileMap::UpdatePreviewCursor(const FViewportCursorLocation& Ray)
 {
 	DrawPreviewDimensionsLS = FVector::ZeroVector;
@@ -743,7 +826,7 @@ void FEdModeTileMap::UpdatePreviewCursor(const FViewportCursorLocation& Ray)
 	{
 		int32 LocalTileX0;
 		int32 LocalTileY0;
-		if (UPaperTileLayer* TileLayer = GetSelectedLayerUnderCursor(Ray, /*out*/ LocalTileX0, /*out*/ LocalTileY0))
+		if (UPaperTileLayer* TileLayer = GetSelectedLayerUnderCursor(Ray, /*out*/ LocalTileX0, /*out*/ LocalTileY0, /*bAllowOutOfBounds=*/ true))
 		{
 			const int32 LocalTileX1 = LocalTileX0 + CursorWidth;
 			const int32 LocalTileY1 = LocalTileY0 + CursorHeight;
@@ -759,6 +842,10 @@ void FEdModeTileMap::UpdatePreviewCursor(const FViewportCursorLocation& Ray)
 			DrawPreviewLocation = (WorldPosition + WorldPositionBR) * 0.5f;
 
 			DrawPreviewDimensionsLS = 0.5f*((PaperAxisX * CursorWidth * TileMap->TileWidth) + (PaperAxisY * -CursorHeight * TileMap->TileHeight));
+
+			const FVector ComponentPreviewLocation = ComponentToWorld.TransformPosition(TileMap->GetTilePositionInLocalSpace(LocalTileX0 + 0.5f, LocalTileY0 + 0.5f, LayerIndex));
+			CursorPreviewComponent->SetWorldLocation(ComponentPreviewLocation);
+			SynchronizePreviewWithTileMap(TileMap);
 		}
 	}
 }
@@ -805,8 +892,8 @@ void FEdModeTileMap::RefreshBrushSize()
 	BrushHeight = 1;
 	CursorWidth = 1;
 	CursorHeight = 1;
-
-	if ( GetActiveLayerPaintingMode() != ETileMapLayerPaintingMode::CollisionLayers )
+	
+	if (GetActiveLayerPaintingMode() != ETileMapLayerPaintingMode::CollisionLayers)
 	{
 		switch (ActiveTool)
 		{
@@ -815,18 +902,21 @@ void FEdModeTileMap::RefreshBrushSize()
 			BrushHeight = PaintSourceDimensions.Y;
 			CursorWidth = FMath::Max(1, PaintSourceDimensions.X);
 			CursorHeight = FMath::Max(1, PaintSourceDimensions.Y);
+			CursorPreviewComponent->SetVisibility(true);
 			break;
 		case ETileMapEditorTool::Eraser:
 			BrushWidth = EraseBrushSize;
 			BrushHeight = EraseBrushSize;
 			CursorWidth = EraseBrushSize;
 			CursorHeight = EraseBrushSize;
+			CursorPreviewComponent->SetVisibility(false);
 			break;
 		case ETileMapEditorTool::PaintBucket:
 			BrushWidth = PaintSourceDimensions.X;
 			BrushHeight = PaintSourceDimensions.Y;
 			CursorWidth = 1;
 			CursorHeight = 1;
+			CursorPreviewComponent->SetVisibility(false);
 			break;
 		default:
 			check(false);
diff --git a/Engine/Plugins/2D/Paper2D/Source/Paper2DEditor/Private/TileMapEditing/EdModeTileMap.h b/Engine/Plugins/2D/Paper2D/Source/Paper2DEditor/Private/TileMapEditing/EdModeTileMap.h
--- a/Engine/Plugins/2D/Paper2D/Source/Paper2DEditor/Private/TileMapEditing/EdModeTileMap.h
+++ b/Engine/Plugins/2D/Paper2D/Source/Paper2DEditor/Private/TileMapEditing/EdModeTileMap.h
@@ -45,6 +45,8 @@ public:
 	//virtual bool BoxSelect(FBox& InBox, bool bInSelect) override;
 
 	//	virtual void PostUndo() override;
+	virtual bool MouseEnter(FEditorViewportClient* ViewportClient, FViewport* Viewport, int32 x, int32 y) override;
+	virtual bool MouseLeave(FEditorViewportClient* ViewportClient, FViewport* Viewport) override;
 	virtual bool MouseMove(FEditorViewportClient* ViewportClient, FViewport* Viewport, int32 x, int32 y) override;
 	virtual bool CapturedMouseMove(FEditorViewportClient* InViewportClient, FViewport* InViewport, int32 InMouseX, int32 InMouseY) override;
 	virtual bool StartTracking(FEditorViewportClient* InViewportClient, FViewport* InViewport) override;
@@ -83,7 +85,7 @@ protected:
 
 	// Returns the selected layer under the cursor, and the intersection tile coordinates
 	// Note: The tile coordinates can be negative if the brush is off the top or left of the tile map, but still overlaps the map!!!
-	UPaperTileLayer* GetSelectedLayerUnderCursor(const FViewportCursorLocation& Ray, int32& OutTileX, int32& OutTileY) const;
+	UPaperTileLayer* GetSelectedLayerUnderCursor(const FViewportCursorLocation& Ray, int32& OutTileX, int32& OutTileY, bool bAllowOutOfBounds = false) const;
 
 	// Compute a world space ray from the screen space mouse coordinates
 	FViewportCursorLocation CalculateViewRay(FEditorViewportClient* InViewportClient, FViewport* InViewport);
@@ -93,16 +95,20 @@ protected:
 	void EnableTileMapEditMode();
 	bool IsTileMapEditModeActive() const;
 
+	void SynchronizePreviewWithTileMap(UPaperTileMap* NewTileMap);
 public:
 	UPaperTileMapComponent* FindSelectedComponent() const;
 
 protected:
 	bool bIsPainting;
+
+	// Ink source
 	TWeakObjectPtr<UPaperTileSet> PaintSourceTileSet;
 
 	FIntPoint PaintSourceTopLeft;
 	FIntPoint PaintSourceDimensions;
-
+	
+	//
 	FTransform DrawPreviewSpace;
 	FVector DrawPreviewLocation;
 	FVector DrawPreviewDimensionsLS;
@@ -114,6 +120,8 @@ protected:
 	int32 BrushWidth;
 	int32 BrushHeight;
 
+	UPaperTileMapComponent* CursorPreviewComponent;
+
 	ETileMapEditorTool::Type ActiveTool;
 	ETileMapLayerPaintingMode::Type LayerPaintingMode;
 	mutable FTransform ComponentToWorld;
