Paper2D: Tile map editing: Add current tool text right on the preview cursor for erase and paint bucket modes, and hide the preview when the mouse is outside of the window

diff --git a/Engine/Plugins/2D/Paper2D/Source/Paper2DEditor/Private/TileMapEditing/EdModeTileMap.cpp b/Engine/Plugins/2D/Paper2D/Source/Paper2DEditor/Private/TileMapEditing/EdModeTileMap.cpp
--- a/Engine/Plugins/2D/Paper2D/Source/Paper2DEditor/Private/TileMapEditing/EdModeTileMap.cpp
+++ b/Engine/Plugins/2D/Paper2D/Source/Paper2DEditor/Private/TileMapEditing/EdModeTileMap.cpp
@@ -152,12 +152,20 @@ void FEdModeTileMap::ActorSelectionChangeNotify()
 
 bool FEdModeTileMap::MouseEnter(FEditorViewportClient* ViewportClient, FViewport* Viewport, int32 x, int32 y)
 {
+	if (ViewportClient->EngineShowFlags.ModeWidgets)
+	{
+		const FViewportCursorLocation Ray = CalculateViewRay(ViewportClient, Viewport);
+		UpdatePreviewCursor(Ray);
+	}
+
 	RefreshBrushSize();
+
 	return FEdMode::MouseEnter(ViewportClient, Viewport, x, y);
 }
 
 bool FEdModeTileMap::MouseLeave(FEditorViewportClient* ViewportClient, FViewport* Viewport)
 {
+	DrawPreviewDimensionsLS = FVector::ZeroVector;
 	CursorPreviewComponent->SetVisibility(false);
 	return FEdMode::MouseLeave(ViewportClient, Viewport);
 }
@@ -322,6 +330,37 @@ void FEdModeTileMap::DrawHUD(FEditorViewportClient* ViewportClient, FViewport* V
 
 	FCanvasTextItem Msg(FVector2D(10, 30), FText::FromString(InkInfo), GEngine->GetMediumFont(), FLinearColor::White);
 	Canvas->DrawItem(Msg);
+
+
+	bool bDrawToolDescription = false;
+	FText ToolDescription = LOCTEXT("NoTool", "No tool selected");
+	switch (ActiveTool)
+	{
+	case ETileMapEditorTool::Eraser:
+		ToolDescription = LOCTEXT("EraserTool", "Eraser");
+		bDrawToolDescription = true;
+		break;
+	case ETileMapEditorTool::Paintbrush:
+		break;
+	case ETileMapEditorTool::PaintBucket:
+		ToolDescription = LOCTEXT("PaintBucketTool", "Fill");
+		bDrawToolDescription = true;
+		break;
+	}
+
+	if (bDrawToolDescription && !DrawPreviewDimensionsLS.IsNearlyZero())
+	{
+		const FString ToolDescriptionString = ToolDescription.ToString();
+
+		FVector2D ScreenSpacePreviewLocation;
+		if (View->WorldToPixel(DrawPreviewTopLeft, /*out*/ ScreenSpacePreviewLocation))
+		{
+			int32 XL;
+			int32 YL;
+			StringSize(GEngine->GetLargeFont(), XL, YL, *ToolDescriptionString);
+			Canvas->DrawShadowedString(ScreenSpacePreviewLocation.X, ScreenSpacePreviewLocation.Y - YL, *ToolDescriptionString, GEngine->GetLargeFont(), FLinearColor::White);
+		}
+	}
 }
 
 bool FEdModeTileMap::AllowWidgetMove()
@@ -835,7 +874,8 @@ void FEdModeTileMap::UpdatePreviewCursor(const FViewportCursorLocation& Ray)
 			int32 LayerIndex;
 			ensure(TileMap->TileLayers.Find(TileLayer, LayerIndex));
 
-			const FVector WorldPosition = ComponentToWorld.TransformPosition(TileMap->GetTilePositionInLocalSpace(LocalTileX0, LocalTileY0, LayerIndex));
+			DrawPreviewTopLeft = ComponentToWorld.TransformPosition(TileMap->GetTilePositionInLocalSpace(LocalTileX0, LocalTileY0, LayerIndex));
+			const FVector WorldPosition = DrawPreviewTopLeft;
 			const FVector WorldPositionBR = ComponentToWorld.TransformPosition(TileMap->GetTilePositionInLocalSpace(LocalTileX1, LocalTileY1, LayerIndex));
 
 			DrawPreviewSpace = ComponentToWorld;
@@ -892,6 +932,7 @@ void FEdModeTileMap::RefreshBrushSize()
 	BrushHeight = 1;
 	CursorWidth = 1;
 	CursorHeight = 1;
+	const bool bShowPreviewDesired = !DrawPreviewDimensionsLS.IsNearlyZero();
 	
 	if (GetActiveLayerPaintingMode() != ETileMapLayerPaintingMode::CollisionLayers)
 	{
@@ -902,7 +943,7 @@ void FEdModeTileMap::RefreshBrushSize()
 			BrushHeight = PaintSourceDimensions.Y;
 			CursorWidth = FMath::Max(1, PaintSourceDimensions.X);
 			CursorHeight = FMath::Max(1, PaintSourceDimensions.Y);
-			CursorPreviewComponent->SetVisibility(true);
+			CursorPreviewComponent->SetVisibility(bShowPreviewDesired);
 			break;
 		case ETileMapEditorTool::Eraser:
 			BrushWidth = EraseBrushSize;
diff --git a/Engine/Plugins/2D/Paper2D/Source/Paper2DEditor/Private/TileMapEditing/EdModeTileMap.h b/Engine/Plugins/2D/Paper2D/Source/Paper2DEditor/Private/TileMapEditing/EdModeTileMap.h
--- a/Engine/Plugins/2D/Paper2D/Source/Paper2DEditor/Private/TileMapEditing/EdModeTileMap.h
+++ b/Engine/Plugins/2D/Paper2D/Source/Paper2DEditor/Private/TileMapEditing/EdModeTileMap.h
@@ -42,16 +42,12 @@ public:
 	virtual bool UsesToolkits() const override;
 	virtual void Enter() override;
 	virtual void Exit() override;
-	//virtual bool BoxSelect(FBox& InBox, bool bInSelect) override;
-
-	//	virtual void PostUndo() override;
 	virtual bool MouseEnter(FEditorViewportClient* ViewportClient, FViewport* Viewport, int32 x, int32 y) override;
 	virtual bool MouseLeave(FEditorViewportClient* ViewportClient, FViewport* Viewport) override;
 	virtual bool MouseMove(FEditorViewportClient* ViewportClient, FViewport* Viewport, int32 x, int32 y) override;
 	virtual bool CapturedMouseMove(FEditorViewportClient* InViewportClient, FViewport* InViewport, int32 InMouseX, int32 InMouseY) override;
 	virtual bool StartTracking(FEditorViewportClient* InViewportClient, FViewport* InViewport) override;
 	virtual bool EndTracking(FEditorViewportClient* InViewportClient, FViewport* InViewport) override;
-	// 	virtual void Tick(FEditorViewportClient* ViewportClient,float DeltaTime) override;
 	virtual bool InputKey(FEditorViewportClient* InViewportClient, FViewport* InViewport, FKey InKey, EInputEvent InEvent) override;
 	virtual bool InputDelta(FEditorViewportClient* InViewportClient, FViewport* InViewport, FVector& InDrag, FRotator& InRot, FVector& InScale) override;
 	virtual void Render(const FSceneView* View, FViewport* Viewport, FPrimitiveDrawInterface* PDI) override;
@@ -60,8 +56,6 @@ public:
 	virtual bool AllowWidgetMove();
 	virtual bool ShouldDrawWidget() const override;
 	virtual bool UsesTransformWidget() const override;
-	// 	virtual int32 GetWidgetAxisToDraw(FWidget::EWidgetMode InWidgetMode) const override;
-	// 	virtual bool DisallowMouseDeltaTracking() const override;
 	// End of FEdMode interface
 
 	void SetActiveTool(ETileMapEditorTool::Type NewTool);
@@ -110,9 +104,16 @@ protected:
 	
 	//
 	FTransform DrawPreviewSpace;
+
+	// Center of preview rectangle
 	FVector DrawPreviewLocation;
+
+	// Size of rectangle
 	FVector DrawPreviewDimensionsLS;
 
+	// Top left of the component bounds
+	FVector DrawPreviewTopLeft;
+
 	int32 EraseBrushSize;
 
 	int32 CursorWidth;
