Paper2D: Tile map editing: Fix several issues with tile brush preview
- Prevent motion blur on the brush preview
- Fix a refresh issue that would prevent the preview from being larger than 4x4
- Orient the preview to match the component rotation
- Remove the textual ink preview info since it's no longer necessary

diff --git a/Engine/Plugins/2D/Paper2D/Source/Paper2DEditor/Private/TileMapEditing/EdModeTileMap.cpp b/Engine/Plugins/2D/Paper2D/Source/Paper2DEditor/Private/TileMapEditing/EdModeTileMap.cpp
--- a/Engine/Plugins/2D/Paper2D/Source/Paper2DEditor/Private/TileMapEditing/EdModeTileMap.cpp
+++ b/Engine/Plugins/2D/Paper2D/Source/Paper2DEditor/Private/TileMapEditing/EdModeTileMap.cpp
@@ -115,6 +115,7 @@ void FEdModeTileMap::Enter()
 	CursorPreviewComponent->UpdateBounds();
 	CursorPreviewComponent->AddToRoot();
 	CursorPreviewComponent->RegisterComponentWithWorld(World);
+	CursorPreviewComponent->SetMobility(EComponentMobility::Static);
 
 	SetActiveTool(ETileMapEditorTool::Paintbrush);
 	SetActiveLayerPaintingMode(ETileMapLayerPaintingMode::VisualLayers);
@@ -325,13 +326,6 @@ void FEdModeTileMap::Render(const FSceneView* View, FViewport* Viewport, FPrimit
 
 void FEdModeTileMap::DrawHUD(FEditorViewportClient* ViewportClient, FViewport* Viewport, const FSceneView* View, FCanvas* Canvas)
 {
-	FString InkInfo = FString::Printf(TEXT("Ink: (%d, %d)  %dx%d  %s"), PaintSourceTopLeft.X, PaintSourceTopLeft.Y, PaintSourceDimensions.X, PaintSourceDimensions.Y, 
-		(PaintSourceTileSet.Get() != nullptr) ? (*PaintSourceTileSet.Get()->GetName()) : TEXT("(no tile ink selected)"));
-
-	FCanvasTextItem Msg(FVector2D(10, 30), FText::FromString(InkInfo), GEngine->GetMediumFont(), FLinearColor::White);
-	Canvas->DrawItem(Msg);
-
-
 	bool bDrawToolDescription = false;
 	FText ToolDescription = LOCTEXT("NoTool", "No tool selected");
 	switch (ActiveTool)
@@ -801,6 +795,8 @@ void FEdModeTileMap::SetActivePaint(UPaperTileSet* TileSet, FIntPoint TopLeft, F
 	UPaperTileMap* PreviewMap = CursorPreviewComponent->TileMap;
 	PreviewMap->MapWidth = FMath::Max<int32>(Dimensions.X, 1);
 	PreviewMap->MapHeight = FMath::Max<int32>(Dimensions.Y, 1);
+	FPropertyChangedEvent EditedMapSizeEvent(UPaperTileMap::StaticClass()->FindPropertyByName(GET_MEMBER_NAME_CHECKED(UPaperTileMap, MapWidth)));
+	PreviewMap->PostEditChangeProperty(EditedMapSizeEvent);
 
 	UPaperTileLayer* PreviewLayer = PreviewMap->TileLayers[0];
 	for (int32 Y = 0; Y < PreviewMap->MapHeight; ++Y)
@@ -885,6 +881,7 @@ void FEdModeTileMap::UpdatePreviewCursor(const FViewportCursorLocation& Ray)
 
 			const FVector ComponentPreviewLocation = ComponentToWorld.TransformPosition(TileMap->GetTilePositionInLocalSpace(LocalTileX0 + 0.5f, LocalTileY0 + 0.5f, LayerIndex));
 			CursorPreviewComponent->SetWorldLocation(ComponentPreviewLocation);
+			CursorPreviewComponent->SetWorldRotation(FRotator(ComponentToWorld.GetRotation()));
 			SynchronizePreviewWithTileMap(TileMap);
 		}
 	}
