Paper2D: Optimize UPaperTileMap::PostEditChangeProperty to only reregister components actually using this particular tile map

diff --git a/Engine/Plugins/2D/Paper2D/Source/Paper2D/Private/PaperTileMap.cpp b/Engine/Plugins/2D/Paper2D/Source/Paper2D/Private/PaperTileMap.cpp
--- a/Engine/Plugins/2D/Paper2D/Source/Paper2D/Private/PaperTileMap.cpp
+++ b/Engine/Plugins/2D/Paper2D/Source/Paper2D/Private/PaperTileMap.cpp
@@ -37,12 +37,49 @@ UPaperTileMap::UPaperTileMap(const FObjectInitializer& ObjectInitializer)
 #include "PaperTileMapComponent.h"
 #include "ComponentReregisterContext.h"
 
+/** Removes all components that use the specified sprite asset from their scenes for the lifetime of the class. */
+class FTileMapReregisterContext
+{
+public:
+	/** Initialization constructor. */
+	FTileMapReregisterContext(UPaperTileMap* TargetAsset)
+	{
+		// Look at tile map components
+		for (TObjectIterator<UPaperTileMapComponent> TileMapIt; TileMapIt; ++TileMapIt)
+		{
+			if (UPaperTileMapComponent* TestComponent = *TileMapIt)
+			{
+				if (TestComponent->TileMap == TargetAsset)
+				{
+					AddComponentToRefresh(TestComponent);
+				}
+			}
+		}
+	}
+
+protected:
+	void AddComponentToRefresh(UActorComponent* Component)
+	{
+		if (ComponentContexts.Num() == 0)
+		{
+			// wait until resources are released
+			FlushRenderingCommands();
+		}
+
+		new (ComponentContexts) FComponentReregisterContext(Component);
+	}
+
+private:
+	/** The recreate contexts for the individual components. */
+	TIndirectArray<FComponentReregisterContext> ComponentContexts;
+};
+
 void UPaperTileMap::PostEditChangeProperty(FPropertyChangedEvent& PropertyChangedEvent)
 {
 	const FName PropertyName = (PropertyChangedEvent.Property != nullptr) ? PropertyChangedEvent.Property->GetFName() : NAME_None;
 
 	//@TODO: Determine when these are really needed, as they're seriously expensive!
-	TComponentReregisterContext<UPaperTileMapComponent> ReregisterStaticComponents;
+	FTileMapReregisterContext ReregisterTileMapComponents(this);
 
 	ValidateSelectedLayerIndex();
 
