Paper2D: Tile map importing now sets the selected tile set to the first tile set that was imported

diff --git a/Engine/Plugins/2D/Paper2D/Source/PaperTiledImporter/Private/PaperTiledImporterFactory.cpp b/Engine/Plugins/2D/Paper2D/Source/PaperTiledImporter/Private/PaperTiledImporterFactory.cpp
--- a/Engine/Plugins/2D/Paper2D/Source/PaperTiledImporter/Private/PaperTiledImporterFactory.cpp
+++ b/Engine/Plugins/2D/Paper2D/Source/PaperTiledImporter/Private/PaperTiledImporterFactory.cpp
@@ -237,6 +237,11 @@ UObject* UPaperTiledImporterFactory::FactoryCreateText(UClass* InClass, UObject*
 			}
 		}
 
+		if (GlobalInfo.CreatedTileSetAssets.Num() > 0)
+		{
+			// Bind our selected tile set to the first tile set that was imported so something is already picked
+			Result->SelectedTileSet = GlobalInfo.CreatedTileSetAssets[0];
+		}
 
 		// Create the layers
 		for (const FTileLayerFromTiled& LayerData : GlobalInfo.Layers)
