Paper2D: Add proper support for IsometricDiamond projection tile maps, fixing bounds calcuations, editor mouse support, etc...

diff --git a/Engine/Plugins/2D/Paper2D/Source/Paper2D/Classes/PaperTileMap.h b/Engine/Plugins/2D/Paper2D/Source/Paper2D/Classes/PaperTileMap.h
--- a/Engine/Plugins/2D/Paper2D/Source/Paper2D/Classes/PaperTileMap.h
+++ b/Engine/Plugins/2D/Paper2D/Source/Paper2D/Classes/PaperTileMap.h
@@ -108,9 +108,20 @@ public:
 #endif
 	// End of UObject interface
 
+	// Returns the tile coordinates of the specified local space position
+	void GetTileCoordinatesFromLocalSpacePosition(const FVector& Position, int32& OutTileX, int32& OutTileY) const;
 
+	// Returns the top left corner of the specified tile in local space
 	FVector GetTilePositionInLocalSpace(float TileX, float TileY, int32 LayerIndex = 0) const;
 
+	// Returns the center of the specified tile in local space
+	FVector GetTileCenterInLocalSpace(float TileX, float TileY, int32 LayerIndex = 0) const;
+
+
+	void GetTileToLocalParameters(FVector& OutCornerPosition, FVector& OutStepX, FVector& OutStepY, FVector& OutOffsetYFactor) const;
+	void GetLocalToTileParameters(FVector& OutCornerPosition, FVector& OutStepX, FVector& OutStepY, FVector& OutOffsetYFactor) const;
+
+
 	FBoxSphereBounds GetRenderBounds() const;
 
 	// Creates and adds a new layer and returns it
diff --git a/Engine/Plugins/2D/Paper2D/Source/Paper2D/Private/PaperTileMap.cpp b/Engine/Plugins/2D/Paper2D/Source/Paper2D/Private/PaperTileMap.cpp
--- a/Engine/Plugins/2D/Paper2D/Source/Paper2D/Private/PaperTileMap.cpp
+++ b/Engine/Plugins/2D/Paper2D/Source/Paper2D/Private/PaperTileMap.cpp
@@ -175,28 +175,146 @@ void UPaperTileMap::UpdateBodySetup()
 	}
 }
 
+void UPaperTileMap::GetTileToLocalParameters(FVector& OutCornerPosition, FVector& OutStepX, FVector& OutStepY, FVector& OutOffsetYFactor) const
+{
+	switch (ProjectionMode)
+	{
+	case ETileMapProjectionMode::Orthogonal:
+	default:
+		OutCornerPosition = -(TileWidth * PaperAxisX * 0.5f) + (TileHeight * PaperAxisY * 0.5f);
+		OutOffsetYFactor = FVector::ZeroVector;
+		OutStepX = PaperAxisX * TileWidth;
+		OutStepY = -PaperAxisY * TileHeight;
+		break;
+	case ETileMapProjectionMode::IsometricDiamond:
+		OutCornerPosition = (TileHeight * PaperAxisY * 0.5f);
+		OutOffsetYFactor = FVector::ZeroVector;
+		OutStepX = (TileWidth * PaperAxisX * 0.5f) - (TileHeight * PaperAxisY * 0.5f);
+		OutStepY = (TileWidth * PaperAxisX * -0.5f) - (TileHeight * PaperAxisY * 0.5f);
+		break;
+	case ETileMapProjectionMode::IsometricStaggered:
+		OutCornerPosition = 0.5f * TileHeight * PaperAxisY;
+		OutOffsetYFactor = 0.5f * TileWidth * PaperAxisX;
+		OutStepX = PaperAxisX * TileWidth;
+		OutStepY = -PaperAxisY * TileHeight;
+		break;
+	}
+}
+
+void UPaperTileMap::GetLocalToTileParameters(FVector& OutCornerPosition, FVector& OutStepX, FVector& OutStepY, FVector& OutOffsetYFactor) const
+{
+	switch (ProjectionMode)
+	{
+	case ETileMapProjectionMode::Orthogonal:
+	default:
+		OutCornerPosition = -(TileWidth * PaperAxisX * 0.5f) + (TileHeight * PaperAxisY * 0.5f);
+		OutOffsetYFactor = FVector::ZeroVector;
+		OutStepX = PaperAxisX / TileWidth;
+		OutStepY = -PaperAxisY / TileHeight;
+		break;
+	case ETileMapProjectionMode::IsometricDiamond:
+		OutCornerPosition = (TileHeight * PaperAxisY * 0.5f);
+		OutOffsetYFactor = FVector::ZeroVector;
+		OutStepX = (PaperAxisX / TileWidth) - (PaperAxisY / TileHeight);
+		OutStepY = (-PaperAxisX / TileWidth) - (PaperAxisY / TileHeight);
+		break;
+	case ETileMapProjectionMode::IsometricStaggered:
+		OutCornerPosition = 0.5f * TileHeight * PaperAxisY;
+		OutOffsetYFactor = 0.5f * TileWidth * PaperAxisX;
+		OutStepX = PaperAxisX / TileWidth;
+		OutStepY = -PaperAxisY / TileHeight;
+		break;
+	}
+}
+
+void UPaperTileMap::GetTileCoordinatesFromLocalSpacePosition(const FVector& Position, int32& OutTileX, int32& OutTileY) const
+{
+	FVector CornerPosition;
+	FVector OffsetYFactor;
+	FVector ParameterAxisX;
+	FVector ParameterAxisY;
+
+	// position is in pixels
+	// axis is tiles to pixels
+
+	GetLocalToTileParameters(/*out*/ CornerPosition, /*out*/ ParameterAxisX, /*out*/ ParameterAxisY, /*out*/ OffsetYFactor);
+
+	const FVector RelativePosition = Position - CornerPosition;
+	const float ProjectionSpaceXInTiles = FVector::DotProduct(RelativePosition, ParameterAxisX);
+	const float ProjectionSpaceYInTiles = FVector::DotProduct(RelativePosition, ParameterAxisY);
+
+	OutTileX = FMath::FloorToInt(ProjectionSpaceXInTiles);
+	OutTileY = FMath::FloorToInt(ProjectionSpaceYInTiles);
+}
+
 FVector UPaperTileMap::GetTilePositionInLocalSpace(float TileX, float TileY, int32 LayerIndex) const
 {
-	//@TODO: Tile pivot issue
-	const FVector PartialX = PaperAxisX * (TileX - 0.5f) * TileWidth;
-	const FVector PartialY = PaperAxisY * -(TileY - 0.5f) * TileHeight;
-	const FVector PartialZ = PaperAxisZ * (LayerIndex * SeparationPerLayer);
+	FVector CornerPosition;
+	FVector OffsetYFactor;
+	FVector StepX;
+	FVector StepY;
+
+	GetTileToLocalParameters(/*out*/ CornerPosition, /*out*/ StepX, /*out*/ StepY, /*out*/ OffsetYFactor);
+
+	FVector TotalOffset;
+	switch (ProjectionMode)
+	{
+	case ETileMapProjectionMode::Orthogonal:
+	default:
+		TotalOffset = CornerPosition;
+		break;
+	case ETileMapProjectionMode::IsometricDiamond:
+		TotalOffset = CornerPosition;
+		break;
+	case ETileMapProjectionMode::IsometricStaggered:
+		TotalOffset = CornerPosition + ((int32)TileY & 1) * OffsetYFactor;
+		break;
+	}
+
+	const FVector PartialX = TileX * StepX;
+	const FVector PartialY = TileY * StepY;
+
+	const float TotalSeparation = (SeparationPerLayer * LayerIndex) + (SeparationPerTileX * TileX) + (SeparationPerTileY * TileY);
+	const FVector PartialZ = TotalSeparation * PaperAxisZ;
+
+	const FVector LocalPos(PartialX + PartialY + PartialZ + TotalOffset);
 
-	const FVector LocalPos(PartialX + PartialY + PartialZ);
-	
 	return LocalPos;
 }
 
+FVector UPaperTileMap::GetTileCenterInLocalSpace(float TileX, float TileY, int32 LayerIndex) const
+{
+	return GetTilePositionInLocalSpace(TileX + 0.5f, TileY + 0.5f, LayerIndex);
+}
+
 FBoxSphereBounds UPaperTileMap::GetRenderBounds() const
 {
-	//@TODO: Tile pivot issue
 	const float Depth = SeparationPerLayer * (TileLayers.Num() - 1);
 	const float HalfThickness = 2.0f;
-	const FVector TopLeft((-0.5f)*TileWidth, -HalfThickness - Depth, -(MapHeight - 0.5f) * TileHeight);
-	const FVector Dimenisons(MapWidth*TileWidth, Depth + 2 * HalfThickness, MapHeight * TileHeight);
 
-	const FBox Box(TopLeft, TopLeft + Dimenisons);
-	return FBoxSphereBounds(Box);
+	switch (ProjectionMode)
+	{
+		case ETileMapProjectionMode::Orthogonal:
+		default:
+		{
+			const FVector BottomLeft((-0.5f) * TileWidth, -HalfThickness - Depth, -(MapHeight - 0.5f) * TileHeight);
+			const FVector Dimensions(MapWidth*TileWidth, Depth + 2 * HalfThickness, MapHeight * TileHeight);
+
+			const FBox Box(BottomLeft, BottomLeft + Dimensions);
+			return FBoxSphereBounds(Box);
+		}
+		case ETileMapProjectionMode::IsometricDiamond:
+		{
+			 const FVector BottomLeft((-0.5f) * TileWidth * MapWidth, -HalfThickness - Depth, -MapHeight * TileHeight);
+			 const FVector Dimensions(MapWidth*TileWidth, Depth + 2 * HalfThickness, (MapHeight + 1) * TileHeight);
+
+			 const FBox Box(BottomLeft, BottomLeft + Dimensions);
+			 return FBoxSphereBounds(Box);
+		}
+// 		case ETileMapProjectionMode::IsometricStaggered:
+// 		{
+// 		}
+	}
 }
 
 UPaperTileLayer* UPaperTileMap::AddNewLayer(bool bCollisionLayer, int32 InsertionIndex)
diff --git a/Engine/Plugins/2D/Paper2D/Source/Paper2D/Private/PaperTileMapComponent.cpp b/Engine/Plugins/2D/Paper2D/Source/Paper2D/Private/PaperTileMapComponent.cpp
--- a/Engine/Plugins/2D/Paper2D/Source/Paper2D/Private/PaperTileMapComponent.cpp
+++ b/Engine/Plugins/2D/Paper2D/Source/Paper2D/Private/PaperTileMapComponent.cpp
@@ -144,34 +144,12 @@ void UPaperTileMapComponent::RebuildRenderData(FPaperTileMapRenderSceneProxy* Pr
 		return;
 	}
 
-
-	FVector WorldOffsetNoY;
-	FVector WorldOffsetYFactor = FVector::ZeroVector;
-	FVector WorldStepX;
-	FVector WorldStepY;
-
-	switch (TileMap->ProjectionMode)
-	{
-	case ETileMapProjectionMode::Orthogonal:
-	default:
-		WorldOffsetNoY = -(TileMap->TileWidth * PaperAxisX * 0.5f) - (TileMap->TileHeight * PaperAxisY * 0.5f);
-		WorldStepX = PaperAxisX * TileMap->TileWidth;
-		WorldStepY = -PaperAxisY * TileMap->TileHeight;
-		break;
-	case ETileMapProjectionMode::IsometricDiamond:
-		WorldOffsetNoY = (TileMap->MapWidth * TileMap->TileWidth * 0.5f * PaperAxisX);
-		WorldOffsetYFactor = -(TileMap->TileHeight * 0.5f * PaperAxisY);
-		WorldStepX = (TileMap->TileWidth * PaperAxisX * 0.5f) - (TileMap->TileHeight * PaperAxisY * 0.5f);
-		WorldStepY = (TileMap->TileWidth * PaperAxisX * -0.5f) - (TileMap->TileHeight * PaperAxisY * 0.5f);
-		break;
-	case ETileMapProjectionMode::IsometricStaggered:
-		WorldOffsetNoY = 0.5f * TileMap->TileHeight * PaperAxisY;
-		WorldOffsetYFactor = 0.5f * TileMap->TileWidth * PaperAxisX;
-		WorldStepX = PaperAxisX * TileMap->TileWidth;
-		WorldStepY = -PaperAxisY * TileMap->TileHeight;
-		break;
-	}
-
+	FVector CornerOffset;
+	FVector OffsetYFactor;
+	FVector StepPerTileX;
+	FVector StepPerTileY;
+	TileMap->GetTileToLocalParameters(/*out*/ CornerOffset, /*out*/ StepPerTileX, /*out*/ StepPerTileY, /*out*/ OffsetYFactor);
+	
 	UTexture2D* LastSourceTexture = nullptr;
 	FVector TileSetOffset = FVector::ZeroVector;
 	FVector2D InverseTextureSize(1.0f, 1.0f);
@@ -202,19 +180,19 @@ void UPaperTileMapComponent::RebuildRenderData(FPaperTileMapRenderSceneProxy* Pr
 		for (int32 Y = 0; Y < TileMap->MapHeight; ++Y)
 		{
 			// In pixels
-			FVector WorldOffset;
+			FVector EffectiveTopLeftCorner;
 
 			switch (TileMap->ProjectionMode)
 			{
 			case ETileMapProjectionMode::Orthogonal:
 			default:
-				WorldOffset = WorldOffsetNoY;
+				EffectiveTopLeftCorner = CornerOffset;
 				break;
 			case ETileMapProjectionMode::IsometricDiamond:
-				WorldOffset = WorldOffsetNoY + Y * WorldOffsetYFactor;
+				EffectiveTopLeftCorner = CornerOffset - StepPerTileX;
 				break;
 			case ETileMapProjectionMode::IsometricStaggered:
-				WorldOffset = WorldOffsetNoY + (Y & 1) * WorldOffsetYFactor;
+				EffectiveTopLeftCorner = CornerOffset + (Y & 1) * OffsetYFactor;
 				break;
 			}
 
@@ -224,8 +202,8 @@ void UPaperTileMapComponent::RebuildRenderData(FPaperTileMapRenderSceneProxy* Pr
 
 				// do stuff
 				const float TotalSeparation = (TileMap->SeparationPerLayer * Z) + (TileMap->SeparationPerTileX * X) + (TileMap->SeparationPerTileY * Y);
-				FVector WorldPos = (WorldStepX * X) + (WorldStepY * Y) + WorldOffset;
-				WorldPos += TotalSeparation * PaperAxisZ;
+				FVector TopLeftCornerOfTile = (StepPerTileX * X) + (StepPerTileY * Y) + EffectiveTopLeftCorner;
+				TopLeftCornerOfTile += TotalSeparation * PaperAxisZ;
 
 				const int32 TileWidth = TileMap->TileWidth;
 				const int32 TileHeight = TileMap->TileHeight;
@@ -267,7 +245,7 @@ void UPaperTileMapComponent::RebuildRenderData(FPaperTileMapRenderSceneProxy* Pr
 						CurrentBatch = (new (BatchedSprites) FSpriteDrawCallRecord());
 						CurrentBatch->Texture = SourceTexture;
 						CurrentBatch->Color = DrawColor;
-						CurrentBatch->Destination = WorldPos.ProjectOnTo(PaperAxisZ);
+						CurrentBatch->Destination = TopLeftCornerOfTile.ProjectOnTo(PaperAxisZ);
 					}
 
 					if (SourceTexture != LastSourceTexture)
@@ -288,20 +266,20 @@ void UPaperTileMapComponent::RebuildRenderData(FPaperTileMapRenderSceneProxy* Pr
 						}
 						LastSourceTexture = SourceTexture;
 					}
-					WorldPos += TileSetOffset;
+					TopLeftCornerOfTile += TileSetOffset;
 
 					SourceUV.X *= InverseTextureSize.X;
 					SourceUV.Y *= InverseTextureSize.Y;
 
 					FSpriteDrawCallRecord& NewTile = *CurrentBatch;
 
-					const float WX0 = FVector::DotProduct(WorldPos, PaperAxisX);
-					const float WY0 = FVector::DotProduct(WorldPos, PaperAxisY);
+					const float WX0 = FVector::DotProduct(TopLeftCornerOfTile, PaperAxisX);
+					const float WY0 = FVector::DotProduct(TopLeftCornerOfTile, PaperAxisY);
 
-					const FVector4 BottomLeft(WX0, WY0, SourceUV.X, SourceUV.Y + SourceDimensionsUV.Y);
-					const FVector4 BottomRight(WX0 + TileSizeXY.X, WY0, SourceUV.X + SourceDimensionsUV.X, SourceUV.Y + SourceDimensionsUV.Y);
-					const FVector4 TopRight(WX0 + TileSizeXY.X, WY0 + TileSizeXY.Y, SourceUV.X + SourceDimensionsUV.X, SourceUV.Y);
-					const FVector4 TopLeft(WX0, WY0 + TileSizeXY.Y, SourceUV.X, SourceUV.Y);
+					const FVector4 BottomLeft(WX0, WY0 - TileSizeXY.Y, SourceUV.X, SourceUV.Y + SourceDimensionsUV.Y);
+					const FVector4 BottomRight(WX0 + TileSizeXY.X, WY0 - TileSizeXY.Y, SourceUV.X + SourceDimensionsUV.X, SourceUV.Y + SourceDimensionsUV.Y);
+					const FVector4 TopRight(WX0 + TileSizeXY.X, WY0, SourceUV.X + SourceDimensionsUV.X, SourceUV.Y);
+					const FVector4 TopLeft(WX0, WY0, SourceUV.X, SourceUV.Y);
 
 					new (NewTile.RenderVerts) FVector4(BottomLeft);
 					new (NewTile.RenderVerts) FVector4(TopRight);
diff --git a/Engine/Plugins/2D/Paper2D/Source/Paper2DEditor/Private/TileMapEditing/EdModeTileMap.cpp b/Engine/Plugins/2D/Paper2D/Source/Paper2DEditor/Private/TileMapEditing/EdModeTileMap.cpp
--- a/Engine/Plugins/2D/Paper2D/Source/Paper2DEditor/Private/TileMapEditing/EdModeTileMap.cpp
+++ b/Engine/Plugins/2D/Paper2D/Source/Paper2DEditor/Private/TileMapEditing/EdModeTileMap.cpp
@@ -83,10 +83,12 @@ FEdModeTileMap::FEdModeTileMap()
 	, PaintSourceTileSet(nullptr)
 	, PaintSourceTopLeft(0, 0)
 	, PaintSourceDimensions(0, 0)
+	, bIsLastCursorValid(false)
 	, DrawPreviewDimensionsLS(0.0f, 0.0f, 0.0f)
 	, EraseBrushSize(1)
 {
 	bDrawPivot = false;
+	bDrawGrid = false;
 }
 
 FEdModeTileMap::~FEdModeTileMap()
@@ -167,6 +169,9 @@ bool FEdModeTileMap::MouseEnter(FEditorViewportClient* ViewportClient, FViewport
 bool FEdModeTileMap::MouseLeave(FEditorViewportClient* ViewportClient, FViewport* Viewport)
 {
 	DrawPreviewDimensionsLS = FVector::ZeroVector;
+	bIsLastCursorValid = false;
+	LastCursorTileMap.Reset();
+
 	CursorPreviewComponent->SetVisibility(false);
 	return FEdMode::MouseLeave(ViewportClient, Viewport);
 }
@@ -304,23 +309,25 @@ void FEdModeTileMap::Render(const FSceneView* View, FViewport* Viewport, FPrimit
 		return;
 	}
 
-	// Make sure the cursor is visible OR we're flood filling.  No point drawing a paint cue when there's no cursor.
-// 	if( Viewport->IsCursorVisible() || IsForceRendered())
-// 	{
-// 		if( !PDI->IsHitTesting() )
-// 		{
-// 		}
-// 	}
-
 	// Draw the preview cursor
-	if (!DrawPreviewDimensionsLS.IsNearlyZero())
+	if (bIsLastCursorValid)
 	{
-		FVector X = DrawPreviewSpace.GetScaledAxis(EAxis::X);
-		FVector Y = DrawPreviewSpace.GetScaledAxis(EAxis::Y);
-		FVector Z = DrawPreviewSpace.GetScaledAxis(EAxis::Z);
-		FVector Base = DrawPreviewLocation;
-
-		DrawOrientedWireBox(PDI, Base, X, Y, Z, DrawPreviewDimensionsLS, FLinearColor::White, SDPG_Foreground, 0.0f, 0.0001f);
+		if (UPaperTileMap* TileMap = LastCursorTileMap.Get())
+		{
+			// Slight depth bias so that the wireframe grid overlay doesn't z-fight with the tiles themselves
+			const float DepthBias = 0.0001f;
+			FLinearColor CursorWireColor = FLinearColor::White;
+
+			const FVector TL(ComponentToWorld.TransformPosition(TileMap->GetTilePositionInLocalSpace(LastCursorTileX + 0, LastCursorTileY + 0, LastCursorTileZ)));
+			const FVector TR(ComponentToWorld.TransformPosition(TileMap->GetTilePositionInLocalSpace(LastCursorTileX + CursorWidth, LastCursorTileY + 0, LastCursorTileZ)));
+			const FVector BL(ComponentToWorld.TransformPosition(TileMap->GetTilePositionInLocalSpace(LastCursorTileX + 0, LastCursorTileY + CursorHeight, LastCursorTileZ)));
+			const FVector BR(ComponentToWorld.TransformPosition(TileMap->GetTilePositionInLocalSpace(LastCursorTileX + CursorWidth, LastCursorTileY + CursorHeight, LastCursorTileZ)));
+
+			PDI->DrawLine(TL, TR, CursorWireColor, SDPG_Foreground, 0.0f, DepthBias);
+			PDI->DrawLine(TR, BR, CursorWireColor, SDPG_Foreground, 0.0f, DepthBias);
+			PDI->DrawLine(BR, BL, CursorWireColor, SDPG_Foreground, 0.0f, DepthBias);
+			PDI->DrawLine(BL, TL, CursorWireColor, SDPG_Foreground, 0.0f, DepthBias);
+		}
 	}
 }
 
@@ -355,6 +362,21 @@ void FEdModeTileMap::DrawHUD(FEditorViewportClient* ViewportClient, FViewport* V
 			Canvas->DrawShadowedString(ScreenSpacePreviewLocation.X, ScreenSpacePreviewLocation.Y - YL, *ToolDescriptionString, GEngine->GetLargeFont(), FLinearColor::White);
 		}
 	}
+
+
+	if (bIsLastCursorValid)
+	{
+		const FString ToolDescriptionString = FString::Printf(TEXT("%d, %d"), LastCursorTileX, LastCursorTileY);
+
+		FVector2D ScreenSpacePreviewLocation;
+		if (View->WorldToPixel(DrawPreviewTopLeft, /*out*/ ScreenSpacePreviewLocation))
+		{
+			int32 XL;
+			int32 YL;
+			StringSize(GEngine->GetLargeFont(), XL, YL, *ToolDescriptionString);
+			Canvas->DrawShadowedString(ScreenSpacePreviewLocation.X, ScreenSpacePreviewLocation.Y - YL, *ToolDescriptionString, GEngine->GetLargeFont(), FLinearColor::White);
+		}
+	}
 }
 
 bool FEdModeTileMap::AllowWidgetMove()
@@ -427,9 +449,6 @@ UPaperTileLayer* FEdModeTileMap::GetSelectedLayerUnderCursor(const FViewportCurs
 			{
 				UPaperTileLayer* Layer = TileMap->TileLayers[LayerIndex];
 
-				const float WX = TileMap->MapWidth * TileMap->TileWidth;
-				const float WY = TileMap->MapHeight * TileMap->TileHeight;
-
 				ComponentToWorld = (TileMapComponent != nullptr) ? TileMapComponent->ComponentToWorld : FTransform::Identity;
 				const FVector LocalStart = ComponentToWorld.InverseTransformPosition(TraceStart);
 				const FVector LocalDirection = ComponentToWorld.InverseTransformVector(TraceDir);
@@ -442,11 +461,7 @@ UPaperTileLayer* FEdModeTileMap::GetSelectedLayerUnderCursor(const FViewportCurs
 				FVector Intersection;
 				if (FMath::SegmentPlaneIntersection(LocalStart, LocalEnd, LayerPlane, /*out*/ Intersection))
 				{
-					const float NormalizedX = (FVector::DotProduct(Intersection, PaperAxisX) + 0.5f * TileMap->TileWidth) / WX;
-					const float NormalizedY = (-FVector::DotProduct(Intersection, PaperAxisY) + 0.5f * TileMap->TileHeight) / WY;
-
-					OutTileX = FMath::FloorToInt(NormalizedX * TileMap->MapWidth);
-					OutTileY = FMath::FloorToInt(NormalizedY * TileMap->MapHeight);
+					TileMap->GetTileCoordinatesFromLocalSpacePosition(Intersection, /*out*/ OutTileX, /*out*/ OutTileY);
 
 					const bool bInBounds = (OutTileX > -BrushWidth) && (OutTileX < TileMap->MapWidth) && (OutTileY > -BrushHeight) && (OutTileY < TileMap->MapHeight);
 					if (bInBounds || bAllowOutOfBounds)
@@ -856,6 +871,8 @@ void FEdModeTileMap::SynchronizePreviewWithTileMap(UPaperTileMap* NewTileMap)
 void FEdModeTileMap::UpdatePreviewCursor(const FViewportCursorLocation& Ray)
 {
 	DrawPreviewDimensionsLS = FVector::ZeroVector;
+	bIsLastCursorValid = false;
+	LastCursorTileMap.Reset();
 
 	// See if we should draw the preview
 	{
@@ -863,13 +880,19 @@ void FEdModeTileMap::UpdatePreviewCursor(const FViewportCursorLocation& Ray)
 		int32 LocalTileY0;
 		if (UPaperTileLayer* TileLayer = GetSelectedLayerUnderCursor(Ray, /*out*/ LocalTileX0, /*out*/ LocalTileY0, /*bAllowOutOfBounds=*/ true))
 		{
-			const int32 LocalTileX1 = LocalTileX0 + CursorWidth;
-			const int32 LocalTileY1 = LocalTileY0 + CursorHeight;
-
 			UPaperTileMap* TileMap = TileLayer->GetTileMap();
 			int32 LayerIndex;
 			ensure(TileMap->TileLayers.Find(TileLayer, LayerIndex));
 
+			LastCursorTileX = LocalTileX0;
+			LastCursorTileY = LocalTileY0;
+			LastCursorTileZ = LayerIndex;
+			bIsLastCursorValid = true;
+			LastCursorTileMap = TileMap;
+
+			const int32 LocalTileX1 = LocalTileX0 + CursorWidth;
+			const int32 LocalTileY1 = LocalTileY0 + CursorHeight;
+
 			DrawPreviewTopLeft = ComponentToWorld.TransformPosition(TileMap->GetTilePositionInLocalSpace(LocalTileX0, LocalTileY0, LayerIndex));
 			const FVector WorldPosition = DrawPreviewTopLeft;
 			const FVector WorldPositionBR = ComponentToWorld.TransformPosition(TileMap->GetTilePositionInLocalSpace(LocalTileX1, LocalTileY1, LayerIndex));
@@ -879,7 +902,7 @@ void FEdModeTileMap::UpdatePreviewCursor(const FViewportCursorLocation& Ray)
 
 			DrawPreviewDimensionsLS = 0.5f*((PaperAxisX * CursorWidth * TileMap->TileWidth) + (PaperAxisY * -CursorHeight * TileMap->TileHeight));
 
-			const FVector ComponentPreviewLocation = ComponentToWorld.TransformPosition(TileMap->GetTilePositionInLocalSpace(LocalTileX0 + 0.5f, LocalTileY0 + 0.5f, LayerIndex));
+			const FVector ComponentPreviewLocation = ComponentToWorld.TransformPosition(TileMap->GetTileCenterInLocalSpace(LocalTileX0, LocalTileY0, LayerIndex));
 			CursorPreviewComponent->SetWorldLocation(ComponentPreviewLocation);
 			CursorPreviewComponent->SetWorldRotation(FRotator(ComponentToWorld.GetRotation()));
 			SynchronizePreviewWithTileMap(TileMap);
diff --git a/Engine/Plugins/2D/Paper2D/Source/Paper2DEditor/Private/TileMapEditing/EdModeTileMap.h b/Engine/Plugins/2D/Paper2D/Source/Paper2DEditor/Private/TileMapEditing/EdModeTileMap.h
--- a/Engine/Plugins/2D/Paper2D/Source/Paper2DEditor/Private/TileMapEditing/EdModeTileMap.h
+++ b/Engine/Plugins/2D/Paper2D/Source/Paper2DEditor/Private/TileMapEditing/EdModeTileMap.h
@@ -109,6 +109,12 @@ protected:
 	FVector DrawPreviewLocation;
 
 	// Size of rectangle
+	int32 LastCursorTileX;
+	int32 LastCursorTileY;
+	int32 LastCursorTileZ;
+	bool bIsLastCursorValid;
+	TWeakObjectPtr<UPaperTileMap> LastCursorTileMap;
+
 	FVector DrawPreviewDimensionsLS;
 
 	// Top left of the component bounds
