Paper2D: Disable built-in grid rendering in the tile map editor

diff --git a/Engine/Plugins/2D/Paper2D/Source/Paper2DEditor/Private/TileMapEditing/TileMapEditorViewportClient.cpp b/Engine/Plugins/2D/Paper2D/Source/Paper2DEditor/Private/TileMapEditing/TileMapEditorViewportClient.cpp
--- a/Engine/Plugins/2D/Paper2D/Source/Paper2DEditor/Private/TileMapEditing/TileMapEditorViewportClient.cpp
+++ b/Engine/Plugins/2D/Paper2D/Source/Paper2DEditor/Private/TileMapEditing/TileMapEditorViewportClient.cpp
@@ -34,6 +34,8 @@ FTileMapEditorViewportClient::FTileMapEditorViewportClient(TWeakPtr<FTileMapEdit
 	bManipulationDirtiedSomething = false;
 	ScopedTransaction = nullptr;
 
+	DrawHelper.bDrawGrid = false;
+	DrawHelper.bDrawPivot = false;
 	bShowPivot = true;
 
 	bDeferZoomToTileMap = true;
