Paper2D: Add a per-asset configurable background color for the tile set viewer

diff --git a/Engine/Plugins/2D/Paper2D/Source/Paper2D/Classes/PaperTileSet.h b/Engine/Plugins/2D/Paper2D/Source/Paper2D/Classes/PaperTileSet.h
--- a/Engine/Plugins/2D/Paper2D/Source/Paper2D/Classes/PaperTileSet.h
+++ b/Engine/Plugins/2D/Paper2D/Source/Paper2D/Classes/PaperTileSet.h
@@ -36,6 +36,13 @@ class PAPER2D_API UPaperTileSet : public UObject
 	UPROPERTY(Category=TileSet, EditAnywhere, meta=(UIMin=0, ClampMin=0))
 	FIntPoint DrawingOffset;
 
+#if WITH_EDITORONLY_DATA
+	/** The background color displayed in the tile set viewer */
+	UPROPERTY(Category=TileSet, EditAnywhere)
+	FLinearColor BackgroundColor;
+#endif
+
+
 	// UObject interface
 #if WITH_EDITOR
 	virtual void PostEditChangeProperty(FPropertyChangedEvent& PropertyChangedEvent) override;
diff --git a/Engine/Plugins/2D/Paper2D/Source/Paper2D/Private/PaperTileSet.cpp b/Engine/Plugins/2D/Paper2D/Source/Paper2D/Private/PaperTileSet.cpp
--- a/Engine/Plugins/2D/Paper2D/Source/Paper2D/Private/PaperTileSet.cpp
+++ b/Engine/Plugins/2D/Paper2D/Source/Paper2D/Private/PaperTileSet.cpp
@@ -10,6 +10,7 @@ UPaperTileSet::UPaperTileSet(const FObjectInitializer& ObjectInitializer)
 {
 	TileWidth = 32;
 	TileHeight = 32;
+	BackgroundColor = FColor(0, 0, 127);
 }
 
 int32 UPaperTileSet::GetTileCount() const
diff --git a/Engine/Plugins/2D/Paper2D/Source/Paper2DEditor/Private/PaperEditorViewportClient.cpp b/Engine/Plugins/2D/Paper2D/Source/Paper2DEditor/Private/PaperEditorViewportClient.cpp
--- a/Engine/Plugins/2D/Paper2D/Source/Paper2DEditor/Private/PaperEditorViewportClient.cpp
+++ b/Engine/Plugins/2D/Paper2D/Source/Paper2DEditor/Private/PaperEditorViewportClient.cpp
@@ -96,9 +96,15 @@ FPaperEditorViewportClient::~FPaperEditorViewportClient()
 {
 }
 
+FLinearColor FPaperEditorViewportClient::GetBackgroundColor() const
+{
+	//@TODO: Make adjustable - TextureEditorPtr.Pin()->GetBackgroundColor());
+	return FLinearColor(0, 0, 127, 0);
+}
+
 void FPaperEditorViewportClient::Draw(FViewport* Viewport, FCanvas* Canvas)
 {
-	Canvas->Clear(FColor(0, 0, 127, 0));//@TODO: Make adjustable - TextureEditorPtr.Pin()->GetBackgroundColor());
+	Canvas->Clear(GetBackgroundColor());
 }
 
 void FPaperEditorViewportClient::DrawSelectionRectangles(FViewport* Viewport, FCanvas* Canvas)
diff --git a/Engine/Plugins/2D/Paper2D/Source/Paper2DEditor/Private/PaperEditorViewportClient.h b/Engine/Plugins/2D/Paper2D/Source/Paper2DEditor/Private/PaperEditorViewportClient.h
--- a/Engine/Plugins/2D/Paper2D/Source/Paper2DEditor/Private/PaperEditorViewportClient.h
+++ b/Engine/Plugins/2D/Paper2D/Source/Paper2DEditor/Private/PaperEditorViewportClient.h
@@ -47,6 +47,10 @@ public:
 	virtual void Draw(FViewport* Viewport, FCanvas* Canvas) override;
 	// End of FViewportClient interface
 
+	// FEditorViewportClient interface
+	virtual FLinearColor GetBackgroundColor() const override;
+	// End of FEditorViewportClient interface
+
 	// FSerializableObject interface
 	virtual void AddReferencedObjects(FReferenceCollector& Collector) override;
 	// End of FSerializableObject interface
diff --git a/Engine/Plugins/2D/Paper2D/Source/Paper2DEditor/Private/TileSetEditor.cpp b/Engine/Plugins/2D/Paper2D/Source/Paper2DEditor/Private/TileSetEditor.cpp
--- a/Engine/Plugins/2D/Paper2D/Source/Paper2DEditor/Private/TileSetEditor.cpp
+++ b/Engine/Plugins/2D/Paper2D/Source/Paper2DEditor/Private/TileSetEditor.cpp
@@ -54,6 +54,10 @@ public:
 	virtual void Draw(FViewport* Viewport, FCanvas* Canvas) override;
 	// End of FViewportClient interface
 
+	// FEditorViewportClient interface
+	virtual FLinearColor GetBackgroundColor() const override;
+	// End of FEditorViewportClient interface
+
 public:
 	// Tile set
 	TWeakObjectPtr<UPaperTileSet> TileSetBeingEdited;
@@ -125,6 +129,18 @@ void FTileSetEditorViewportClient::Draw(FViewport* Viewport, FCanvas* Canvas)
 	}
 }
 
+FLinearColor FTileSetEditorViewportClient::GetBackgroundColor() const
+{
+	if (UPaperTileSet* TileSet = TileSetBeingEdited.Get())
+	{
+		return TileSet->BackgroundColor;
+	}
+	else
+	{
+		return FEditorViewportClient::GetBackgroundColor();
+	}
+}
+
 //////////////////////////////////////////////////////////////////////////
 // STileSetSelectorViewport
 
