Paper2D: Add a per-asset configurable background color for the tile map editor

diff --git a/Engine/Plugins/2D/Paper2D/Source/Paper2D/Classes/PaperTileMap.h b/Engine/Plugins/2D/Paper2D/Source/Paper2D/Classes/PaperTileMap.h
--- a/Engine/Plugins/2D/Paper2D/Source/Paper2D/Classes/PaperTileMap.h
+++ b/Engine/Plugins/2D/Paper2D/Source/Paper2D/Classes/PaperTileMap.h
@@ -93,6 +93,10 @@ public:
 	/** The currently selected layer index */
 	UPROPERTY()
 	int32 SelectedLayerIndex;
+
+	/** The background color displayed in the tile map editor */
+	UPROPERTY(Category=Setup, EditAnywhere)
+	FLinearColor BackgroundColor;
 #endif
 
 	/** The naming index to start at when trying to create a new layer */
diff --git a/Engine/Plugins/2D/Paper2D/Source/Paper2D/Private/PaperTileMap.cpp b/Engine/Plugins/2D/Paper2D/Source/Paper2D/Private/PaperTileMap.cpp
--- a/Engine/Plugins/2D/Paper2D/Source/Paper2D/Private/PaperTileMap.cpp
+++ b/Engine/Plugins/2D/Paper2D/Source/Paper2D/Private/PaperTileMap.cpp
@@ -24,6 +24,7 @@ UPaperTileMap::UPaperTileMap(const FObjectInitializer& ObjectInitializer)
 
 #if WITH_EDITORONLY_DATA
 	SelectedLayerIndex = INDEX_NONE;
+	BackgroundColor = FColor(55, 55, 55);
 #endif
 
 	LayerNameIndex = 0;
diff --git a/Engine/Plugins/2D/Paper2D/Source/Paper2DEditor/Private/TileMapEditing/TileMapEditorViewportClient.cpp b/Engine/Plugins/2D/Paper2D/Source/Paper2DEditor/Private/TileMapEditing/TileMapEditorViewportClient.cpp
--- a/Engine/Plugins/2D/Paper2D/Source/Paper2DEditor/Private/TileMapEditing/TileMapEditorViewportClient.cpp
+++ b/Engine/Plugins/2D/Paper2D/Source/Paper2DEditor/Private/TileMapEditing/TileMapEditorViewportClient.cpp
@@ -148,6 +148,18 @@ void FTileMapEditorViewportClient::Tick(float DeltaSeconds)
 	}
 }
 
+FLinearColor FTileMapEditorViewportClient::GetBackgroundColor() const
+{
+	if (UPaperTileMap* TileMap = RenderTileMapComponent->TileMap)
+	{
+		return TileMap->BackgroundColor;
+	}
+	else
+	{
+		return FEditorViewportClient::GetBackgroundColor();
+	}
+}
+
 void FTileMapEditorViewportClient::ToggleShowMeshEdges()
 {
 	EngineShowFlags.MeshEdges = !EngineShowFlags.MeshEdges;
diff --git a/Engine/Plugins/2D/Paper2D/Source/Paper2DEditor/Private/TileMapEditing/TileMapEditorViewportClient.h b/Engine/Plugins/2D/Paper2D/Source/Paper2DEditor/Private/TileMapEditing/TileMapEditorViewportClient.h
--- a/Engine/Plugins/2D/Paper2D/Source/Paper2DEditor/Private/TileMapEditing/TileMapEditorViewportClient.h
+++ b/Engine/Plugins/2D/Paper2D/Source/Paper2DEditor/Private/TileMapEditing/TileMapEditorViewportClient.h
@@ -27,6 +27,7 @@ public:
 	// End of FViewportClient interface
 
 	// FEditorViewportClient interface
+	virtual FLinearColor GetBackgroundColor() const override;
 	// End of FEditorViewportClient interface
 
 	void ToggleShowPivot() { bShowPivot = !bShowPivot; Invalidate(); }
