Paper2D: Fix overlap between (x, y) display and tool type display on the tile map editor preview cursor

diff --git a/Engine/Plugins/2D/Paper2D/Source/Paper2DEditor/Private/TileMapEditing/EdModeTileMap.cpp b/Engine/Plugins/2D/Paper2D/Source/Paper2DEditor/Private/TileMapEditing/EdModeTileMap.cpp
--- a/Engine/Plugins/2D/Paper2D/Source/Paper2DEditor/Private/TileMapEditing/EdModeTileMap.cpp
+++ b/Engine/Plugins/2D/Paper2D/Source/Paper2DEditor/Private/TileMapEditing/EdModeTileMap.cpp
@@ -342,6 +342,8 @@ void FEdModeTileMap::DrawHUD(FEditorViewportClient* ViewportClient, FViewport* V
 		bDrawToolDescription = true;
 		break;
 	case ETileMapEditorTool::Paintbrush:
+		ToolDescription = FText::GetEmpty();
+		bDrawToolDescription = true;
 		break;
 	case ETileMapEditorTool::PaintBucket:
 		ToolDescription = LOCTEXT("PaintBucketTool", "Fill");
@@ -351,22 +353,7 @@ void FEdModeTileMap::DrawHUD(FEditorViewportClient* ViewportClient, FViewport* V
 
 	if (bDrawToolDescription && !DrawPreviewDimensionsLS.IsNearlyZero())
 	{
-		const FString ToolDescriptionString = ToolDescription.ToString();
-
-		FVector2D ScreenSpacePreviewLocation;
-		if (View->WorldToPixel(DrawPreviewTopLeft, /*out*/ ScreenSpacePreviewLocation))
-		{
-			int32 XL;
-			int32 YL;
-			StringSize(GEngine->GetLargeFont(), XL, YL, *ToolDescriptionString);
-			Canvas->DrawShadowedString(ScreenSpacePreviewLocation.X, ScreenSpacePreviewLocation.Y - YL, *ToolDescriptionString, GEngine->GetLargeFont(), FLinearColor::White);
-		}
-	}
-
-
-	if (bIsLastCursorValid)
-	{
-		const FString ToolDescriptionString = FString::Printf(TEXT("%d, %d"), LastCursorTileX, LastCursorTileY);
+		const FString ToolDescriptionString = FString::Printf(TEXT("(%d, %d) %s"), LastCursorTileX, LastCursorTileY, *ToolDescription.ToString());
 
 		FVector2D ScreenSpacePreviewLocation;
 		if (View->WorldToPixel(DrawPreviewTopLeft, /*out*/ ScreenSpacePreviewLocation))
