Build fix for non-editor configs

diff --git a/Engine/Plugins/2D/Paper2D/Source/Paper2D/Private/PaperTileSet.cpp b/Engine/Plugins/2D/Paper2D/Source/Paper2D/Private/PaperTileSet.cpp
--- a/Engine/Plugins/2D/Paper2D/Source/Paper2D/Private/PaperTileSet.cpp
+++ b/Engine/Plugins/2D/Paper2D/Source/Paper2D/Private/PaperTileSet.cpp
@@ -10,7 +10,10 @@ UPaperTileSet::UPaperTileSet(const FObjectInitializer& ObjectInitializer)
 {
 	TileWidth = 32;
 	TileHeight = 32;
+
+#if WITH_EDITORONLY_DATA
 	BackgroundColor = FColor(0, 0, 127);
+#endif
 }
 
 int32 UPaperTileSet::GetTileCount() const
