[INTEGRATE]
Change 2414478 by [REDACTED_PERSON][REDACTED_HANDLE] on 2015/01/21 17:08:46
Paper2D: Tile Maps: Improve comments and exposed various properties to Blueprints (read only for now)

diff --git a/Engine/Plugins/2D/Paper2D/Source/Paper2D/Classes/PaperTileLayer.h b/Engine/Plugins/2D/Paper2D/Source/Paper2D/Classes/PaperTileLayer.h
--- a/Engine/Plugins/2D/Paper2D/Source/Paper2D/Classes/PaperTileLayer.h
+++ b/Engine/Plugins/2D/Paper2D/Source/Paper2D/Classes/PaperTileLayer.h
@@ -6,15 +6,18 @@
 
 #include "PaperTileLayer.generated.h"
 
+// This is the contents of a tile map cell
 USTRUCT(BlueprintType)
 struct FPaperTileInfo
 {
 	GENERATED_USTRUCT_BODY()
 
-	UPROPERTY()
+	// The tile set that this tile comes from
+	UPROPERTY(BlueprintReadWrite, Category=Sprite)
 	UPaperTileSet* TileSet;
 
-	UPROPERTY()
+	// This is the index of the current tile within the tile set
+	UPROPERTY(BlueprintReadWrite, Category=Sprite)
 	int32 PackedTileIndex;
 
 	FPaperTileInfo()
@@ -39,6 +42,7 @@ struct FPaperTileInfo
 	}
 };
 
+// This class represents a single layer in a tile map.  All layers in the map must have the size dimensions.
 UCLASS()
 class PAPER2D_API UPaperTileLayer : public UObject
 {
@@ -49,11 +53,11 @@ class PAPER2D_API UPaperTileLayer : public UObject
 	FText LayerName;
 
 	// Width of the layer (in tiles)
-	UPROPERTY()
+	UPROPERTY(BlueprintReadOnly, Category=Sprite)
 	int32 LayerWidth;
 
 	// Height of the layer (in tiles)
-	UPROPERTY()
+	UPROPERTY(BlueprintReadOnly, Category=Sprite)
 	int32 LayerHeight;
 
 #if WITH_EDITORONLY_DATA
@@ -67,7 +71,7 @@ class PAPER2D_API UPaperTileLayer : public UObject
 #endif
 
 	// Should this layer be hidden in the game?
-	UPROPERTY()
+	UPROPERTY(BlueprintReadOnly, Category=Sprite)
 	bool bHiddenInGame;
 
 	// Is this layer a collision layer?
diff --git a/Engine/Plugins/2D/Paper2D/Source/Paper2D/Classes/PaperTileMap.h b/Engine/Plugins/2D/Paper2D/Source/Paper2D/Classes/PaperTileMap.h
--- a/Engine/Plugins/2D/Paper2D/Source/Paper2D/Classes/PaperTileMap.h
+++ b/Engine/Plugins/2D/Paper2D/Source/Paper2D/Classes/PaperTileMap.h
@@ -13,71 +13,73 @@ namespace ETileMapProjectionMode
 {
 	enum Type
 	{
-		// 
+		// Square tile layout
 		Orthogonal,
 
-		//
+		// Isometric tile layout (shaped like a diamond)
 		IsometricDiamond,
 
-		//
+		// Isometric tile layout (roughly in a square with alternating rows staggered)
 		IsometricStaggered
 	};
 }
 
-
+// A tile map is a 2D grid with a defined width and height (in tiles).  There can be multiple layers, each of which can specify which tile should appear in each cell of the map for that layer.
 UCLASS()
 class PAPER2D_API UPaperTileMap : public UObject
 {
 	GENERATED_UCLASS_BODY()
 
 	// Width of map (in tiles)
-	UPROPERTY(Category=Setup, EditAnywhere, meta=(UIMin=1, ClampMin=1))
+	UPROPERTY(Category=Setup, EditAnywhere, BlueprintReadOnly, meta=(UIMin=1, ClampMin=1))
 	int32 MapWidth;
 
 	// Height of map (in tiles)
-	UPROPERTY(Category=Setup, EditAnywhere, meta=(UIMin=1, ClampMin=1))
+	UPROPERTY(Category=Setup, EditAnywhere, BlueprintReadOnly, meta=(UIMin=1, ClampMin=1))
 	int32 MapHeight;
 
 	// Width of one tile (in pixels)
-	UPROPERTY(Category=Setup, EditAnywhere, meta=(UIMin=1, ClampMin=1))
+	UPROPERTY(Category=Setup, EditAnywhere, BlueprintReadOnly, meta=(UIMin=1, ClampMin=1))
 	int32 TileWidth;
 
 	// Height of one tile (in pixels)
-	UPROPERTY(Category=Setup, EditAnywhere, meta=(UIMin=1, ClampMin=1))
+	UPROPERTY(Category=Setup, EditAnywhere, BlueprintReadOnly, meta=(UIMin=1, ClampMin=1))
 	int32 TileHeight;
 
-	// Pixels per Unreal Unit (pixels per cm)
+	// Pixels per Unreal Unit (pixels per cm) - Note: Currently unused!
 	UPROPERTY(Category=Setup, EditAnywhere)
 	float PixelsPerUnit;
 
-	// 
-	UPROPERTY(Category = Setup, EditAnywhere)
+	// The Z-separation incurred as you travel in X (not strictly applied, batched tiles will be put at the same Z level) 
+	UPROPERTY(Category=Setup, EditAnywhere, AdvancedDisplay)
 	float SeparationPerTileX;
 
-	UPROPERTY(Category = Setup, EditAnywhere)
+	// The Z-separation incurred as you travel in Y (not strictly applied, batched tiles will be put at the same Z level) 
+	UPROPERTY(Category=Setup, EditAnywhere, AdvancedDisplay)
 	float SeparationPerTileY;
 	
-	UPROPERTY(Category=Setup, EditAnywhere)
+	// The Z-separation between each layer of the tile map
+	UPROPERTY(Category=Setup, EditAnywhere, BlueprintReadOnly)
 	float SeparationPerLayer;
 
 	// Last tile set that was selected when editing the tile map
 	UPROPERTY()
 	TAssetPtr<class UPaperTileSet> SelectedTileSet;
 
-	// Test material
+	// The material to use on a tile map instance if not overridden
 	UPROPERTY(Category=Setup, EditAnywhere)
 	UMaterialInterface* Material;
 
 	// The list of layers
-	UPROPERTY(Instanced)
+	UPROPERTY(Instanced, Category=Sprite, BlueprintReadOnly)
 	TArray<class UPaperTileLayer*> TileLayers;
 
 	// Collision domain (no collision, 2D, or 3D)
-	UPROPERTY(Category=Collision, EditAnywhere)
+	UPROPERTY(Category=Collision, EditAnywhere, BlueprintReadOnly)
 	TEnumAsByte<ESpriteCollisionMode::Type> SpriteCollisionDomain;
 
 	// Tile map type
-	UPROPERTY(Category = Setup, EditAnywhere)
+	UPROPERTY(Category=Setup, EditAnywhere, BlueprintReadOnly)
 	TEnumAsByte<ETileMapProjectionMode::Type> ProjectionMode;
 
 	// Baked physics data.
