[INTEGRATE]
Change 2414477 by [REDACTED_PERSON][REDACTED_HANDLE] on 2015/01/21 17:08:43

Paper2D: Remarked tile map components as experimental in 4.7 branch

diff --git a/Engine/Plugins/2D/Paper2D/Source/Paper2D/Classes/PaperTileMapComponent.h b/Engine/Plugins/2D/Paper2D/Source/Paper2D/Classes/PaperTileMapComponent.h
--- a/Engine/Plugins/2D/Paper2D/Source/Paper2D/Classes/PaperTileMapComponent.h
+++ b/Engine/Plugins/2D/Paper2D/Source/Paper2D/Classes/PaperTileMapComponent.h
@@ -6,7 +6,7 @@
 
 #include "PaperTileMapComponent.generated.h"
 
-UCLASS(hideCategories=Object, ClassGroup=Paper2D, EarlyAccessPreview, meta=(BlueprintSpawnableComponent))
+UCLASS(hideCategories=Object, ClassGroup=Paper2D, Experimental, meta=(BlueprintSpawnableComponent))
 class PAPER2D_API UPaperTileMapComponent : public UPrimitiveComponent
 {
 	GENERATED_UCLASS_BODY()
@@ -71,6 +71,7 @@ public:
 	void CreateNewOwnedTileMap();
 
 	// Does this component own the tile map (is it instanced instead of being an asset reference)?
+	UFUNCTION(BlueprintCallable, Category="Sprite")
 	bool OwnsTileMap() const;
 
 	/** Change the PaperTileMap used by this instance. */
