[INTEGRATE]
Change 2414677 by [REDACTED_PERSON][REDACTED_HANDLE] on 2015/01/21 19:30:57

Paper2D: Switch UPaperTileMapComponent to derive from UMeshComponent, allowing material overrides/etc...
[UE-8011]

diff --git a/Engine/Plugins/2D/Paper2D/Source/Paper2D/Classes/PaperTileMap.h b/Engine/Plugins/2D/Paper2D/Source/Paper2D/Classes/PaperTileMap.h
--- a/Engine/Plugins/2D/Paper2D/Source/Paper2D/Classes/PaperTileMap.h
+++ b/Engine/Plugins/2D/Paper2D/Source/Paper2D/Classes/PaperTileMap.h
@@ -67,7 +67,7 @@ class PAPER2D_API UPaperTileMap : public UObject
 	TAssetPtr<class UPaperTileSet> SelectedTileSet;
 
 	// The material to use on a tile map instance if not overridden
-	UPROPERTY(Category=Setup, EditAnywhere)
+	UPROPERTY(Category=Setup, EditAnywhere, BlueprintReadOnly)
 	UMaterialInterface* Material;
 
 	// The list of layers
diff --git a/Engine/Plugins/2D/Paper2D/Source/Paper2D/Classes/PaperTileMapComponent.h b/Engine/Plugins/2D/Paper2D/Source/Paper2D/Classes/PaperTileMapComponent.h
--- a/Engine/Plugins/2D/Paper2D/Source/Paper2D/Classes/PaperTileMapComponent.h
+++ b/Engine/Plugins/2D/Paper2D/Source/Paper2D/Classes/PaperTileMapComponent.h
@@ -6,13 +6,23 @@
 
 #include "PaperTileMapComponent.generated.h"
 
+/**
+ * A component that handles rendering and collision for a single instance of a UPaperTileMap asset.
+ *
+ * This component is created when you drag a tile map asset from the content browser into a Blueprint, or
+ * contained inside of the actor created when you drag one into the level.
+ *
+ * NOTE: This is an experimental class and certain functionality is not properly supported yet (such as collision).  Use at your own risk!
+ *
+ * @see UPrimitiveComponent, UPaperTileMap
+ */
+
 UCLASS(hideCategories=Object, ClassGroup=Paper2D, Experimental, meta=(BlueprintSpawnableComponent))
-class PAPER2D_API UPaperTileMapComponent : public UPrimitiveComponent
+class PAPER2D_API UPaperTileMapComponent : public UMeshComponent
 {
 	GENERATED_UCLASS_BODY()
 
 private:
-	// Width of map (in tiles)
 	UPROPERTY()
 	int32 MapWidth_DEPRECATED;
 
@@ -25,15 +35,12 @@ private:
 	UPROPERTY()
 	int32 TileHeight_DEPRECATED;
 
-	// Default tile set to use for new layers
 	UPROPERTY()
 	UPaperTileSet* DefaultLayerTileSet_DEPRECATED;
 
-	// Test material
 	UPROPERTY()
 	UMaterialInterface* Material_DEPRECATED;
 
-	// The list of layers
 	UPROPERTY()
 	TArray<UPaperTileLayer*> TileLayers_DEPRECATED;
 
@@ -65,6 +72,9 @@ public:
 	virtual FPrimitiveSceneProxy* CreateSceneProxy() override;
 	virtual FBoxSphereBounds CalcBounds(const FTransform& LocalToWorld) const override;
 	virtual class UBodySetup* GetBodySetup() override;
+	virtual void GetUsedTextures(TArray<UTexture*>& OutTextures, EMaterialQualityLevel::Type QualityLevel) override;
+	virtual UMaterialInterface* GetMaterial(int32 MaterialIndex) const override;
+	virtual int32 GetNumMaterials() const override;
 	// End of UPrimitiveComponent interface
 
 	// Creates a new tile map internally, replacing the TileMap reference (or dropping the previous owned one)
diff --git a/Engine/Plugins/2D/Paper2D/Source/Paper2D/Private/PaperTileMapComponent.cpp b/Engine/Plugins/2D/Paper2D/Source/Paper2D/Private/PaperTileMapComponent.cpp
--- a/Engine/Plugins/2D/Paper2D/Source/Paper2D/Private/PaperTileMapComponent.cpp
+++ b/Engine/Plugins/2D/Paper2D/Source/Paper2D/Private/PaperTileMapComponent.cpp
@@ -22,6 +22,9 @@ UPaperTileMapComponent::UPaperTileMapComponent(const FObjectInitializer& ObjectI
 
 	static ConstructorHelpers::FObjectFinder<UMaterial> DefaultMaterial(TEXT("/Paper2D/DefaultSpriteMaterial"));
 	Material_DEPRECATED = DefaultMaterial.Object;
+
+	CastShadow = false;
+	bUseAsOccluder = false;
 }
 
 FPrimitiveSceneProxy* UPaperTileMapComponent::CreateSceneProxy()
@@ -122,6 +125,50 @@ UBodySetup* UPaperTileMapComponent::GetBodySetup()
 	return (TileMap != nullptr) ? TileMap->BodySetup : nullptr;
 }
 
+void UPaperTileMapComponent::GetUsedTextures(TArray<UTexture*>& OutTextures, EMaterialQualityLevel::Type QualityLevel)
+{
+	// Get the texture referenced by the tile maps
+	if (TileMap != nullptr)
+	{
+		for (UPaperTileLayer* Layer : TileMap->TileLayers)
+		{
+			for (int32 Y = 0; Y < TileMap->MapHeight; ++Y)
+			{
+				for (int32 X = 0; X < TileMap->MapWidth; ++X)
+				{
+					FPaperTileInfo TileInfo = Layer->GetCell(X, Y);
+					if (TileInfo.IsValid() && (TileInfo.TileSet != nullptr) && (TileInfo.TileSet->TileSheet != nullptr))
+					{
+						OutTextures.AddUnique(TileInfo.TileSet->TileSheet);
+					}
+				}
+			}
+		}
+	}
+		
+	// Get any textures referenced by our materials
+	Super::GetUsedTextures(OutTextures, QualityLevel);
+}
+
+UMaterialInterface* UPaperTileMapComponent::GetMaterial(int32 MaterialIndex) const
+{
+	if (OverrideMaterials.IsValidIndex(MaterialIndex) && (OverrideMaterials[MaterialIndex] != nullptr))
+	{
+		return OverrideMaterials[MaterialIndex];
+	}
+	else if (TileMap != nullptr)
+	{
+		return TileMap->Material;
+	}
+
+	return nullptr;
+}
+
+int32 UPaperTileMapComponent::GetNumMaterials() const
+{
+	return FMath::Max<int32>(OverrideMaterials.Num(), 1);
+}
+
 const UObject* UPaperTileMapComponent::AdditionalStatObject() const
 {
 	if (TileMap != nullptr)
diff --git a/Engine/Plugins/2D/Paper2D/Source/Paper2D/Private/PaperTileMapRenderSceneProxy.cpp b/Engine/Plugins/2D/Paper2D/Source/Paper2D/Private/PaperTileMapRenderSceneProxy.cpp
--- a/Engine/Plugins/2D/Paper2D/Source/Paper2D/Private/PaperTileMapRenderSceneProxy.cpp
+++ b/Engine/Plugins/2D/Paper2D/Source/Paper2D/Private/PaperTileMapRenderSceneProxy.cpp
@@ -15,12 +15,8 @@ FPaperTileMapRenderSceneProxy::FPaperTileMapRenderSceneProxy(const UPaperTileMap
 	if (const UPaperTileMapComponent* InTileComponent = Cast<const UPaperTileMapComponent>(InComponent))
 	{
 		TileMap = InTileComponent->TileMap;
-		Material = (TileMap != nullptr) ? TileMap->Material : nullptr;
-
-		if (Material)
-		{
-			MaterialRelevance = Material->GetRelevance(GetScene().GetFeatureLevel());
-		}
+		Material = InTileComponent->GetMaterial(0);
+		MaterialRelevance = InTileComponent->GetMaterialRelevance(GetScene().GetFeatureLevel());
 	}
 }
 
