[INTEGRATE]
Change 2414703 by [REDACTED_PERSON][REDACTED_HANDLE] on 2015/01/21 19:46:23

Paper2D: Make tile map preview cursor respect component scale of the component being painted on [UE-8022]

diff --git a/Engine/Plugins/2D/Paper2D/Source/Paper2DEditor/Private/TileMapEditing/EdModeTileMap.cpp b/Engine/Plugins/2D/Paper2D/Source/Paper2DEditor/Private/TileMapEditing/EdModeTileMap.cpp
--- a/Engine/Plugins/2D/Paper2D/Source/Paper2DEditor/Private/TileMapEditing/EdModeTileMap.cpp
+++ b/Engine/Plugins/2D/Paper2D/Source/Paper2DEditor/Private/TileMapEditing/EdModeTileMap.cpp
@@ -892,6 +892,7 @@ void FEdModeTileMap::UpdatePreviewCursor(const FViewportCursorLocation& Ray)
 			const FVector ComponentPreviewLocation = ComponentToWorld.TransformPosition(TileMap->GetTileCenterInLocalSpace(LocalTileX0, LocalTileY0, LayerIndex));
 			CursorPreviewComponent->SetWorldLocation(ComponentPreviewLocation);
 			CursorPreviewComponent->SetWorldRotation(FRotator(ComponentToWorld.GetRotation()));
+			CursorPreviewComponent->SetWorldScale3D(ComponentToWorld.GetScale3D());
 			SynchronizePreviewWithTileMap(TileMap);
 		}
 	}
