[INTEGRATE]
Change 2414734 by [REDACTED_PERSON][REDACTED_HANDLE] on 2015/01/21 20:08:38

Paper2D: Revert background color in sprite and flipbook editors to dark grey [UE-8025]

diff --git a/Engine/Plugins/2D/Paper2D/Source/Paper2DEditor/Private/FlipbookEditor/FlipbookEditorViewportClient.cpp b/Engine/Plugins/2D/Paper2D/Source/Paper2DEditor/Private/FlipbookEditor/FlipbookEditorViewportClient.cpp
--- a/Engine/Plugins/2D/Paper2D/Source/Paper2DEditor/Private/FlipbookEditor/FlipbookEditorViewportClient.cpp
+++ b/Engine/Plugins/2D/Paper2D/Source/Paper2DEditor/Private/FlipbookEditor/FlipbookEditorViewportClient.cpp
@@ -105,7 +105,6 @@ void FFlipbookEditorViewportClient::Tick(float DeltaSeconds)
 	}
 }
 
-
 bool FFlipbookEditorViewportClient::InputKey(FViewport* Viewport, int32 ControllerId, FKey Key, EInputEvent Event, float AmountDepressed, bool bGamepad)
 {
 	bool bHandled = false;
@@ -120,6 +119,11 @@ bool FFlipbookEditorViewportClient::InputKey(FViewport* Viewport, int32 Controll
 	return (bHandled) ? true : FEditorViewportClient::InputKey(Viewport,  ControllerId, Key, Event, AmountDepressed, bGamepad);
 }
 
+FLinearColor FFlipbookEditorViewportClient::GetBackgroundColor() const
+{
+	return FEditorViewportClient::GetBackgroundColor();
+}
+
 //////////////////////////////////////////////////////////////////////////
 
 #undef LOCTEXT_NAMESPACE
\ No newline at end of file
diff --git a/Engine/Plugins/2D/Paper2D/Source/Paper2DEditor/Private/FlipbookEditor/FlipbookEditorViewportClient.h b/Engine/Plugins/2D/Paper2D/Source/Paper2DEditor/Private/FlipbookEditor/FlipbookEditorViewportClient.h
--- a/Engine/Plugins/2D/Paper2D/Source/Paper2DEditor/Private/FlipbookEditor/FlipbookEditorViewportClient.h
+++ b/Engine/Plugins/2D/Paper2D/Source/Paper2DEditor/Private/FlipbookEditor/FlipbookEditorViewportClient.h
@@ -28,6 +28,7 @@ public:
 
 	// FEditorViewportClient interface
 	virtual bool InputKey(FViewport* Viewport, int32 ControllerId, FKey Key, EInputEvent Event, float AmountDepressed, bool bGamepad) override;
+	virtual FLinearColor GetBackgroundColor() const override;
 	// End of FEditorViewportClient interface
 
 	void ToggleShowPivot() { bShowPivot = !bShowPivot; Invalidate(); }
diff --git a/Engine/Plugins/2D/Paper2D/Source/Paper2DEditor/Private/SpriteEditor/SpriteEditorViewportClient.cpp b/Engine/Plugins/2D/Paper2D/Source/Paper2DEditor/Private/SpriteEditor/SpriteEditorViewportClient.cpp
--- a/Engine/Plugins/2D/Paper2D/Source/Paper2DEditor/Private/SpriteEditor/SpriteEditorViewportClient.cpp
+++ b/Engine/Plugins/2D/Paper2D/Source/Paper2DEditor/Private/SpriteEditor/SpriteEditorViewportClient.cpp
@@ -1392,6 +1392,11 @@ ECoordSystem FSpriteEditorViewportClient::GetWidgetCoordSystemSpace() const
 	return COORD_World;
 }
 
+FLinearColor FSpriteEditorViewportClient::GetBackgroundColor() const
+{
+	return FEditorViewportClient::GetBackgroundColor();
+}
+
 void FSpriteEditorViewportClient::BeginTransaction(const FText& SessionName)
 {
 	if (ScopedTransaction == nullptr)
diff --git a/Engine/Plugins/2D/Paper2D/Source/Paper2DEditor/Private/SpriteEditor/SpriteEditorViewportClient.h b/Engine/Plugins/2D/Paper2D/Source/Paper2DEditor/Private/SpriteEditor/SpriteEditorViewportClient.h
--- a/Engine/Plugins/2D/Paper2D/Source/Paper2DEditor/Private/SpriteEditor/SpriteEditorViewportClient.h
+++ b/Engine/Plugins/2D/Paper2D/Source/Paper2DEditor/Private/SpriteEditor/SpriteEditorViewportClient.h
@@ -64,6 +64,7 @@ public:
 	virtual FVector GetWidgetLocation() const;
 	virtual FMatrix GetWidgetCoordSystem() const;
 	virtual ECoordSystem GetWidgetCoordSystemSpace() const;
+	virtual FLinearColor GetBackgroundColor() const override;
 	// End of FEditorViewportClient interface
 
 	// Process marquee selection, return true of a selection has been performed
