[INTEGRATE]
Change 2414853 by [REDACTED_PERSON][REDACTED_HANDLE] on 2015/01/21 21:42:37

Paper2D: Remove 'early access preview' watermark in sprite editor

diff --git a/Engine/Plugins/2D/Paper2D/Source/Paper2DEditor/Private/SpriteEditor/SpriteEditor.cpp b/Engine/Plugins/2D/Paper2D/Source/Paper2DEditor/Private/SpriteEditor/SpriteEditor.cpp
--- a/Engine/Plugins/2D/Paper2D/Source/Paper2DEditor/Private/SpriteEditor/SpriteEditor.cpp
+++ b/Engine/Plugins/2D/Paper2D/Source/Paper2DEditor/Private/SpriteEditor/SpriteEditor.cpp
@@ -281,25 +281,7 @@ TSharedRef<SDockTab> FSpriteEditor::SpawnTab_Viewport(const FSpawnTabArgs& Args)
 	return SNew(SDockTab)
 		.Label(LOCTEXT("ViewportTab_Title", "Viewport"))
 		[
-			SNew(SOverlay)
-
-			// The sprite editor viewport
-			+SOverlay::Slot()
-			[
-				ViewportPtr.ToSharedRef()
-			]
-
-			// Bottom-right corner text indicating the preview nature of the sprite editor
-			+SOverlay::Slot()
-			.Padding(10)
-			.VAlign(VAlign_Bottom)
-			.HAlign(HAlign_Right)
-			[
-				SNew(STextBlock)
-				.Visibility( EVisibility::HitTestInvisible )
-				.TextStyle( FEditorStyle::Get(), "Graph.CornerText" )
-				.Text(LOCTEXT("SpriteEditorViewportExperimentalWarning", "Early access preview"))
-			]
+			ViewportPtr.ToSharedRef()
 		];
 }
 
