[INTEGRATE]
Change 2414855 by [REDACTED_PERSON][REDACTED_HANDLE] on 2015/01/21 21:42:40

Paper2D: Change default 'show grid' setting for sprite and flipbook editors to false

diff --git a/Engine/Plugins/2D/Paper2D/Source/Paper2DEditor/Private/FlipbookEditor/FlipbookEditorViewportClient.cpp b/Engine/Plugins/2D/Paper2D/Source/Paper2DEditor/Private/FlipbookEditor/FlipbookEditorViewportClient.cpp
--- a/Engine/Plugins/2D/Paper2D/Source/Paper2DEditor/Private/FlipbookEditor/FlipbookEditorViewportClient.cpp
+++ b/Engine/Plugins/2D/Paper2D/Source/Paper2DEditor/Private/FlipbookEditor/FlipbookEditorViewportClient.cpp
@@ -31,6 +31,7 @@ FFlipbookEditorViewportClient::FFlipbookEditorViewportClient(const TAttribute<UP
 
 	bShowPivot = false;
 	bDeferZoomToSprite = true;
+	DrawHelper.bDrawGrid = false;
 
 	EngineShowFlags.DisableAdvancedFeatures();
 	EngineShowFlags.CompositeEditorPrimitives = true;
diff --git a/Engine/Plugins/2D/Paper2D/Source/Paper2DEditor/Private/SpriteEditor/SpriteEditorViewportClient.cpp b/Engine/Plugins/2D/Paper2D/Source/Paper2DEditor/Private/SpriteEditor/SpriteEditorViewportClient.cpp
--- a/Engine/Plugins/2D/Paper2D/Source/Paper2DEditor/Private/SpriteEditor/SpriteEditorViewportClient.cpp
+++ b/Engine/Plugins/2D/Paper2D/Source/Paper2DEditor/Private/SpriteEditor/SpriteEditorViewportClient.cpp
@@ -90,6 +90,8 @@ FSpriteEditorViewportClient::FSpriteEditorViewportClient(TWeakPtr<FSpriteEditor>
 
 	bIsMarqueeTracking = false;
 
+	DrawHelper.bDrawGrid = false;
+
 	EngineShowFlags.DisableAdvancedFeatures();
 	EngineShowFlags.CompositeEditorPrimitives = true;
 
