[INTEGRATE]
Change 2414856 by [REDACTED_PERSON][REDACTED_HANDLE] on 2015/01/21 21:42:44

Paper2D: Add a plugin setting (defaults to on) that analyzes the texture when you create a new sprite to pick the best material (translucent versus masked, instead of always defaulting to masked)
Upgrade note: If you used the 'default materials' plugin settings you will need to set them again, as the names have changed

diff --git a/Engine/Plugins/2D/Paper2D/Source/Paper2D/Classes/PaperRuntimeSettings.h b/Engine/Plugins/2D/Paper2D/Source/Paper2D/Classes/PaperRuntimeSettings.h
--- a/Engine/Plugins/2D/Paper2D/Source/Paper2D/Classes/PaperRuntimeSettings.h
+++ b/Engine/Plugins/2D/Paper2D/Source/Paper2D/Classes/PaperRuntimeSettings.h
@@ -17,13 +17,21 @@ class PAPER2D_API UPaperRuntimeSettings : public UObject
 	float DefaultPixelsPerUnrealUnit;
 
 #if WITH_EDITORONLY_DATA
-	// The default translucent or masked material for newly created sprites
-	UPROPERTY(Config, EditAnywhere, Category=Settings, meta=(AllowedClasses="MaterialInterface", DisplayName="Default Translucent Sprite Material"))
-	FStringAssetReference DefaultTranslucentSpriteMaterialName;
+	// Should the source texture be scanned when creating new sprites to determine the appropriate material? (if false, the Default Masked Material is always used)
+	UPROPERTY(config, EditAnywhere, Category=Settings)
+	bool bPickBestMaterialWhenCreatingSprite;
+
+	// The default masked material for newly created sprites (masked means binary opacity: things are either opaque or see-thru, with nothing in between)
+	UPROPERTY(config, EditAnywhere, Category=Settings, meta=(AllowedClasses="MaterialInterface", DisplayName="Default Masked Material"))
+	FStringAssetReference DefaultMaskedMaterialName;
+
+	// The default translucent material for newly created sprites (translucent means smooth opacity which can vary continuously from 0..1, but translucent rendering is more expensive that opaque or masked rendering and has different sorting rules)
+	UPROPERTY(config, EditAnywhere, Category=Settings, meta=(AllowedClasses="MaterialInterface", DisplayName="Default Translucent Material", EditCondition="bPickBestMaterialWhenCreatingSprite"))
+	FStringAssetReference DefaultTranslucentMaterialName;
 
 	// The default opaque material for newly created sprites
-	UPROPERTY(Config, EditAnywhere, Category=Settings, meta=(AllowedClasses="MaterialInterface", DisplayName="Default Opaque Sprite Material"))
-	FStringAssetReference DefaultOpaqueSpriteMaterialName;
+	UPROPERTY(config, EditAnywhere, Category=Settings, meta=(AllowedClasses="MaterialInterface", DisplayName="Default Opaque Sprite Material"))
+	FStringAssetReference DefaultOpaqueMaterialName;
 #endif
 
 	// Enables experimental *incomplete and unsupported* texture atlas groups that sprites can be assigned to
diff --git a/Engine/Plugins/2D/Paper2D/Source/Paper2D/Private/AlphaBitmap.h b/Engine/Plugins/2D/Paper2D/Source/Paper2D/Private/AlphaBitmap.h
--- a/Engine/Plugins/2D/Paper2D/Source/Paper2D/Private/AlphaBitmap.h
+++ b/Engine/Plugins/2D/Paper2D/Source/Paper2D/Private/AlphaBitmap.h
@@ -322,6 +322,35 @@ struct FAlphaBitmap
 		}
 	}
 
+	// Checks the image to determine if it is suitable for opaque, masked or translucent rendering
+	void AnalyzeImage(int32 StartX, int32 StartY, int32 Width, int32 Height, bool& bOutHasZeros, bool& bOutHasIntermediateValues)
+	{
+		bOutHasZeros = false;
+		bOutHasIntermediateValues = false;
+
+		const int32 X0 = StartX;
+		const int32 Y0 = StartY;
+		const int32 X1 = StartX + Width;
+		const int32 Y1 = StartY + Height;
+		for (int Y = Y0; Y < Y1; ++Y)
+		{
+			for (int X = X0; X < X1; ++X)
+			{
+				const uint8 Value = GetPixel(X, Y);
+
+				if ((Value > 0) && (Value < 255))
+				{
+					bOutHasIntermediateValues = true;
+				}
+
+				if (Value == 0)
+				{
+					bOutHasZeros = true;
+				}
+			}
+		}
+	}
+
 public:
 	TArray<uint8> RawData;
 	int32 Width;
diff --git a/Engine/Plugins/2D/Paper2D/Source/Paper2D/Private/PaperRuntimeSettings.cpp b/Engine/Plugins/2D/Paper2D/Source/Paper2D/Private/PaperRuntimeSettings.cpp
--- a/Engine/Plugins/2D/Paper2D/Source/Paper2D/Private/PaperRuntimeSettings.cpp
+++ b/Engine/Plugins/2D/Paper2D/Source/Paper2D/Private/PaperRuntimeSettings.cpp
@@ -15,8 +15,10 @@ UPaperRuntimeSettings::UPaperRuntimeSettings(const FObjectInitializer& ObjectIni
 #if WITH_EDITORONLY_DATA
 	if (!IsRunningCommandlet())
 	{
-		DefaultTranslucentSpriteMaterialName = FStringAssetReference("/Paper2D/MaskedUnlitSpriteMaterial.MaskedUnlitSpriteMaterial");
-		DefaultOpaqueSpriteMaterialName = FStringAssetReference("/Paper2D/OpaqueUnlitSpriteMaterial.OpaqueUnlitSpriteMaterial");
+		bPickBestMaterialWhenCreatingSprite = true;
+		DefaultMaskedMaterialName = FStringAssetReference("/Paper2D/MaskedUnlitSpriteMaterial.MaskedUnlitSpriteMaterial");
+		DefaultTranslucentMaterialName = FStringAssetReference("/Paper2D/TranslucentUnlitSpriteMaterial.TranslucentUnlitSpriteMaterial");
+		DefaultOpaqueMaterialName = FStringAssetReference("/Paper2D/OpaqueUnlitSpriteMaterial.OpaqueUnlitSpriteMaterial");
 	}
 #endif
 }
diff --git a/Engine/Plugins/2D/Paper2D/Source/Paper2D/Private/PaperSprite.cpp b/Engine/Plugins/2D/Paper2D/Source/Paper2D/Private/PaperSprite.cpp
--- a/Engine/Plugins/2D/Paper2D/Source/Paper2D/Private/PaperSprite.cpp
+++ b/Engine/Plugins/2D/Paper2D/Source/Paper2D/Private/PaperSprite.cpp
@@ -1452,12 +1452,38 @@ void UPaperSprite::InitializeSprite(const FSpriteAssetInitParameters& InitParams
 		const UPaperRuntimeSettings* DefaultSettings = GetDefault<UPaperRuntimeSettings>();
 		PixelsPerUnrealUnit = DefaultSettings->DefaultPixelsPerUnrealUnit;
 
-		if (UMaterialInterface* TranslucentMaterial = LoadObject<UMaterialInterface>(nullptr, *DefaultSettings->DefaultTranslucentSpriteMaterialName.ToString(), nullptr, LOAD_None, nullptr))
+		bool bUseMaskedTexture = true;
+
+		// Analyze the texture if desired (to see if it's got greyscale alpha or just binary alpha, picking either a ranslucent or masked material)
+		if (DefaultSettings->bPickBestMaterialWhenCreatingSprite)
+		{
+			if (InitParams.Texture != nullptr)
+			{
+				FAlphaBitmap AlphaBitmap(InitParams.Texture);
+				bool bHasIntermediateValues;
+				bool bHasZeros;
+				AlphaBitmap.AnalyzeImage((int32)InitParams.Offset.X, (int32)InitParams.Offset.Y, (int32)InitParams.Dimension.X, (int32)InitParams.Dimension.Y, /*out*/ bHasZeros, /*out*/ bHasIntermediateValues);
+
+				bUseMaskedTexture = !bHasIntermediateValues;
+			}
+		}
+
+		if (bUseMaskedTexture)
 		{
-			DefaultMaterial = TranslucentMaterial;
+			if (UMaterialInterface* MaskedMaterial = LoadObject<UMaterialInterface>(nullptr, *DefaultSettings->DefaultMaskedMaterialName.ToString(), nullptr, LOAD_None, nullptr))
+			{
+				DefaultMaterial = MaskedMaterial;
+			}
+		}
+		else
+		{
+			if (UMaterialInterface* TranslucentMaterial = LoadObject<UMaterialInterface>(nullptr, *DefaultSettings->DefaultTranslucentMaterialName.ToString(), nullptr, LOAD_None, nullptr))
+			{
+				DefaultMaterial = TranslucentMaterial;
+			}
 		}
 
-		if (UMaterialInterface* OpaqueMaterial = LoadObject<UMaterialInterface>(nullptr, *DefaultSettings->DefaultOpaqueSpriteMaterialName.ToString(), nullptr, LOAD_None, nullptr))
+		if (UMaterialInterface* OpaqueMaterial = LoadObject<UMaterialInterface>(nullptr, *DefaultSettings->DefaultOpaqueMaterialName.ToString(), nullptr, LOAD_None, nullptr))
 		{
 			AlternateMaterial = OpaqueMaterial;
 		}
