[INTEGRATE]
Change 2414921 by [REDACTED_PERSON][REDACTED_HANDLE] on 2015/01/21 22:42:13

Paper2D: Add partial support for isometric staggered and hexagonal staggered tile map projection modes (neither is fully supported yet, but isometric staggered is rendering correctly)

diff --git a/Engine/Plugins/2D/Paper2D/Source/Paper2D/Classes/PaperTileMap.h b/Engine/Plugins/2D/Paper2D/Source/Paper2D/Classes/PaperTileMap.h
--- a/Engine/Plugins/2D/Paper2D/Source/Paper2D/Classes/PaperTileMap.h
+++ b/Engine/Plugins/2D/Paper2D/Source/Paper2D/Classes/PaperTileMap.h
@@ -19,8 +19,11 @@ namespace ETileMapProjectionMode
 		// Isometric tile layout (shaped like a diamond)
 		IsometricDiamond,
 
-		// Isometric tile layout (roughly in a square with alternating rows staggered)
-		IsometricStaggered
+		// Isometric tile layout (roughly in a square with alternating rows staggered).  Warning: Not fully supported yet.
+		IsometricStaggered,
+
+		// Hexagonal tile layout (roughly in a square with alternating rows staggered).  Warning: Not fully supported yet.
+		HexagonalStaggered
 	};
 }
 
diff --git a/Engine/Plugins/2D/Paper2D/Source/Paper2D/Private/PaperTileMap.cpp b/Engine/Plugins/2D/Paper2D/Source/Paper2D/Private/PaperTileMap.cpp
--- a/Engine/Plugins/2D/Paper2D/Source/Paper2D/Private/PaperTileMap.cpp
+++ b/Engine/Plugins/2D/Paper2D/Source/Paper2D/Private/PaperTileMap.cpp
@@ -193,11 +193,12 @@ void UPaperTileMap::GetTileToLocalParameters(FVector& OutCornerPosition, FVector
 		OutStepX = (TileWidth * PaperAxisX * 0.5f) - (TileHeight * PaperAxisY * 0.5f);
 		OutStepY = (TileWidth * PaperAxisX * -0.5f) - (TileHeight * PaperAxisY * 0.5f);
 		break;
+	case ETileMapProjectionMode::HexagonalStaggered:
 	case ETileMapProjectionMode::IsometricStaggered:
-		OutCornerPosition = 0.5f * TileHeight * PaperAxisY;
+		OutCornerPosition = -(TileWidth * PaperAxisX * 0.5f) + (TileHeight * PaperAxisY * 1.0f);
 		OutOffsetYFactor = 0.5f * TileWidth * PaperAxisX;
 		OutStepX = PaperAxisX * TileWidth;
-		OutStepY = -PaperAxisY * TileHeight;
+		OutStepY = 0.5f * -PaperAxisY * TileHeight;
 		break;
 	}
 }
@@ -219,8 +220,9 @@ void UPaperTileMap::GetLocalToTileParameters(FVector& OutCornerPosition, FVector
 		OutStepX = (PaperAxisX / TileWidth) - (PaperAxisY / TileHeight);
 		OutStepY = (-PaperAxisX / TileWidth) - (PaperAxisY / TileHeight);
 		break;
+	case ETileMapProjectionMode::HexagonalStaggered:
 	case ETileMapProjectionMode::IsometricStaggered:
-		OutCornerPosition = 0.5f * TileHeight * PaperAxisY;
+		OutCornerPosition = -(TileWidth * PaperAxisX * 0.5f) + (TileHeight * PaperAxisY * 1.0f);
 		OutOffsetYFactor = 0.5f * TileWidth * PaperAxisX;
 		OutStepX = PaperAxisX / TileWidth;
 		OutStepY = -PaperAxisY / TileHeight;
@@ -267,6 +269,7 @@ FVector UPaperTileMap::GetTilePositionInLocalSpace(float TileX, float TileY, int
 	case ETileMapProjectionMode::IsometricDiamond:
 		TotalOffset = CornerPosition;
 		break;
+	case ETileMapProjectionMode::HexagonalStaggered:
 	case ETileMapProjectionMode::IsometricStaggered:
 		TotalOffset = CornerPosition + ((int32)TileY & 1) * OffsetYFactor;
 		break;
@@ -312,9 +315,7 @@ FBoxSphereBounds UPaperTileMap::GetRenderBounds() const
 			 const FBox Box(BottomLeft, BottomLeft + Dimensions);
 			 return FBoxSphereBounds(Box);
 		}
-// 		case ETileMapProjectionMode::IsometricStaggered:
-// 		{
-// 		}
+//@TODO: verify bounds for IsometricStaggered and HexagonalStaggered
 	}
 }
 
diff --git a/Engine/Plugins/2D/Paper2D/Source/Paper2D/Private/PaperTileMapComponent.cpp b/Engine/Plugins/2D/Paper2D/Source/Paper2D/Private/PaperTileMapComponent.cpp
--- a/Engine/Plugins/2D/Paper2D/Source/Paper2D/Private/PaperTileMapComponent.cpp
+++ b/Engine/Plugins/2D/Paper2D/Source/Paper2D/Private/PaperTileMapComponent.cpp
@@ -239,6 +239,7 @@ void UPaperTileMapComponent::RebuildRenderData(FPaperTileMapRenderSceneProxy* Pr
 				EffectiveTopLeftCorner = CornerOffset - StepPerTileX;
 				break;
 			case ETileMapProjectionMode::IsometricStaggered:
+			case ETileMapProjectionMode::HexagonalStaggered:
 				EffectiveTopLeftCorner = CornerOffset + (Y & 1) * OffsetYFactor;
 				break;
 			}
diff --git a/Engine/Plugins/2D/Paper2D/Source/PaperTiledImporter/Classes/PaperTiledImporterFactory.h b/Engine/Plugins/2D/Paper2D/Source/PaperTiledImporter/Classes/PaperTiledImporterFactory.h
--- a/Engine/Plugins/2D/Paper2D/Source/PaperTiledImporter/Classes/PaperTiledImporterFactory.h
+++ b/Engine/Plugins/2D/Paper2D/Source/PaperTiledImporter/Classes/PaperTiledImporterFactory.h
@@ -13,7 +13,8 @@ namespace ETiledOrientation
 		Unknown,
 		Orthogonal,
 		Isometric,
-		Staggered
+		Staggered,
+		Hexagonal
 	};
 }
 
@@ -117,6 +118,8 @@ struct FTileMapFromTiled
 
 	ETiledOrientation::Type Orientation;
 
+	int32 HexSideLength;
+
 	TArray<FTileSetFromTiled> TileSets;
 	TArray<UPaperTileSet*> CreatedTileSetAssets;
 	TArray<FTileLayerFromTiled> Layers;
diff --git a/Engine/Plugins/2D/Paper2D/Source/PaperTiledImporter/Private/PaperTiledImporterFactory.cpp b/Engine/Plugins/2D/Paper2D/Source/PaperTiledImporter/Private/PaperTiledImporterFactory.cpp
--- a/Engine/Plugins/2D/Paper2D/Source/PaperTiledImporter/Private/PaperTiledImporterFactory.cpp
+++ b/Engine/Plugins/2D/Paper2D/Source/PaperTiledImporter/Private/PaperTiledImporterFactory.cpp
@@ -128,10 +128,7 @@ UObject* UPaperTiledImporterFactory::FactoryCreateText(UClass* InClass, UObject*
  	const FString NameForErrors(InName.ToString());
  	const FString FileContent(BufferEnd - Buffer, Buffer);
  	TSharedPtr<FJsonObject> DescriptorObject = ParseJSON(FileContent, NameForErrors);
-// 
-// 	TSet<FName> FrameNames;
-// 	TArray<FSpriteFrame> ParsedFrames;
-// 
+
  	UPaperTileMap* Result = nullptr;
  
 	FTileMapFromTiled GlobalInfo;
@@ -428,7 +425,7 @@ void UPaperTiledImporterFactory::ParseGlobalInfoFromJSON(TSharedPtr<FJsonObject>
 {
 	bool bSuccessfullyParsed = true;
 
-	// Parse all of the integer fields
+	// Parse all of the required integer fields
 	FRequiredIntField IntFields[] = {
 		FRequiredIntField( OutParsedInfo.FileVersion, TEXT("version") ),
 		FRequiredIntField( OutParsedInfo.Width, TEXT("width") ),
@@ -438,6 +435,15 @@ void UPaperTiledImporterFactory::ParseGlobalInfoFromJSON(TSharedPtr<FJsonObject>
 	};
 	bSuccessfullyParsed = bSuccessfullyParsed && ParseIntegerFields(IntFields, ARRAY_COUNT(IntFields), Tree, NameForErrors, bSilent);
 
+	// Parse hexsidelength if present
+	FRequiredIntField OptionalIntFields[] = {
+		FRequiredIntField(OutParsedInfo.HexSideLength, TEXT("hexsidelength"), 0)
+	};
+	ParseIntegerFields(OptionalIntFields, ARRAY_COUNT(OptionalIntFields), Tree, NameForErrors, /*bSilent=*/ true);
+
+	//@TODO:
+	// Parse staggeraxis and staggerindex (currently not supported by the renderer anyways)
+
 	// Parse the orientation
 	const FString OrientationModeStr = FPaperJSONHelpers::ReadString(Tree, TEXT("orientation"), TEXT(""));
 	if (OrientationModeStr == TEXT("orthogonal"))
@@ -452,9 +458,13 @@ void UPaperTiledImporterFactory::ParseGlobalInfoFromJSON(TSharedPtr<FJsonObject>
 	{
 		OutParsedInfo.Orientation = ETiledOrientation::Staggered;
 	}
+	else if (OrientationModeStr == TEXT("hexagonal"))
+	{
+		OutParsedInfo.Orientation = ETiledOrientation::Hexagonal;
+	}
 	else
 	{
-		TILED_IMPORT_ERROR(TEXT("Failed to parse '%s'.  Invalid value for '%s' (%s but expected 'orthogonal' or 'isometric')"), *NameForErrors, TEXT("orientation"), *OrientationModeStr);
+		TILED_IMPORT_ERROR(TEXT("Failed to parse '%s'.  Invalid value for '%s' (%s but expected 'orthogonal', 'isometric', 'staggered', or 'hexagonal')"), *NameForErrors, TEXT("orientation"), *OrientationModeStr);
 		bSuccessfullyParsed = false;
 		OutParsedInfo.Orientation = ETiledOrientation::Unknown;
 	}
@@ -470,6 +480,7 @@ FTileMapFromTiled::FTileMapFromTiled()
 	, TileWidth(0)
 	, TileHeight(0)
 	, Orientation(ETiledOrientation::Unknown)
+	, HexSideLength(0)
 {
 }
 
@@ -513,6 +524,8 @@ ETileMapProjectionMode::Type FTileMapFromTiled::GetOrientationType() const
 		return ETileMapProjectionMode::IsometricDiamond;
 	case ETiledOrientation::Staggered:
 		return ETileMapProjectionMode::IsometricStaggered;
+	case ETiledOrientation::Hexagonal:
+		return ETileMapProjectionMode::HexagonalStaggered;
 	case ETiledOrientation::Orthogonal:
 	default:
 		return ETileMapProjectionMode::Orthogonal;
