[INTEGRATE]
Change 2414939 by [REDACTED_PERSON][REDACTED_HANDLE] on 2015/01/21 23:37:18

Paper2D: Expose GetTile, SetTile, ResizeMap, and AddNewLayer on UPaperTileMapComponent to Blueprints

diff --git a/Engine/Plugins/2D/Paper2D/Source/Paper2D/Classes/PaperTileLayer.h b/Engine/Plugins/2D/Paper2D/Source/Paper2D/Classes/PaperTileLayer.h
--- a/Engine/Plugins/2D/Paper2D/Source/Paper2D/Classes/PaperTileLayer.h
+++ b/Engine/Plugins/2D/Paper2D/Source/Paper2D/Classes/PaperTileLayer.h
@@ -13,11 +13,11 @@ struct FPaperTileInfo
 	GENERATED_USTRUCT_BODY()
 
 	// The tile set that this tile comes from
-	UPROPERTY(BlueprintReadWrite, Category=Sprite)
+	UPROPERTY(EditAnywhere, BlueprintReadWrite, Category = Sprite)
 	UPaperTileSet* TileSet;
 
 	// This is the index of the current tile within the tile set
-	UPROPERTY(BlueprintReadWrite, Category=Sprite)
+	UPROPERTY(EditAnywhere, BlueprintReadWrite, Category=Sprite)
 	int32 PackedTileIndex;
 
 	FPaperTileInfo()
diff --git a/Engine/Plugins/2D/Paper2D/Source/Paper2D/Classes/PaperTileMap.h b/Engine/Plugins/2D/Paper2D/Source/Paper2D/Classes/PaperTileMap.h
--- a/Engine/Plugins/2D/Paper2D/Source/Paper2D/Classes/PaperTileMap.h
+++ b/Engine/Plugins/2D/Paper2D/Source/Paper2D/Classes/PaperTileMap.h
@@ -139,6 +139,8 @@ public:
 	// Creates a reasonable new layer name
 	static FText GenerateNewLayerName(UPaperTileMap* TileMap);
 
+	// Resize the tile map and all layers
+	void ResizeMap(int32 NewWidth, int32 NewHeight, bool bForceResize = true);
 protected:
 	virtual void UpdateBodySetup();
 };
diff --git a/Engine/Plugins/2D/Paper2D/Source/Paper2D/Classes/PaperTileMapComponent.h b/Engine/Plugins/2D/Paper2D/Source/Paper2D/Classes/PaperTileMapComponent.h
--- a/Engine/Plugins/2D/Paper2D/Source/Paper2D/Classes/PaperTileMapComponent.h
+++ b/Engine/Plugins/2D/Paper2D/Source/Paper2D/Classes/PaperTileMapComponent.h
@@ -87,6 +87,22 @@ public:
 	/** Change the PaperTileMap used by this instance. */
 	UFUNCTION(BlueprintCallable, Category="Sprite")
 	virtual bool SetTileMap(class UPaperTileMap* NewTileMap);
+
+	// Returns the contents of a specified tile cell
+	UFUNCTION(BlueprintPure, Category="Sprite", meta=(Layer="0"))
+	FPaperTileInfo GetTile(int32 X, int32 Y, int32 Layer) const;
+
+	// Modifies the contents of a specified tile cell (Note: This will only work on components that own their own tile map (OwnsTileMap returns true), you cannot modify standalone tile map assets) 
+	UFUNCTION(BlueprintCallable, Category="Sprite", meta=(Layer="0"))
+	void SetTile(int32 X, int32 Y, int32 Layer, FPaperTileInfo NewValue);
+
+	// Resizes the tile map (Note: This will only work on components that own their own tile map (OwnsTileMap returns true), you cannot modify standalone tile map assets) 
+	UFUNCTION(BlueprintCallable, Category="Sprite")
+	void ResizeMap(int32 NewWidthInTiles, int32 NewHeightInTiles);
+
+	// Creates and adds a new layer to the tile map (Note: This will only work on components that own their own tile map(OwnsTileMap returns true), you cannot modify standalone tile map assets)
+	UFUNCTION(BlueprintCallable, Category="Sprite")
+	class UPaperTileLayer* AddNewLayer();
 };
 
 // Allow the old name to continue to work for one release
diff --git a/Engine/Plugins/2D/Paper2D/Source/Paper2D/Private/PaperTileMap.cpp b/Engine/Plugins/2D/Paper2D/Source/Paper2D/Private/PaperTileMap.cpp
--- a/Engine/Plugins/2D/Paper2D/Source/Paper2D/Private/PaperTileMap.cpp
+++ b/Engine/Plugins/2D/Paper2D/Source/Paper2D/Private/PaperTileMap.cpp
@@ -86,15 +86,7 @@ void UPaperTileMap::PostEditChangeProperty(FPropertyChangedEvent& PropertyChange
 
 	if ((PropertyName == GET_MEMBER_NAME_CHECKED(UPaperTileMap, MapWidth)) || (PropertyName == GET_MEMBER_NAME_CHECKED(UPaperTileMap, MapHeight)))
 	{
-		MapWidth = FMath::Max(MapWidth, 1);
-		MapHeight = FMath::Max(MapHeight, 1);
-
-		// Resize all of the existing layers
-		for (int32 LayerIndex = 0; LayerIndex < TileLayers.Num(); ++LayerIndex)
-		{
-			UPaperTileLayer* TileLayer = TileLayers[LayerIndex];
-			TileLayer->ResizeMap(MapWidth, MapHeight);
-		}
+		ResizeMap(MapWidth, MapHeight, /*bForceResize=*/ true);
 	}
 
 	if (!IsTemplate())
@@ -368,6 +360,22 @@ FText UPaperTileMap::GenerateNewLayerName(UPaperTileMap* TileMap)
 	return TestLayerName;
 }
 
+void UPaperTileMap::ResizeMap(int32 NewWidth, int32 NewHeight, bool bForceResize)
+{
+	if (bForceResize || (NewWidth != MapWidth) || (NewHeight != MapHeight))
+	{
+		MapWidth = FMath::Max(NewWidth, 1);
+		MapHeight = FMath::Max(NewHeight, 1);
+
+		// Resize all of the existing layers
+		for (int32 LayerIndex = 0; LayerIndex < TileLayers.Num(); ++LayerIndex)
+		{
+			UPaperTileLayer* TileLayer = TileLayers[LayerIndex];
+			TileLayer->ResizeMap(MapWidth, MapHeight);
+		}
+	}
+}
+
 //////////////////////////////////////////////////////////////////////////
 
 #undef LOCTEXT_NAMESPACE
\ No newline at end of file
diff --git a/Engine/Plugins/2D/Paper2D/Source/Paper2D/Private/PaperTileMapComponent.cpp b/Engine/Plugins/2D/Paper2D/Source/Paper2D/Private/PaperTileMapComponent.cpp
--- a/Engine/Plugins/2D/Paper2D/Source/Paper2D/Private/PaperTileMapComponent.cpp
+++ b/Engine/Plugins/2D/Paper2D/Source/Paper2D/Private/PaperTileMapComponent.cpp
@@ -386,4 +386,60 @@ bool UPaperTileMapComponent::SetTileMap(class UPaperTileMap* NewTileMap)
 	return false;
 }
 
+FPaperTileInfo UPaperTileMapComponent::GetTile(int32 X, int32 Y, int32 Layer) const
+{
+	FPaperTileInfo Result;
+	if (TileMap != nullptr)
+	{
+		if (TileMap->TileLayers.IsValidIndex(Layer))
+		{
+			Result = TileMap->TileLayers[Layer]->GetCell(X, Y);
+		}
+	}
+
+	return Result;
+}
+
+void UPaperTileMapComponent::SetTile(int32 X, int32 Y, int32 Layer, FPaperTileInfo NewValue)
+{
+	if (OwnsTileMap())
+	{
+		if (TileMap->TileLayers.IsValidIndex(Layer))
+		{
+			TileMap->TileLayers[Layer]->SetCell(X, Y, NewValue);
+
+			MarkRenderStateDirty();
+		}
+	}
+}
+
+void UPaperTileMapComponent::ResizeMap(int32 NewWidthInTiles, int32 NewHeightInTiles)
+{
+	if (OwnsTileMap())
+	{
+		TileMap->ResizeMap(NewWidthInTiles, NewHeightInTiles);
+		
+		MarkRenderStateDirty();
+		RecreatePhysicsState();
+		UpdateBounds();
+	}
+}
+
+UPaperTileLayer* UPaperTileMapComponent::AddNewLayer()
+{
+	UPaperTileLayer* Result = nullptr;
+
+	if (OwnsTileMap())
+	{
+		Result = TileMap->AddNewLayer();
+
+		MarkRenderStateDirty();
+		RecreatePhysicsState();
+		UpdateBounds();
+	}
+
+	return Result;
+}
+
+
 #undef LOCTEXT_NAMESPACE
\ No newline at end of file
