[INTEGRATE]
Change 2414955 by [REDACTED_PERSON][REDACTED_HANDLE] on 2015/01/21 23:53:30

Paper2D: Add display of the index of the top-left tile in the tile set viewer, which can be used to make literal tile infos when programmatically constructing tile maps

diff --git a/Engine/Plugins/2D/Paper2D/Source/Paper2DEditor/Private/TileSetEditor.cpp b/Engine/Plugins/2D/Paper2D/Source/Paper2DEditor/Private/TileSetEditor.cpp
--- a/Engine/Plugins/2D/Paper2D/Source/Paper2DEditor/Private/TileSetEditor.cpp
+++ b/Engine/Plugins/2D/Paper2D/Source/Paper2DEditor/Private/TileSetEditor.cpp
@@ -47,6 +47,7 @@ public:
 	FTileSetEditorViewportClient(UPaperTileSet* InTileSet)
 		: TileSetBeingEdited(InTileSet)
 		, bHasValidPaintRectangle(false)
+		, TileIndex(INDEX_NONE)
 	{
 	}
 
@@ -64,6 +65,7 @@ public:
 
 	bool bHasValidPaintRectangle;
 	FViewportSelectionRectangle ValidPaintRectangle;
+	int32 TileIndex;
 };
 
 void FTileSetEditorViewportClient::Draw(FViewport* Viewport, FCanvas* Canvas)
@@ -126,6 +128,17 @@ void FTileSetEditorViewportClient::Draw(FViewport* Viewport, FCanvas* Canvas)
 		FCanvasBoxItem BoxItem(FVector2D(X, Y), FVector2D(W, H));
 		BoxItem.SetColor(Rect.Color);
 		Canvas->DrawItem(BoxItem);
+
+	}
+
+	if (TileIndex != INDEX_NONE)
+	{
+		const FString TileIndexString = FString::Printf(TEXT("Tile# %d"), TileIndex);
+
+		int32 XL;
+		int32 YL;
+		StringSize(GEngine->GetLargeFont(), XL, YL, *TileIndexString);
+		Canvas->DrawShadowedString(4, Viewport->GetSizeXY().Y - YL - 4, *TileIndexString, GEngine->GetLargeFont(), FLinearColor::White);
 	}
 }
 
@@ -250,6 +263,9 @@ void STileSetSelectorViewport::RefreshSelectionRectangle()
 		const bool bHasSelection = (SelectionDimensions.X * SelectionDimensions.Y) > 0;
 		Client->bHasValidPaintRectangle = bHasSelection && (TileSetBeingEdited != nullptr);
 
+		const int32 TileIndex = (bHasSelection && (TileSetBeingEdited != nullptr)) ? (SelectionTopLeft.X + SelectionTopLeft.Y * TileSetBeingEdited->GetTileCountX()) : INDEX_NONE;
+		Client->TileIndex = TileIndex;
+
 		if (bHasSelection)
 		{
 			Client->ValidPaintRectangle.Color = FLinearColor::White;
