[INTEGRATE]
Change 2416451 by [REDACTED_PERSON][REDACTED_HANDLE] on 2015/01/23 00:39:28

Paper2D: Prevent a crash when switching levels after selecting a tile map actor that had an instanced tile map, which was caused by the SListView keeping a FGCObject reference to the layers and therefore indirectly the world itself [UE-7728]

diff --git a/Engine/Plugins/2D/Paper2D/Source/Paper2DEditor/Private/TileMapEditing/STileLayerItem.cpp b/Engine/Plugins/2D/Paper2D/Source/Paper2DEditor/Private/TileMapEditing/STileLayerItem.cpp
--- a/Engine/Plugins/2D/Paper2D/Source/Paper2DEditor/Private/TileMapEditing/STileLayerItem.cpp
+++ b/Engine/Plugins/2D/Paper2D/Source/Paper2DEditor/Private/TileMapEditing/STileLayerItem.cpp
@@ -12,9 +12,10 @@
 //////////////////////////////////////////////////////////////////////////
 // STileLayerItem
 
-void STileLayerItem::Construct(const FArguments& InArgs, class UPaperTileLayer* InItem, FIsSelected InIsSelectedDelegate)
+void STileLayerItem::Construct(const FArguments& InArgs, int32 Index, class UPaperTileMap* InMap, FIsSelected InIsSelectedDelegate)
 {
-	MyLayer = InItem;
+	MyMap = InMap;
+	MyIndex = Index;
 
 	static const FName EyeClosedBrushName("TileMapEditor.LayerEyeClosed");
 	static const FName EyeOpenedBrushName("TileMapEditor.LayerEyeOpened");
@@ -60,12 +61,13 @@ FText STileLayerItem::GetLayerDisplayName() const
 {
 	const FText UnnamedText = LOCTEXT("NoLayerName", "(unnamed)");
 
-	return MyLayer->LayerName.IsEmpty() ? UnnamedText : MyLayer->LayerName;
+	return GetMyLayer()->LayerName.IsEmpty() ? UnnamedText : GetMyLayer()->LayerName;
 }
 
 void STileLayerItem::OnLayerNameCommitted(const FText& NewText, ETextCommit::Type CommitInfo)
 {
 	const FScopedTransaction Transaction( LOCTEXT("TileMapRenameLayer", "Rename Layer") );
+	UPaperTileLayer* MyLayer = GetMyLayer();
 	MyLayer->SetFlags(RF_Transactional);
 	MyLayer->Modify();
 	MyLayer->LayerName = NewText;
@@ -74,6 +76,7 @@ void STileLayerItem::OnLayerNameCommitted(const FText& NewText, ETextCommit::Typ
 FReply STileLayerItem::OnToggleVisibility()
 {
 	const FScopedTransaction Transaction( LOCTEXT("ToggleVisibility", "Toggle Layer Visibility") );
+	UPaperTileLayer* MyLayer = GetMyLayer();
 	MyLayer->SetFlags(RF_Transactional);
 	MyLayer->Modify();
 	MyLayer->bHiddenInEditor = !MyLayer->bHiddenInEditor;
@@ -83,7 +86,7 @@ FReply STileLayerItem::OnToggleVisibility()
 
 const FSlateBrush* STileLayerItem::GetVisibilityBrushForLayer() const
 {
-	return MyLayer->bHiddenInEditor ? EyeClosed : EyeOpened;
+	return GetMyLayer()->bHiddenInEditor ? EyeClosed : EyeOpened;
 }
 
 FSlateColor STileLayerItem::GetForegroundColorForVisibilityButton() const
diff --git a/Engine/Plugins/2D/Paper2D/Source/Paper2DEditor/Private/TileMapEditing/STileLayerItem.h b/Engine/Plugins/2D/Paper2D/Source/Paper2DEditor/Private/TileMapEditing/STileLayerItem.h
--- a/Engine/Plugins/2D/Paper2D/Source/Paper2DEditor/Private/TileMapEditing/STileLayerItem.h
+++ b/Engine/Plugins/2D/Paper2D/Source/Paper2DEditor/Private/TileMapEditing/STileLayerItem.h
@@ -11,10 +11,12 @@ public:
 	SLATE_BEGIN_ARGS(STileLayerItem) {}
 	SLATE_END_ARGS()
 
-	void Construct(const FArguments& InArgs, class UPaperTileLayer* InItem, FIsSelected InIsSelectedDelegate);
+	void Construct(const FArguments& InArgs, int32 Index, class UPaperTileMap* InMap, FIsSelected InIsSelectedDelegate);
 
 protected:
-	class UPaperTileLayer* MyLayer;
+	int32 MyIndex;
+	class UPaperTileMap* MyMap;
+	class UPaperTileLayer* GetMyLayer() const { return MyMap->TileLayers[MyIndex]; }
 
 	TSharedPtr<SButton> VisibilityButton;
 
diff --git a/Engine/Plugins/2D/Paper2D/Source/Paper2DEditor/Private/TileMapEditing/STileLayerList.cpp b/Engine/Plugins/2D/Paper2D/Source/Paper2DEditor/Private/TileMapEditing/STileLayerList.cpp
--- a/Engine/Plugins/2D/Paper2D/Source/Paper2DEditor/Private/TileMapEditing/STileLayerList.cpp
+++ b/Engine/Plugins/2D/Paper2D/Source/Paper2DEditor/Private/TileMapEditing/STileLayerList.cpp
@@ -68,19 +68,22 @@ void STileLayerList::Construct(const FArguments& InArgs, UPaperTileMap* InTileMa
 
 	TSharedRef<SWidget> Toolbar = ToolbarBuilder.MakeWidget();
 
+	RefreshMirrorList();
+
 	ListViewWidget = SNew(SPaperLayerListView)
 		.SelectionMode(ESelectionMode::Single)
 		.ClearSelectionOnClick(false)
-		.ListItemsSource(&(InTileMap->TileLayers))
+		.ListItemsSource(&MirrorList)
 		.OnSelectionChanged(this, &STileLayerList::OnSelectionChanged)
 		.OnGenerateRow(this, &STileLayerList::OnGenerateLayerListRow)
 		.OnContextMenuOpening(this, &STileLayerList::OnConstructContextMenu);
 
 	// Restore the selection
 	InTileMap->ValidateSelectedLayerIndex();
-	if (InTileMap->SelectedLayerIndex != INDEX_NONE)
+	if (InTileMap->TileLayers.IsValidIndex(InTileMap->SelectedLayerIndex))
 	{
-		SetSelectedLayer(InTileMap->TileLayers[InTileMap->SelectedLayerIndex]);
+		UPaperTileLayer* SelectedLayer = InTileMap->TileLayers[InTileMap->SelectedLayerIndex];
+		SetSelectedLayer(SelectedLayer);
 	}
 
 	ChildSlot
@@ -102,22 +105,35 @@ void STileLayerList::Construct(const FArguments& InArgs, UPaperTileMap* InTileMa
 	];
 }
 
-TSharedRef<ITableRow> STileLayerList::OnGenerateLayerListRow(class UPaperTileLayer* Item, const TSharedRef<STableViewBase>& OwnerTable)
+TSharedRef<ITableRow> STileLayerList::OnGenerateLayerListRow(FMirrorEntry Item, const TSharedRef<STableViewBase>& OwnerTable)
 {
-	typedef STableRow<class UPaperTileLayer*> RowType;
+	typedef STableRow<FMirrorEntry> RowType;
 
 	TSharedRef<RowType> NewRow = SNew(RowType, OwnerTable)
 		.Style(&FPaperStyle::Get()->GetWidgetStyle<FTableRowStyle>("TileMapEditor.LayerBrowser.TableViewRow"));
 
 	FIsSelected IsSelectedDelegate = FIsSelected::CreateSP(NewRow, &RowType::IsSelectedExclusively);
-	NewRow->SetContent(SNew(STileLayerItem, Item, IsSelectedDelegate));
+	NewRow->SetContent(SNew(STileLayerItem, *Item, TileMapPtr.Get(), IsSelectedDelegate));
 
 	return NewRow;
 }
 
 UPaperTileLayer* STileLayerList::GetSelectedLayer() const
 {
-	return (ListViewWidget->GetNumItemsSelected() > 0) ? ListViewWidget->GetSelectedItems()[0] : nullptr;
+	if (UPaperTileMap* TileMap = TileMapPtr.Get())
+	{
+		if (ListViewWidget->GetNumItemsSelected() > 0)
+		{
+			FMirrorEntry SelectedItem = ListViewWidget->GetSelectedItems()[0];
+			const int32 SelectedIndex = *SelectedItem;
+			if (TileMap->TileLayers.IsValidIndex(SelectedIndex))
+			{
+				return TileMap->TileLayers[SelectedIndex];
+			}
+		}
+	}
+
+	return nullptr;
 }
 
 FText STileLayerList::GenerateDuplicatedLayerName(const FString& InputNameRaw, UPaperTileMap* TileMap)
@@ -212,6 +228,7 @@ void STileLayerList::ChangeLayerOrdering(int32 OldIndex, int32 NewIndex)
 			if (TileMap->SelectedLayerIndex == OldIndex)
 			{
 				TileMap->SelectedLayerIndex = NewIndex;
+				SetSelectedLayer(LayerToMove);
 			}
 
 			PostEditNotfications();
@@ -377,10 +394,20 @@ bool STileLayerList::CanExecuteActionNeedingSelectedLayer() const
 
 void STileLayerList::SetSelectedLayer(UPaperTileLayer* SelectedLayer)
 {
-	ListViewWidget->SetSelection(SelectedLayer);
+	if (UPaperTileMap* TileMap = TileMapPtr.Get())
+	{
+		int32 NewIndex;
+		if (TileMap->TileLayers.Find(SelectedLayer, /*out*/ NewIndex))
+		{
+			if (MirrorList.IsValidIndex(NewIndex))
+			{
+				ListViewWidget->SetSelection(MirrorList[NewIndex]);
+			}
+		}
+	}
 }
 
-void STileLayerList::OnSelectionChanged(UPaperTileLayer* ItemChangingState, ESelectInfo::Type SelectInfo)
+void STileLayerList::OnSelectionChanged(FMirrorEntry ItemChangingState, ESelectInfo::Type SelectInfo)
 {
 	if (UPaperTileMap* TileMap = TileMapPtr.Get())
 	{
@@ -410,6 +437,8 @@ TSharedPtr<SWidget> STileLayerList::OnConstructContextMenu()
 
 void STileLayerList::PostEditNotfications()
 {
+	RefreshMirrorList();
+
 	ListViewWidget->RequestListRefresh();
 
 	if (UPaperTileMap* TileMap = TileMapPtr.Get())
@@ -425,6 +454,31 @@ void STileLayerList::PostEditNotfications()
 	}
 }
 
+void STileLayerList::RefreshMirrorList()
+{
+	if (UPaperTileMap* TileMap = TileMapPtr.Get())
+	{
+		const int32 NumEntriesToAdd = TileMap->TileLayers.Num() - MirrorList.Num();
+		if (NumEntriesToAdd < 0)
+		{
+			MirrorList.RemoveAt(TileMap->TileLayers.Num(), -NumEntriesToAdd);
+		}
+		else if (NumEntriesToAdd > 0)
+		{
+			for (int32 Count = 0; Count < NumEntriesToAdd; ++Count)
+			{
+				TSharedPtr<int32> NewEntry = MakeShareable(new int32);
+				*NewEntry = MirrorList.Num();
+				MirrorList.Add(NewEntry);
+			}
+		}
+	}
+	else
+	{
+		MirrorList.Empty();
+	}
+}
+
 //////////////////////////////////////////////////////////////////////////
 
 #undef LOCTEXT_NAMESPACE
\ No newline at end of file
diff --git a/Engine/Plugins/2D/Paper2D/Source/Paper2DEditor/Private/TileMapEditing/STileLayerList.h b/Engine/Plugins/2D/Paper2D/Source/Paper2DEditor/Private/TileMapEditing/STileLayerList.h
--- a/Engine/Plugins/2D/Paper2D/Source/Paper2DEditor/Private/TileMapEditing/STileLayerList.h
+++ b/Engine/Plugins/2D/Paper2D/Source/Paper2DEditor/Private/TileMapEditing/STileLayerList.h
@@ -14,7 +14,10 @@ public:
 	void Construct(const FArguments& InArgs, UPaperTileMap* TileMap, FNotifyHook* InNotifyHook);
 
 protected:
-	typedef SListView<class UPaperTileLayer*> SPaperLayerListView;
+	typedef TSharedPtr<int32> FMirrorEntry;
+
+	TArray<FMirrorEntry> MirrorList;
+	typedef SListView<FMirrorEntry> SPaperLayerListView;
 
 protected:
 	TSharedPtr<SPaperLayerListView> ListViewWidget;
@@ -22,7 +25,7 @@ protected:
 	TWeakObjectPtr<class UPaperTileMap> TileMapPtr;
 	FNotifyHook* NotifyHook;
 protected:
-	TSharedRef<ITableRow> OnGenerateLayerListRow(class UPaperTileLayer* Item, const TSharedRef<STableViewBase>& OwnerTable);
+	TSharedRef<ITableRow> OnGenerateLayerListRow(FMirrorEntry Item, const TSharedRef<STableViewBase>& OwnerTable);
 
 	class UPaperTileLayer* GetSelectedLayer() const;
 
@@ -51,9 +54,11 @@ protected:
 
 	void SetSelectedLayer(UPaperTileLayer* SelectedLayer);
 
-	void OnSelectionChanged(UPaperTileLayer* ItemChangingState, ESelectInfo::Type SelectInfo);
+	void OnSelectionChanged(FMirrorEntry ItemChangingState, ESelectInfo::Type SelectInfo);
 	TSharedPtr<SWidget> OnConstructContextMenu();
 
+	void RefreshMirrorList();
+
 	// Called after edits are finished
 	void PostEditNotfications();
 };
