Paper2D: Validate socket names when editing sprites, makes sure names are unique and not-empty

diff --git a/Engine/Plugins/2D/Paper2D/Source/Paper2D/Classes/PaperSprite.h b/Engine/Plugins/2D/Paper2D/Source/Paper2D/Classes/PaperSprite.h
--- a/Engine/Plugins/2D/Paper2D/Source/Paper2D/Classes/PaperSprite.h
+++ b/Engine/Plugins/2D/Paper2D/Source/Paper2D/Classes/PaperSprite.h
@@ -241,6 +241,12 @@ public:
 	void OnObjectReimported(UObject* InObject);
 #endif
 
+#if WITH_EDITOR
+	// Make sure all socket names are valid
+	// All duplicate / empty names will be made unique
+	void ValidateSocketNames();
+#endif
+
 	// Return the scaling factor between pixels and Unreal units (cm)
 	float GetPixelsPerUnrealUnit() const { return PixelsPerUnrealUnit; }
 
diff --git a/Engine/Plugins/2D/Paper2D/Source/Paper2D/Private/PaperSprite.cpp b/Engine/Plugins/2D/Paper2D/Source/Paper2D/Private/PaperSprite.cpp
--- a/Engine/Plugins/2D/Paper2D/Source/Paper2D/Private/PaperSprite.cpp
+++ b/Engine/Plugins/2D/Paper2D/Source/Paper2D/Private/PaperSprite.cpp
@@ -513,6 +513,12 @@ void UPaperSprite::PostEditChangeProperty(FPropertyChangedEvent& PropertyChanged
 		bBothModified = true;
 	}
 
+	if (PropertyName == GET_MEMBER_NAME_CHECKED(UPaperSprite, Sockets) ||
+		(MemberPropertyName == GET_MEMBER_NAME_CHECKED(UPaperSprite, Sockets) && PropertyName == GET_MEMBER_NAME_CHECKED(FPaperSpriteSocket, SocketName)))
+	{
+		ValidateSocketNames();
+	}
+
 	// The texture dimensions have changed
 	//if (NeedRescaleSpriteData())
 	//{
@@ -1768,6 +1774,52 @@ void UPaperSprite::QuerySupportedSockets(TArray<FComponentSocketDescription>& Ou
 	}
 }
 
+#if WITH_EDITOR
+void UPaperSprite::ValidateSocketNames()
+{
+	TSet<FName> SocketNames;
+	struct Local
+	{
+		static FName GetUniqueName(const TSet<FName>& SocketNames, FName Name)
+		{
+			int Counter = Name.GetNumber();
+			FName TestName;
+			do 
+			{
+				TestName = Name;
+				TestName.SetNumber(++Counter);
+			} while (SocketNames.Contains(TestName));
+
+			return TestName;
+		}
+	};
+
+	bool bHasChanged = false;
+	for (int32 SocketIndex = 0; SocketIndex < Sockets.Num(); ++SocketIndex)
+	{
+		FPaperSpriteSocket& Socket = Sockets[SocketIndex];
+		if (Socket.SocketName.IsNone())
+		{
+			Socket.SocketName = Local::GetUniqueName(SocketNames, FName(TEXT("Socket")));
+			bHasChanged = true;
+		}
+		else if (SocketNames.Contains(Socket.SocketName))
+		{
+			Socket.SocketName = Local::GetUniqueName(SocketNames, Socket.SocketName);
+			bHasChanged = true;
+		}
+
+		// Add the corrected name
+		SocketNames.Add(Socket.SocketName);
+	}
+
+	if (bHasChanged)
+	{
+		PostEditChange();
+	}
+}
+#endif
+
 #if WITH_EDITOR
 void UPaperSprite::Serialize(FArchive& Ar)
 {
