Paper2D: Add extract sprites dialog allowing choosing between and configuring grid based and automatic extraction methods

diff --git a/Engine/Plugins/2D/Paper2D/Source/Paper2DEditor/Private/ContentBrowserExtensions/ContentBrowserExtensions.cpp b/Engine/Plugins/2D/Paper2D/Source/Paper2DEditor/Private/ContentBrowserExtensions/ContentBrowserExtensions.cpp
--- a/Engine/Plugins/2D/Paper2D/Source/Paper2DEditor/Private/ContentBrowserExtensions/ContentBrowserExtensions.cpp
+++ b/Engine/Plugins/2D/Paper2D/Source/Paper2DEditor/Private/ContentBrowserExtensions/ContentBrowserExtensions.cpp
@@ -4,6 +4,8 @@
 #include "ContentBrowserModule.h"
 #include "ContentBrowserExtensions.h"
 
+#include "ExtractSprites/SPaperExtractSpritesDialog.h"
+
 #define LOCTEXT_NAMESPACE "Paper2D"
 
 DECLARE_LOG_CATEGORY_EXTERN(LogPaperCBExtensions, Log, All);
@@ -45,8 +47,6 @@ struct FCreateSpriteFromTextureExtension : public FContentBrowserSelectedAssetEx
 
 	void CreateSpritesFromTextures(TArray<UTexture2D*>& Textures)
 	{
-		const FString DefaultSuffix = TEXT("_Sprite");
-
 		FAssetToolsModule& AssetToolsModule = FModuleManager::Get().LoadModuleChecked<FAssetToolsModule>("AssetTools");
 		FContentBrowserModule& ContentBrowserModule = FModuleManager::LoadModuleChecked<FContentBrowserModule>("ContentBrowser");
 
@@ -67,6 +67,7 @@ struct FCreateSpriteFromTextureExtension : public FContentBrowserSelectedAssetEx
 			if (!bExtractSprites)
 			{
 				// Get a unique name for the sprite
+				const FString DefaultSuffix = TEXT("_Sprite");
 				AssetToolsModule.Get().CreateUniqueAssetName(Texture->GetOutermost()->GetName(), DefaultSuffix, /*out*/ PackageName, /*out*/ Name);
 				const FString PackagePath = FPackageName::GetLongPackagePath(PackageName);
 
@@ -77,81 +78,7 @@ struct FCreateSpriteFromTextureExtension : public FContentBrowserSelectedAssetEx
 			}
 			else
 			{
-				FScopedSlowTask Feedback(1, NSLOCTEXT("Paper2D", "Paper2D_ExtractSpritesFromTexture", "Extracting Sprites From Texture"));
-				Feedback.MakeDialog(true);
-
-				// First extract the rects from the texture
-				TArray<FIntRect> ExtractedRects;
-				UPaperSprite::ExtractRectsFromTexture(Texture, /*out*/ ExtractedRects);
-
-				// Sort the rectangles by approximate row
-				struct FRectangleSortHelper
-				{
-					FRectangleSortHelper(TArray<FIntRect>& InOutSprites)
-					{
-						// Sort by Y, then by X (top left corner), descending order (so we can use it as a stack from the top row down)
-						TArray<FIntRect> SpritesLeft = InOutSprites;
-						SpritesLeft.Sort([](const FIntRect& A, const FIntRect& B) { return (A.Min.Y == B.Min.Y) ? (A.Min.X > B.Min.X) : (A.Min.Y > B.Min.Y); });
-						InOutSprites.Reset();
-
-						// Start pulling sprites out, the first one in each row will dominate remaining ones and cause them to get labeled
-						TArray<FIntRect> DominatedSprites;
-						DominatedSprites.Empty(SpritesLeft.Num());
- 						while (SpritesLeft.Num())
- 						{
-							FIntRect DominatingSprite = SpritesLeft.Pop();
-							DominatedSprites.Add(DominatingSprite);
-
-							// Find the sprites that are dominated (intersect the infinite horizontal band described by the dominating sprite)
-							for (int32 Index = 0; Index < SpritesLeft.Num();)
-							{
-								const FIntRect& CurElement = SpritesLeft[Index];
-								if ((CurElement.Min.Y <= DominatingSprite.Max.Y) && (CurElement.Max.Y >= DominatingSprite.Min.Y))
-								{
-									DominatedSprites.Add(CurElement);
-									SpritesLeft.RemoveAt(Index, /*Count=*/ 1, /*bAllowShrinking=*/ false);
-								}
-								else
-								{
-									++Index;
-								}
-							}
-
-							// Sort the sprites in the band by X and add them to the result
-							DominatedSprites.Sort([](const FIntRect& A, const FIntRect& B) { return (A.Min.X < B.Min.X); });
-							InOutSprites.Append(DominatedSprites);
-							DominatedSprites.Reset();
- 						}
-					}
-				};
-				FRectangleSortHelper RectSorter(ExtractedRects);
-
-				Feedback.TotalAmountOfWork = ExtractedRects.Num();
-
-				for (int ExtractedRectIndex = 0; ExtractedRectIndex < ExtractedRects.Num(); ++ExtractedRectIndex)
-				{
-					Feedback.EnterProgressFrame(1, NSLOCTEXT("Paper2D", "Paper2D_ExtractSpritesFromTexture", "Extracting Sprites From Texture"));
-
-					FIntRect& ExtractedRect = ExtractedRects[ExtractedRectIndex];
-					SpriteFactory->bUseSourceRegion = true;
-					SpriteFactory->InitialSourceUV = FVector2D(ExtractedRect.Min.X, ExtractedRect.Min.Y);
-					SpriteFactory->InitialSourceDimension = FVector2D(ExtractedRect.Width(), ExtractedRect.Height());
-
-					// Get a unique name for the sprite
-					const FString Suffix = FString::Printf(TEXT("%s_%d"), *DefaultSuffix, ExtractedRectIndex);
-					AssetToolsModule.Get().CreateUniqueAssetName(Texture->GetOutermost()->GetName(), Suffix, /*out*/ PackageName, /*out*/ Name);
-					const FString PackagePath = FPackageName::GetLongPackagePath(PackageName);
-
-					if (UObject* NewAsset = AssetToolsModule.Get().CreateAsset(Name, PackagePath, UPaperSprite::StaticClass(), SpriteFactory))
-					{
-						ObjectsToSync.Add(NewAsset);
-					}
-
-					if (GWarn->ReceivedUserCancel()) 
-					{
-						break;
-					}
-				}
+				SPaperExtractSpritesDialog::ShowWindow(NSLOCTEXT("Paper2D", "Paper2D_ExtractSprites", "Extract sprites"), Texture);
 			}
 		}
 
@@ -301,7 +228,7 @@ public:
 
 		MenuBuilder.AddMenuEntry(
 			LOCTEXT("CB_Extension_Texture_ExtractSprites", "Extract Sprites"),
-			LOCTEXT("CB_Extension_Texture_ExtractSprite_Tooltip", "Detects and extracts sprites from the selected textures using transparency"),
+			LOCTEXT("CB_Extension_Texture_ExtractSprites_Tooltip", "Extract sprites from selected textures"),
 			FSlateIcon(),
 			Action_ExtractSpritesFromTextures,
 			NAME_None,
diff --git a/Engine/Plugins/2D/Paper2D/Source/Paper2DEditor/Private/ExtractSprites/PaperExtractSpritesSettings.cpp b/Engine/Plugins/2D/Paper2D/Source/Paper2DEditor/Private/ExtractSprites/PaperExtractSpritesSettings.cpp
new file mode 100644
--- /dev/null
+++ b/Engine/Plugins/2D/Paper2D/Source/Paper2DEditor/Private/ExtractSprites/PaperExtractSpritesSettings.cpp
@@ -0,0 +1,22 @@
+// Copyright 1998-2015 Epic Games, Inc. All Rights Reserved.
+
+#include "Paper2DEditorPrivatePCH.h"
+#include "PaperExtractSpritesSettings.h"
+
+UPaperExtractSpritesSettings::UPaperExtractSpritesSettings(const FObjectInitializer& ObjectInitializer)
+	: Super(ObjectInitializer)
+{
+	OutlineColor = FLinearColor::Yellow;
+	TextureTint = FLinearColor::Gray;
+}
+
+void UPaperExtractSpritesSettings::PostEditChangeProperty(FPropertyChangedEvent& PropertyChangedEvent)
+{
+	const FName PropertyName = (PropertyChangedEvent.Property != nullptr) ? PropertyChangedEvent.Property->GetFName() : NAME_None;
+	if (PropertyName == GET_MEMBER_NAME_CHECKED(UPaperExtractSpritesSettings, Mode))
+	{
+		// Update bools
+		bAutoMode = (Mode == ESpriteExtractMode::Auto);
+		bGridMode = (Mode == ESpriteExtractMode::Grid);
+	}
+}
diff --git a/Engine/Plugins/2D/Paper2D/Source/Paper2DEditor/Private/ExtractSprites/PaperExtractSpritesSettings.h b/Engine/Plugins/2D/Paper2D/Source/Paper2DEditor/Private/ExtractSprites/PaperExtractSpritesSettings.h
new file mode 100644
--- /dev/null
+++ b/Engine/Plugins/2D/Paper2D/Source/Paper2DEditor/Private/ExtractSprites/PaperExtractSpritesSettings.h
@@ -0,0 +1,61 @@
+// Copyright 1998-2015 Epic Games, Inc. All Rights Reserved.
+
+#pragma once
+
+#include "PaperExtractSpritesSettings.generated.h"
+
+UENUM()
+enum class ESpriteExtractMode : uint8
+{
+	Auto,	// Automatically extract sprites
+	Grid,	// Extract sprites by grid
+};
+
+/**
+*
+*/
+UCLASS()
+class UPaperExtractSpritesSettings : public UObject
+{
+	GENERATED_BODY()
+
+public:
+	UPROPERTY(Category = Settings, EditAnywhere)
+	ESpriteExtractMode Mode;
+
+	UPROPERTY(Category = Settings, EditAnywhere)
+	FLinearColor OutlineColor;
+
+	UPROPERTY(Category = Settings, EditAnywhere)
+	FLinearColor TextureTint;
+
+	UPROPERTY()
+	bool bAutoMode;
+
+	UPROPERTY()
+	bool bGridMode;
+
+
+	//////////////////////////////////////////////////////////////////////////
+	// Grid mode
+
+	UPROPERTY(Category = Grid, EditAnywhere, meta = (EditCondition = "bGridMode", UIMin = 1, ClampMin = 1))
+	int32 GridWidth;
+
+	UPROPERTY(Category = Grid, EditAnywhere, meta = (EditCondition = "bGridMode", UIMin = 1, ClampMin = 1))
+	int32 GridHeight;
+
+	UPROPERTY(Category = Grid, EditAnywhere, meta = (EditCondition = "bGridMode", UIMin = 0, ClampMin = 0))
+	int32 Margin;
+
+	UPROPERTY(Category = Grid, EditAnywhere, meta = (EditCondition = "bGridMode", UIMin = 0, ClampMin = 0))
+	int32 Spacing;
+	
+
+
+	UPaperExtractSpritesSettings(const FObjectInitializer& ObjectInitializer);
+
+	// UObject interface
+	virtual void PostEditChangeProperty(FPropertyChangedEvent& PropertyChangedEvent) override;
+	// End of UObject interface
+};
diff --git a/Engine/Plugins/2D/Paper2D/Source/Paper2DEditor/Private/ExtractSprites/SPaperExtractSpritesDialog.cpp b/Engine/Plugins/2D/Paper2D/Source/Paper2DEditor/Private/ExtractSprites/SPaperExtractSpritesDialog.cpp
new file mode 100644
--- /dev/null
+++ b/Engine/Plugins/2D/Paper2D/Source/Paper2DEditor/Private/ExtractSprites/SPaperExtractSpritesDialog.cpp
@@ -0,0 +1,390 @@
+// Copyright 1998-2015 Epic Games, Inc. All Rights Reserved.
+
+#include "Paper2DEditorPrivatePCH.h"
+#include "SPaperExtractSpritesDialog.h"
+#include "PaperEditorViewportClient.h"
+#include "ExtractSprites/PaperExtractSpritesSettings.h"
+
+#include "PropertyEditorModule.h"
+#include "Editor/PropertyEditor/Public/IDetailsView.h"
+#include "ContentBrowserModule.h"
+#include "AssetToolsModule.h"
+
+#define LOCTEXT_NAMESPACE "PaperEditor"
+
+//////////////////////////////////////////////////////////////////////////
+// FTileSetEditorViewportClient 
+
+class FPaperExtractSpritesViewportClient : public FPaperEditorViewportClient
+{
+public:
+	FPaperExtractSpritesViewportClient(UTexture2D* Texture, const TArray<FPaperExtractedSprite>& InExtractedSprites, const UPaperExtractSpritesSettings* InSettings)
+		: TextureBeingExtracted(Texture)
+		, ExtractedSprites(InExtractedSprites)
+		, Settings(InSettings)
+	{
+	}
+
+	// FViewportClient interface
+	virtual void Draw(FViewport* Viewport, FCanvas* Canvas) override;
+	// End of FViewportClient interface
+
+private:
+	// Tile set
+	TWeakObjectPtr<UTexture2D> TextureBeingExtracted;
+	const TArray<FPaperExtractedSprite>& ExtractedSprites;
+	const UPaperExtractSpritesSettings* Settings;
+
+	void DrawRectangle(FCanvas* Canvas, const FLinearColor& Color, const FIntRect& Rect);
+};
+
+void FPaperExtractSpritesViewportClient::Draw(FViewport* Viewport, FCanvas* Canvas)
+{
+	// Super will clear the viewport 
+	FPaperEditorViewportClient::Draw(Viewport, Canvas);
+
+	UTexture2D* Texture = TextureBeingExtracted.Get();
+
+	if (Texture != nullptr)
+	{
+		const bool bUseTranslucentBlend = Texture->HasAlphaChannel();
+
+		// Fully stream in the texture before drawing it.
+		Texture->SetForceMipLevelsToBeResident(30.0f);
+		Texture->WaitForStreaming();
+
+		FLinearColor TextureDrawColor = Settings->TextureTint;
+		//FLinearColor RectOutlineColor = FLinearColor::Yellow;
+		const FLinearColor RectOutlineColor = Settings->OutlineColor;
+
+		const float XPos = -ZoomPos.X * ZoomAmount;
+		const float YPos = -ZoomPos.Y * ZoomAmount;
+		const float Width = Texture->GetSurfaceWidth() * ZoomAmount;
+		const float Height = Texture->GetSurfaceHeight() * ZoomAmount;
+
+		Canvas->DrawTile(XPos, YPos, Width, Height, 0.0f, 0.0f, 1.0f, 1.0f, TextureDrawColor, Texture->Resource, bUseTranslucentBlend);
+
+		for (FPaperExtractedSprite Sprite : ExtractedSprites)
+		{
+			DrawRectangle(Canvas, RectOutlineColor, Sprite.Rect);
+		}
+	}
+}
+
+void FPaperExtractSpritesViewportClient::DrawRectangle(FCanvas* Canvas, const FLinearColor& Color, const FIntRect& Rect)
+{
+	FVector2D TopLeft(-ZoomPos.X * ZoomAmount + Rect.Min.X * ZoomAmount, -ZoomPos.Y * ZoomAmount + Rect.Min.Y * ZoomAmount);
+	FVector2D BottomRight(-ZoomPos.X * ZoomAmount + Rect.Max.X * ZoomAmount, -ZoomPos.Y * ZoomAmount + Rect.Max.Y * ZoomAmount);
+	FVector2D RectVertices[4];
+	RectVertices[0] = FVector2D(TopLeft.X, TopLeft.Y);
+	RectVertices[1] = FVector2D(BottomRight.X, TopLeft.Y);
+	RectVertices[2] = FVector2D(BottomRight.X, BottomRight.Y);
+	RectVertices[3] = FVector2D(TopLeft.X, BottomRight.Y);
+	for (int32 RectVertexIndex = 0; RectVertexIndex < 4; ++RectVertexIndex)
+	{
+		const int32 NextVertexIndex = (RectVertexIndex + 1) % 4;
+		FCanvasLineItem RectLine(RectVertices[RectVertexIndex], RectVertices[NextVertexIndex]);
+		RectLine.SetColor(Color);
+		Canvas->DrawItem(RectLine);
+	}
+}
+
+//////////////////////////////////////////////////////////////////////////
+// SPaperExtractSpritesViewport
+
+SPaperExtractSpritesViewport::~SPaperExtractSpritesViewport()
+{
+	TypedViewportClient = nullptr;
+}
+
+void SPaperExtractSpritesViewport::Construct(const FArguments& InArgs, UTexture2D* InTexture, const TArray<FPaperExtractedSprite>& ExtractedSprites, const UPaperExtractSpritesSettings* Settings, SPaperExtractSpritesDialog* InDialog)
+{
+	TexturePtr = InTexture;
+	Dialog = InDialog;
+
+	TypedViewportClient = MakeShareable(new FPaperExtractSpritesViewportClient(InTexture, ExtractedSprites, Settings));
+
+	SPaperEditorViewport::Construct(
+		SPaperEditorViewport::FArguments(),
+		TypedViewportClient.ToSharedRef());
+
+	// Make sure we get input instead of the viewport stealing it
+	ViewportWidget->SetVisibility(EVisibility::HitTestInvisible);
+
+	RefreshViewport();
+}
+
+FText SPaperExtractSpritesViewport::GetTitleText() const
+{
+	return FText::FromString(TexturePtr->GetName());
+}
+
+//////////////////////////////////////////////////////////////////////////
+// SPaperExtractSpritesDialog
+
+void SPaperExtractSpritesDialog::Construct(const FArguments& InArgs, UTexture2D* InSourceTexture)
+{
+	SourceTexture = InSourceTexture;
+
+	ExtractSpriteSettings = NewObject<UPaperExtractSpritesSettings>();
+	ExtractSpriteSettings->GridWidth = InSourceTexture->GetSizeX();
+	ExtractSpriteSettings->GridHeight = InSourceTexture->GetSizeY();
+	PreviewExtractedSprites();
+
+	FPropertyEditorModule& EditModule = FModuleManager::Get().GetModuleChecked<FPropertyEditorModule>("PropertyEditor");
+	FDetailsViewArgs DetailsViewArgs(/*bUpdateFromSelection=*/ false, /*bLockable=*/ false, /*bAllowSearch=*/ false, /*bObjectsUseNameArea=*/ true, /*bHideSelectionTip=*/ true);
+	DetailsViewArgs.bHideActorNameArea = true;
+	PropertyView = EditModule.CreateDetailView(DetailsViewArgs);
+	PropertyView->SetObject(ExtractSpriteSettings);
+	PropertyView->OnFinishedChangingProperties().AddSP(this, &SPaperExtractSpritesDialog::OnFinishedChangingProperties);
+	
+	TSharedRef<SPaperExtractSpritesViewport> Viewport = SNew(SPaperExtractSpritesViewport, SourceTexture, ExtractedSprites, ExtractSpriteSettings, this);
+
+	ChildSlot
+	[
+		SNew(SBorder)
+		.BorderImage(FEditorStyle::GetBrush("DetailsView.CategoryTop"))
+		.Padding(FMargin(0.0f, 3.0f, 1.0f, 0.0f))
+		[
+			SNew(SHorizontalBox)
+			+ SHorizontalBox::Slot()
+			.FillWidth(1.0f)
+			[
+				Viewport
+			]
+			+ SHorizontalBox::Slot()
+			.AutoWidth()
+			[
+				SNew(SVerticalBox)
+				+ SVerticalBox::Slot()
+				.Padding(2.0f)
+				[
+					PropertyView.ToSharedRef()
+				]
+				+ SVerticalBox::Slot()
+				.FillHeight(1.0f)
+				+ SVerticalBox::Slot()
+				.Padding(2.0f)
+				.AutoHeight()
+				[
+					SNew(SHorizontalBox)
+					+ SHorizontalBox::Slot()
+					.AutoWidth()
+					[
+						SNew(SButton)
+						.Text(LOCTEXT("PaperExtractSpritesExtractButton", "Extract..."))
+						.OnClicked(this, &SPaperExtractSpritesDialog::ExtractClicked)
+					]
+					+ SHorizontalBox::Slot()
+					.AutoWidth()
+					[
+						SNew(SButton)
+						.Text(LOCTEXT("PaperExtractSpritesCancelButton", "Cancel"))
+						.OnClicked(this, &SPaperExtractSpritesDialog::CancelClicked)
+					]
+				]
+			]
+		]
+	];
+}
+
+bool SPaperExtractSpritesDialog::ShowWindow(const FText& TitleText, UTexture2D* SourceTexture)
+{
+	// Create the window to pick the class
+	TSharedRef<SWindow> ExtractSpritesWindow = SNew(SWindow)
+		.Title(TitleText)
+		.SizingRule(ESizingRule::UserSized)
+		.ClientSize(FVector2D(1000.f, 700.f))
+		.AutoCenter(EAutoCenter::PreferredWorkArea)
+		.SupportsMinimize(false);
+	 
+	TSharedRef<SPaperExtractSpritesDialog> PaperExtractSpritesDialog = SNew(SPaperExtractSpritesDialog, SourceTexture);
+
+	ExtractSpritesWindow->SetContent(PaperExtractSpritesDialog);
+	TSharedPtr<SWindow> RootWindow = FGlobalTabmanager::Get()->GetRootWindow();
+	if (RootWindow.IsValid())
+	{
+		FSlateApplication::Get().AddWindowAsNativeChild(ExtractSpritesWindow, RootWindow.ToSharedRef());
+	}
+	else
+	{
+		FSlateApplication::Get().AddWindow(ExtractSpritesWindow);
+	}
+
+	return false;
+}
+
+// Extract the sprites based on SpriteExtractSettings->Mode
+void SPaperExtractSpritesDialog::PreviewExtractedSprites()
+{
+	FScopedSlowTask SlowTask(2, NSLOCTEXT("Paper2D", "Paper2D_ExtractingSprites", "Extracting Sprites"));
+	SlowTask.MakeDialog(false);
+	SlowTask.EnterProgressFrame();
+
+	const FString DefaultSuffix = TEXT("Sprite");
+	if (ExtractSpriteSettings->Mode == ESpriteExtractMode::Auto)
+	{
+		ExtractedSprites.Empty();
+
+		// First extract the rects from the texture
+		TArray<FIntRect> ExtractedRects;
+		UPaperSprite::ExtractRectsFromTexture(SourceTexture, /*out*/ ExtractedRects);
+
+		// Sort the rectangles by approximate row
+		struct FRectangleSortHelper
+		{
+			FRectangleSortHelper(TArray<FIntRect>& InOutSprites)
+			{
+				// Sort by Y, then by X (top left corner), descending order (so we can use it as a stack from the top row down)
+				TArray<FIntRect> SpritesLeft = InOutSprites;
+				SpritesLeft.Sort([](const FIntRect& A, const FIntRect& B) { return (A.Min.Y == B.Min.Y) ? (A.Min.X > B.Min.X) : (A.Min.Y > B.Min.Y); });
+				InOutSprites.Reset();
+
+				// Start pulling sprites out, the first one in each row will dominate remaining ones and cause them to get labeled
+				TArray<FIntRect> DominatedSprites;
+				DominatedSprites.Empty(SpritesLeft.Num());
+				while (SpritesLeft.Num())
+				{
+					FIntRect DominatingSprite = SpritesLeft.Pop();
+					DominatedSprites.Add(DominatingSprite);
+
+					// Find the sprites that are dominated (intersect the infinite horizontal band described by the dominating sprite)
+					for (int32 Index = 0; Index < SpritesLeft.Num();)
+					{
+						const FIntRect& CurElement = SpritesLeft[Index];
+						if ((CurElement.Min.Y <= DominatingSprite.Max.Y) && (CurElement.Max.Y >= DominatingSprite.Min.Y))
+						{
+							DominatedSprites.Add(CurElement);
+							SpritesLeft.RemoveAt(Index, /*Count=*/ 1, /*bAllowShrinking=*/ false);
+						}
+						else
+						{
+							++Index;
+						}
+					}
+
+					// Sort the sprites in the band by X and add them to the result
+					DominatedSprites.Sort([](const FIntRect& A, const FIntRect& B) { return (A.Min.X < B.Min.X); });
+					InOutSprites.Append(DominatedSprites);
+					DominatedSprites.Reset();
+				}
+			}
+		};
+		FRectangleSortHelper RectSorter(ExtractedRects);
+
+		int32 ExtractedRectIndex = 0;
+		for (FIntRect Rect : ExtractedRects)
+		{
+			FPaperExtractedSprite* Sprite = new(ExtractedSprites)FPaperExtractedSprite();
+			Sprite->Rect = Rect;
+			Sprite->Name = FString::Printf(TEXT("%s_%d"), *DefaultSuffix, ExtractedRectIndex);
+			ExtractedRectIndex++;
+		}
+	}
+	else
+	{
+		// Calculate rects
+		ExtractedSprites.Empty();
+		if (SourceTexture != nullptr)
+		{
+			int32 ExtractedRectIndex = 0;
+			int32 TextureWidth = SourceTexture->GetSizeX();
+			int32 TextureHeight = SourceTexture->GetSizeY();
+			for (int32 Y = ExtractSpriteSettings->Margin; Y + ExtractSpriteSettings->GridHeight <= TextureHeight; Y += ExtractSpriteSettings->GridHeight + ExtractSpriteSettings->Spacing)
+			{
+				for (int32 X = ExtractSpriteSettings->Margin; X + ExtractSpriteSettings->GridWidth <= TextureWidth; X += ExtractSpriteSettings->GridWidth + ExtractSpriteSettings->Spacing)
+				{
+					FPaperExtractedSprite* Sprite = new(ExtractedSprites)FPaperExtractedSprite();
+					Sprite->Rect = FIntRect(X, Y, X + ExtractSpriteSettings->GridWidth, Y + ExtractSpriteSettings->GridHeight);
+					Sprite->Name = FString::Printf(TEXT("%s_%d"), *DefaultSuffix, ExtractedRectIndex);
+					ExtractedRectIndex++;
+				}
+			}
+		}
+	}
+}
+
+void SPaperExtractSpritesDialog::OnFinishedChangingProperties(const FPropertyChangedEvent& PropertyChangedEvent)
+{
+	PreviewExtractedSprites();
+}
+
+FReply SPaperExtractSpritesDialog::ExtractClicked()
+{
+	CreateExtractedSprites();
+
+	CloseContainingWindow();
+	return FReply::Handled();
+}
+
+FReply SPaperExtractSpritesDialog::CancelClicked()
+{
+	CloseContainingWindow();
+	return FReply::Handled();
+}
+
+void SPaperExtractSpritesDialog::CloseContainingWindow()
+{
+	FWidgetPath WidgetPath;
+	TSharedPtr<SWindow> ContainingWindow = FSlateApplication::Get().FindWidgetWindow(AsShared(), WidgetPath);
+	if (ContainingWindow.IsValid())
+	{
+		ContainingWindow->RequestDestroyWindow();
+	}
+}
+
+void SPaperExtractSpritesDialog::CreateExtractedSprites()
+{
+	FAssetToolsModule& AssetToolsModule = FModuleManager::Get().LoadModuleChecked<FAssetToolsModule>("AssetTools");
+	FContentBrowserModule& ContentBrowserModule = FModuleManager::LoadModuleChecked<FContentBrowserModule>("ContentBrowser");
+
+	TArray<UObject*> ObjectsToSync;
+
+	// Create the factory used to generate the sprite
+	UPaperSpriteFactory* SpriteFactory = ConstructObject<UPaperSpriteFactory>(UPaperSpriteFactory::StaticClass());
+	SpriteFactory->InitialTexture = SourceTexture;
+
+	// Create the sprite
+	FString Name;
+	FString PackageName;
+
+	FScopedSlowTask Feedback(ExtractedSprites.Num(), NSLOCTEXT("Paper2D", "Paper2D_ExtractSpritesFromTexture", "Extracting Sprites From Texture"));
+	Feedback.MakeDialog(true);
+
+	for (const FPaperExtractedSprite ExtractedSprite : ExtractedSprites)
+	{
+		Feedback.EnterProgressFrame(1, NSLOCTEXT("Paper2D", "Paper2D_ExtractSpritesFromTexture", "Extracting Sprites From Texture"));
+
+		SpriteFactory->bUseSourceRegion = true;
+		const FIntRect& ExtractedRect = ExtractedSprite.Rect;
+		SpriteFactory->InitialSourceUV = FVector2D(ExtractedRect.Min.X, ExtractedRect.Min.Y);
+		SpriteFactory->InitialSourceDimension = FVector2D(ExtractedRect.Width(), ExtractedRect.Height());
+
+		// Get a unique name for the sprite
+		// Extracted sprite name is a name, we insert a _ as we're appending this to the texture name
+		// Opens up doors to renaming the sprites in the editor, and still ending up with TextureName_UserSpriteName
+		FString Suffix = TEXT("_");
+		Suffix.Append(ExtractedSprite.Name);
+
+		AssetToolsModule.Get().CreateUniqueAssetName(SourceTexture->GetOutermost()->GetName(), Suffix, /*out*/ PackageName, /*out*/ Name);
+		const FString PackagePath = FPackageName::GetLongPackagePath(PackageName);
+
+		if (UObject* NewAsset = AssetToolsModule.Get().CreateAsset(Name, PackagePath, UPaperSprite::StaticClass(), SpriteFactory))
+		{
+			ObjectsToSync.Add(NewAsset);
+		}
+
+		if (GWarn->ReceivedUserCancel())
+		{
+			break;
+		}
+	}
+
+	if (ObjectsToSync.Num() > 0)
+	{
+		ContentBrowserModule.Get().SyncBrowserToAssets(ObjectsToSync);
+	}
+}
+
+#undef LOCTEXT_NAMESPACE
+
diff --git a/Engine/Plugins/2D/Paper2D/Source/Paper2DEditor/Private/ExtractSprites/SPaperExtractSpritesDialog.h b/Engine/Plugins/2D/Paper2D/Source/Paper2DEditor/Private/ExtractSprites/SPaperExtractSpritesDialog.h
new file mode 100644
--- /dev/null
+++ b/Engine/Plugins/2D/Paper2D/Source/Paper2DEditor/Private/ExtractSprites/SPaperExtractSpritesDialog.h
@@ -0,0 +1,75 @@
+// Copyright 1998-2015 Epic Games, Inc. All Rights Reserved.
+
+#pragma once
+
+#include "SPaperEditorViewport.h"
+
+struct FPaperExtractedSprite
+{
+	FString Name;
+	FIntRect Rect;
+};
+
+//////////////////////////////////////////////////////////////////////////
+// SPaperExtractSpritesDialog
+
+class SPaperExtractSpritesDialog : public SCompoundWidget
+{
+public:
+	SLATE_BEGIN_ARGS(SPaperExtractSpritesDialog) {}
+	SLATE_END_ARGS()
+
+	// Constructs this widget with InArgs
+	void Construct(const FArguments& InArgs, UTexture2D* Texture);
+
+	// Show the dialog, returns true if successfully extracted sprites
+	static bool ShowWindow(const FText& TitleText, UTexture2D* SourceTexture);
+
+private:
+	// Handler for Extract button
+	FReply ExtractClicked();
+	
+	// Handler for Cancel button
+	FReply CancelClicked();
+	void OnFinishedChangingProperties(const FPropertyChangedEvent& PropertyChangedEvent);
+	void CloseContainingWindow();
+
+	// Calculate a list of extracted sprite rects
+	void PreviewExtractedSprites();
+
+	// Actually create extracted sprites
+	void CreateExtractedSprites();
+
+private:
+	// Source texture to extract from
+	UTexture2D* SourceTexture;
+
+	class UPaperExtractSpritesSettings* ExtractSpriteSettings;
+
+	TSharedPtr<class IDetailsView> PropertyView;
+	TArray<FPaperExtractedSprite> ExtractedSprites;
+};
+
+
+//////////////////////////////////////////////////////////////////////////
+// SPaperExtractSpritesViewport
+
+class SPaperExtractSpritesViewport : public SPaperEditorViewport
+{
+public:
+	SLATE_BEGIN_ARGS(SPaperExtractSpritesViewport) {}
+	SLATE_END_ARGS()
+
+	~SPaperExtractSpritesViewport();
+	void Construct(const FArguments& InArgs, UTexture2D* Texture, const TArray<FPaperExtractedSprite>& ExtractedSprites, const class UPaperExtractSpritesSettings* Settings, class SPaperExtractSpritesDialog* InDialog);
+
+protected:
+	// SPaperEditorViewport interface
+	virtual FText GetTitleText() const override;
+	// End of SPaperEditorViewport interface
+
+private:
+	TWeakObjectPtr<class UTexture2D> TexturePtr;
+	TSharedPtr<class FPaperExtractSpritesViewportClient> TypedViewportClient;
+	class SPaperExtractSpritesDialog* Dialog;
+};
