Paper2D: Fixed sprite thumbnail renderer to work with atlased sprites

diff --git a/Engine/Plugins/2D/Paper2D/Source/Paper2DEditor/Private/PaperSpriteThumbnailRenderer.cpp b/Engine/Plugins/2D/Paper2D/Source/Paper2DEditor/Private/PaperSpriteThumbnailRenderer.cpp
--- a/Engine/Plugins/2D/Paper2D/Source/Paper2DEditor/Private/PaperSpriteThumbnailRenderer.cpp
+++ b/Engine/Plugins/2D/Paper2D/Source/Paper2DEditor/Private/PaperSpriteThumbnailRenderer.cpp
@@ -43,7 +43,13 @@ void UPaperSpriteThumbnailRenderer::DrawGrid(int32 X, int32 Y, uint32 Width, uin
 
 void UPaperSpriteThumbnailRenderer::DrawFrame(class UPaperSprite* Sprite, int32 X, int32 Y, uint32 Width, uint32 Height, FRenderTarget*, FCanvas* Canvas, FBoxSphereBounds* OverrideRenderBounds)
 {
-	if (const UTexture2D* SourceTexture = (Sprite != nullptr) ? Sprite->GetSourceTexture() : nullptr)
+	const UTexture2D* SourceTexture = nullptr;
+	if (Sprite != nullptr)
+	{
+		SourceTexture = Sprite->GetBakedTexture() ? Sprite->GetBakedTexture() : Sprite->GetSourceTexture();
+	}
+
+	if (SourceTexture != nullptr)
 	{
 		const bool bUseTranslucentBlend = SourceTexture->HasAlphaChannel();
 
