Fixed build

diff --git a/Engine/Plugins/2D/Paper2D/Source/Paper2DEditor/Private/ExtractSprites/SPaperExtractSpritesDialog.cpp b/Engine/Plugins/2D/Paper2D/Source/Paper2DEditor/Private/ExtractSprites/SPaperExtractSpritesDialog.cpp
--- a/Engine/Plugins/2D/Paper2D/Source/Paper2DEditor/Private/ExtractSprites/SPaperExtractSpritesDialog.cpp
+++ b/Engine/Plugins/2D/Paper2D/Source/Paper2DEditor/Private/ExtractSprites/SPaperExtractSpritesDialog.cpp
@@ -132,8 +132,7 @@ void SPaperExtractSpritesDialog::Construct(const FArguments& InArgs, UTexture2D*
 	PreviewExtractedSprites();
 
 	FPropertyEditorModule& EditModule = FModuleManager::Get().GetModuleChecked<FPropertyEditorModule>("PropertyEditor");
-	FDetailsViewArgs DetailsViewArgs(/*bUpdateFromSelection=*/ false, /*bLockable=*/ false, /*bAllowSearch=*/ false, /*bObjectsUseNameArea=*/ true, /*bHideSelectionTip=*/ true);
-	DetailsViewArgs.bHideActorNameArea = true;
+	FDetailsViewArgs DetailsViewArgs(/*bUpdateFromSelection=*/ false, /*bLockable=*/ false, /*bAllowSearch=*/ false, /*InNameAreaSettings=*/ FDetailsViewArgs::HideNameArea, /*bHideSelectionTip=*/ true);
 	PropertyView = EditModule.CreateDetailView(DetailsViewArgs);
 	PropertyView->SetObject(ExtractSpriteSettings);
 	PropertyView->OnFinishedChangingProperties().AddSP(this, &SPaperExtractSpritesDialog::OnFinishedChangingProperties);
