Paper2D: Sprite atlasing improvements
- Exposed atlas texture options (texture compression, mip count, filtering)
- Update atlas when AtlasGroup is changed in UPaperSprite
- Basic rect aware weighted average mipmap generation
- Incremental updates, sprite positions in atlases are preserved as long as possible and prompt to rebuild everything wastage is significant
- Atlas texture reuse to avoid sprite updates where it can be helped

diff --git a/Engine/Plugins/2D/Paper2D/Source/Paper2D/Classes/PaperSprite.h b/Engine/Plugins/2D/Paper2D/Source/Paper2D/Classes/PaperSprite.h
--- a/Engine/Plugins/2D/Paper2D/Source/Paper2D/Classes/PaperSprite.h
+++ b/Engine/Plugins/2D/Paper2D/Source/Paper2D/Classes/PaperSprite.h
@@ -144,6 +144,10 @@ protected:
 	UPROPERTY(Category=Sprite, EditAnywhere, AssetRegistrySearchable)
 	class UPaperSpriteAtlas* AtlasGroup;
 
+	// The previous spritesheet group this belonged to
+	// To make sure we remove ourselves from it if changed or nulled out
+	TAssetPtr<class UPaperSpriteAtlas> PreviousAtlasGroup;
+
 #endif
 
 public:
@@ -233,6 +237,8 @@ public:
 	FVector2D GetSourceUV() const { return SourceUV; }
 	FVector2D GetSourceSize() const { return SourceDimension; }
 	UTexture2D* GetSourceTexture() const { return SourceTexture; }
+
+	const UPaperSpriteAtlas* GetAtlasGroup() const { return AtlasGroup; }
 #endif
 
 #if WITH_EDITOR
diff --git a/Engine/Plugins/2D/Paper2D/Source/Paper2D/Classes/PaperSpriteAtlas.h b/Engine/Plugins/2D/Paper2D/Source/Paper2D/Classes/PaperSpriteAtlas.h
--- a/Engine/Plugins/2D/Paper2D/Source/Paper2D/Classes/PaperSpriteAtlas.h
+++ b/Engine/Plugins/2D/Paper2D/Source/Paper2D/Classes/PaperSpriteAtlas.h
@@ -4,6 +4,45 @@
 
 #include "PaperSpriteAtlas.generated.h"
 
+USTRUCT()
+struct FPaperSpriteAtlasSlot
+{
+	GENERATED_USTRUCT_BODY()
+
+	UPROPERTY()
+	TAssetPtr<class UPaperSprite> SpriteRef;
+
+	UPROPERTY()
+	int32 AtlasIndex;
+
+	UPROPERTY()
+	int32 X;
+
+	UPROPERTY()
+	int32 Y;
+
+	UPROPERTY()
+	int32 Width;
+
+	UPROPERTY()
+	int32 Height;
+
+	bool IsValid() { return AtlasIndex >= 0; }
+};
+
+UENUM()
+enum class EPaperSpriteAtlasPadding : uint8
+{
+	// Dilate the texture to pad the atlas.
+	DilateBorder,
+	// Padding border filled with zeros.
+	PadWithZero,
+	//TODO: Tiled padding
+	//TODO: Sample outside bounds from source texture (for seamless reconstruction of a larger image)
+};
+
+
+
 // Groups together a set of sprites that will try to share the same texture atlas (allowing them to be combined into a single draw call)
 UCLASS(MinimalAPI, BlueprintType, Experimental, meta=(DisplayThumbnail = "true"))
 class UPaperSpriteAtlas : public UObject
@@ -23,15 +62,56 @@ class UPaperSpriteAtlas : public UObject
 	UPROPERTY(EditAnywhere, Category = General)
 	int32 MaxHeight;
 
+	// Maximum atlas page height (single pages might be smaller)
+	UPROPERTY(EditAnywhere, Category = General)
+	int32 MipCount;
+
+	// The type of padding performed on this atlas
+	UPROPERTY(EditAnywhere, Category = General)
+	EPaperSpriteAtlasPadding PaddingType;
+
+	// The number of pixels of padding
+	UPROPERTY(EditAnywhere, Category = General)
+	int32 Padding;
+
+	// Compression settings to use on atlas texture
+	UPROPERTY(EditAnywhere, Category=AtlasTexture)
+	TEnumAsByte<enum TextureCompressionSettings> CompressionSettings;
+
+	// Texture filtering mode when sampling these textures
+	UPROPERTY(EditAnywhere, Category=AtlasTexture)
+	TEnumAsByte<enum TextureFilter> Filter;
+
 	// List of generated atlas textures
-	UPROPERTY(VisibleAnywhere, Category=General)
+	UPROPERTY(VisibleAnywhere, Category=AtlasTexture)
 	TArray<UTexture*> GeneratedTextures;
 
 	// The GUID of the atlas group, used to match up sprites that belong to this group even thru atlas renames
 	UPROPERTY(VisibleAnywhere, Category=General)
 	FGuid AtlasGUID;
-#endif
 
+	// Slots in the atlas
+	UPROPERTY(EditAnywhere, Category=General)
+	bool bRebuildAtlas;
+
+	// Slots in the atlas
+	UPROPERTY()
+	TArray<FPaperSpriteAtlasSlot> AtlasSlots;
+
+	// Track the number of incremental builds
+	UPROPERTY()
+	int32 NumIncrementalBuilds;
+
+	// Values used when building this atlas
+	UPROPERTY()
+	int32 BuiltWidth;
+
+	UPROPERTY()
+	int32 BuiltHeight;
+
+	UPROPERTY()
+	int32 BuiltPadding;
+#endif
 
 
 public:
diff --git a/Engine/Plugins/2D/Paper2D/Source/Paper2D/Private/PaperSprite.cpp b/Engine/Plugins/2D/Paper2D/Source/Paper2D/Private/PaperSprite.cpp
--- a/Engine/Plugins/2D/Paper2D/Source/Paper2D/Private/PaperSprite.cpp
+++ b/Engine/Plugins/2D/Paper2D/Source/Paper2D/Private/PaperSprite.cpp
@@ -369,10 +369,17 @@ void UPaperSprite::OnObjectReimported(UObject* InObject)
 	// If SourceTetxureDimension == 0, we don't have a previous dimension to work off, so can't
 	// rescale sensibly
 	UTexture2D* Texture = Cast<UTexture2D>(InObject);
-	if (Texture != nullptr && Texture == GetSourceTexture() && NeedRescaleSpriteData())
+	if (Texture != nullptr && Texture == GetSourceTexture())
 	{
-		RescaleSpriteData(GetSourceTexture());
-		PostEditChange();
+		if (NeedRescaleSpriteData())
+		{
+			RescaleSpriteData(GetSourceTexture());
+			PostEditChange();
+		}
+		else if (AtlasGroup != nullptr)
+		{
+			AtlasGroup->PostEditChange();
+		}
 	}
 }
 
@@ -519,6 +526,42 @@ void UPaperSprite::PostEditChangeProperty(FPropertyChangedEvent& PropertyChanged
 		ValidateSocketNames();
 	}
 
+	if (PropertyName == GET_MEMBER_NAME_CHECKED(UPaperSprite, AtlasGroup))
+	{
+		auto SpriteAssetPtr = PreviousAtlasGroup;
+		FStringAssetReference AtlasGroupStringRef = SpriteAssetPtr.ToStringReference();
+		UPaperSpriteAtlas* PreviousAtlasGroupPtr = nullptr;
+		if (!AtlasGroupStringRef.ToString().IsEmpty())
+		{
+			PreviousAtlasGroupPtr = Cast<UPaperSpriteAtlas>(StaticLoadObject(UPaperSpriteAtlas::StaticClass(), nullptr, *AtlasGroupStringRef.ToString(), nullptr, LOAD_None, nullptr));
+		}
+		
+		if (PreviousAtlasGroupPtr != AtlasGroup)
+		{
+			// Update previous
+			if (PreviousAtlasGroupPtr != nullptr)
+			{
+				PreviousAtlasGroupPtr->PostEditChange();
+			}
+
+			// Update cached previous atlas group
+			PreviousAtlasGroup = AtlasGroup;
+
+			// Rebuild atlas group
+			if (AtlasGroup != nullptr)
+			{
+				AtlasGroup->PostEditChange();
+			}
+			else
+			{
+				BakedSourceTexture = nullptr;
+				BakedSourceUV = FVector2D(0, 0);
+				bRenderDataModified = true;
+			}
+
+		}
+	}
+
 	// The texture dimensions have changed
 	//if (NeedRescaleSpriteData())
 	//{
diff --git a/Engine/Plugins/2D/Paper2D/Source/Paper2D/Private/PaperSpriteAtlas.cpp b/Engine/Plugins/2D/Paper2D/Source/Paper2D/Private/PaperSpriteAtlas.cpp
--- a/Engine/Plugins/2D/Paper2D/Source/Paper2D/Private/PaperSpriteAtlas.cpp
+++ b/Engine/Plugins/2D/Paper2D/Source/Paper2D/Private/PaperSpriteAtlas.cpp
@@ -10,6 +10,12 @@ UPaperSpriteAtlas::UPaperSpriteAtlas(const FObjectInitializer& ObjectInitializer
 #if WITH_EDITORONLY_DATA
 	, MaxWidth(2048)
 	, MaxHeight(2048)
+	, MipCount(1)
+	, PaddingType(EPaperSpriteAtlasPadding::DilateBorder)
+	, Padding(1)
+	, CompressionSettings(TextureCompressionSettings::TC_Default)
+	, Filter(TextureFilter::TF_Bilinear)
+	, bRebuildAtlas(false)
 #endif
 {
 }
@@ -24,6 +30,7 @@ void UPaperSpriteAtlas::GetAssetRegistryTags(TArray<FAssetRegistryTag>& OutTags)
 }
 
 #if WITH_EDITORONLY_DATA
+
 void UPaperSpriteAtlas::PostDuplicate(bool bDuplicateForPIE)
 {
 	Super::PostDuplicate(bDuplicateForPIE);
@@ -35,4 +42,5 @@ void UPaperSpriteAtlas::PostInitProperties()
 	Super::PostInitProperties();
 	FPlatformMisc::CreateGuid(AtlasGUID);
 }
-#endif
\ No newline at end of file
+
+#endif
diff --git a/Engine/Plugins/2D/Paper2D/Source/Paper2DEditor/Private/Atlasing/PaperAtlasGenerator.cpp b/Engine/Plugins/2D/Paper2D/Source/Paper2DEditor/Private/Atlasing/PaperAtlasGenerator.cpp
--- a/Engine/Plugins/2D/Paper2D/Source/Paper2DEditor/Private/Atlasing/PaperAtlasGenerator.cpp
+++ b/Engine/Plugins/2D/Paper2D/Source/Paper2DEditor/Private/Atlasing/PaperAtlasGenerator.cpp
@@ -9,184 +9,257 @@
 //////////////////////////////////////////////////////////////////////////
 // 
 
-void CopyTextureData(const uint8* Source, uint8* Dest, uint32 SizeX, uint32 SizeY, uint32 BytesPerPixel, uint32 SourceStride, uint32 DestStride)
+#include "PaperAtlasTextureHelpers.h"
+#include "PaperAtlasHelpers.h"
+
+//////////////////////////////////////////////////////////////////////////
+// FPaperAtlasGenerator
+
+void FPaperAtlasGenerator::HandleAssetChangedEvent(UPaperSpriteAtlas* Atlas)
 {
-	const uint32 NumBytesPerRow = SizeX * BytesPerPixel;
+	const int32 MaxAtlasDimension = 4096;
+	static_assert(MaxAtlasDimension < 16384, "PaperAtlasGenerator MaxAtlasDimension exceeds multiplier in sort key");
+	Atlas->MaxWidth = FMath::Clamp(Atlas->MaxWidth, 16, MaxAtlasDimension);
+	Atlas->MaxHeight = FMath::Clamp(Atlas->MaxHeight, 16, MaxAtlasDimension);
+	Atlas->MipCount = ClampMips(Atlas->MaxWidth, Atlas->MaxHeight, Atlas->MipCount);
 
-	for (uint32 Y = 0; Y < SizeY; ++Y)
+	bool bTestForAtlasImprovement = true;
+
+	// Have the atlas settings changed? This will trigger a full rebuild
+	bool bAtlasDimensionsChanged = Atlas->MaxWidth != Atlas->BuiltWidth || Atlas->MaxHeight != Atlas->BuiltHeight;
+	if (bAtlasDimensionsChanged || Atlas->Padding != Atlas->BuiltPadding)
 	{
-		FMemory::Memcpy(
-			Dest + (DestStride * Y),
-			Source + (SourceStride * Y),
-			NumBytesPerRow
-			);
+		Atlas->bRebuildAtlas = true;
 	}
-}
 
-struct FSingleTexturePaperAtlas : public FSlateTextureAtlas
-{
-public:
-	FSingleTexturePaperAtlas(uint32 InWidth, uint32 InHeight)
-	: FSlateTextureAtlas(InWidth, InHeight, sizeof(FColor), ESlateTextureAtlasPaddingStyle::NoPadding)
+	// Save the settings this atlas is being built with
+	Atlas->BuiltWidth = Atlas->MaxWidth;
+	Atlas->BuiltHeight = Atlas->MaxHeight;
+	Atlas->BuiltPadding = Atlas->Padding;
+
+	// Force rebuild an atlas by deleting history
+	if (Atlas->bRebuildAtlas)
 	{
-		//@TODO: check(StrideBytes >= 4*Width);
+		Atlas->AtlasSlots.Empty();
+		Atlas->NumIncrementalBuilds = 0;
+		Atlas->bRebuildAtlas = false;
+	}
+	else
+	{
+		// Keep track of incremental builds
+		Atlas->NumIncrementalBuilds++;
 	}
 
-	virtual void ConditionalUpdateTexture() override {}
+	// Load all sprites that were used in building this atlas
+	TArray<UPaperSprite*> SpritesInPreviousAtlas;
+	LoadAllReferencedSprites(Atlas, SpritesInPreviousAtlas);
 
-	void CopyToTexture(UTexture2D* Texture)
-	{
-		const int32 NumMips = 1;
+	// Load all sprites that currently reference this atlas
+	TArray<UPaperSprite*> SpritesInNewAtlas;
+	LoadAllSpritesWithAtlasGroupGUID(Atlas, SpritesInNewAtlas);
 
-		Texture->Source.Init(AtlasWidth, AtlasHeight, /*NumSlices=*/ 1, NumMips, ETextureSourceFormat::TSF_BGRA8, AtlasData.GetData());
-		Texture->UpdateResource();
+	// Find sprites removed from this atlas, but not null (i.e. deliberately removed from the atlas)
+	bool bWasTextureRemoved = false;
+	for (UPaperSprite* OriginalSprite : SpritesInPreviousAtlas)
+	{
+		if (!SpritesInNewAtlas.Contains(OriginalSprite))
+		{
+			RemoveTextureSlotWithSprite(Atlas, OriginalSprite);
+			bWasTextureRemoved = true;
+		}
+	}
+	if (bWasTextureRemoved)
+	{
+		MergeAdjacentRects(Atlas);
 	}
 
-	const FAtlasedTextureSlot* AddSprite(UPaperSprite* Sprite)
+	// Sort new sprites by size
+	struct Local
+	{
+		static int32 SpriteSortValue(const UPaperSprite& Sprite)
+		{
+			const FVector2D SpriteSizeFloat = Sprite.GetSourceSize();
+			const FIntPoint SpriteSize(FMath::TruncToInt(SpriteSizeFloat.X), FMath::TruncToInt(SpriteSizeFloat.Y));
+			return SpriteSize.X * 16384 + SpriteSize.Y; // Sort wider textures first
+		}
+	};
+	SpritesInNewAtlas.Sort( [](UPaperSprite& A, UPaperSprite& B) { return Local::SpriteSortValue(A) > Local::SpriteSortValue(B); } );
+
+	// Add new sprites
+	TArray<FPaperSpriteAtlasSlot> ImprovementTestAtlas; // A second atlas to compare wastage
+	for (UPaperSprite* Sprite : SpritesInNewAtlas)
 	{
 		const FVector2D SpriteSizeFloat = Sprite->GetSourceSize();
 		const FIntPoint SpriteSize(FMath::TruncToInt(SpriteSizeFloat.X), FMath::TruncToInt(SpriteSizeFloat.Y));
+		const FIntPoint PaddedSpriteSize(SpriteSize.X + Atlas->Padding * 2, SpriteSize.Y + Atlas->Padding * 2);
 
-		TArray<uint8> DummyBuffer;
-		DummyBuffer.AddZeroed(SpriteSize.X * SpriteSize.Y * BytesPerPixel);
-		
-		check(Sprite->GetSourceTexture());
-		FTextureSource& SourceData = Sprite->GetSourceTexture()->Source;
-
-		//@TODO: Handle different texture formats!
-		if (SourceData.GetFormat() == TSF_BGRA8)
+		if (Sprite->GetSourceTexture() == nullptr)
 		{
-			uint32 BytesPerPixel = SourceData.GetBytesPerPixel();
-			uint8* OffsetSource = SourceData.LockMip(0) + (FMath::TruncToInt(Sprite->GetSourceUV().X) + FMath::TruncToInt(Sprite->GetSourceUV().Y) * SourceData.GetSizeX()) * BytesPerPixel;
-			uint8* OffsetDest = DummyBuffer.GetData();
-
-			CopyTextureData(OffsetSource, OffsetDest, SpriteSize.X, SpriteSize.Y, BytesPerPixel, SourceData.GetSizeX() * BytesPerPixel, SpriteSize.X * BytesPerPixel);
-
-			SourceData.UnlockMip(0);
+			UE_LOG(LogPaper2DEditor, Error, TEXT("Sprite %s has no source texture and cannot be packed"), *(Sprite->GetPathName()));
+			continue;
 		}
-		else
+
+		//TODO: Padding should only be considered by the slot finder to allow atlasing
+		// textures flush to the edge
+		if ((PaddedSpriteSize.X > Atlas->MaxWidth) || (PaddedSpriteSize.Y > Atlas->MaxHeight))
 		{
-			UE_LOG(LogPaper2DEditor, Error, TEXT("Sprite %s is not BGRA8, which isn't supported in atlases yet"), *(Sprite->GetPathName()));
+			// This sprite cannot ever fit into an atlas page
+			UE_LOG(LogPaper2DEditor, Error, TEXT("Sprite %s (%d x %d) can never fit into atlas %s (%d x %d) due to maximum page size restrictions"),
+				*(Sprite->GetPathName()),
+				SpriteSize.X,
+				SpriteSize.Y,
+				*(Atlas->GetPathName()),
+				Atlas->MaxWidth,
+				Atlas->MaxHeight);
+			continue;
 		}
 
-		const FAtlasedTextureSlot* Slot = AddTexture(SpriteSize.X, SpriteSize.Y, DummyBuffer);
-		if (Slot != nullptr)
+		bool bSpriteChanged = false;
+		FindBestSlotForTexture(Atlas->AtlasSlots, Atlas->MaxWidth, Atlas->MaxHeight, Sprite, PaddedSpriteSize.X, PaddedSpriteSize.Y, bSpriteChanged);
+		if (bSpriteChanged)
 		{
-			SpriteToSlotMap.Add(Sprite, Slot);
+			//TODO: keep track of sprite moving about in the atlas?
 		}
 
-		return Slot;
+		if (bTestForAtlasImprovement)
+		{
+			// Pack into a second test atlas in parallel
+			bool bUnused = false;
+			FindBestSlotForTexture(ImprovementTestAtlas, Atlas->MaxWidth, Atlas->MaxHeight, Sprite, PaddedSpriteSize.X, PaddedSpriteSize.Y, bSpriteChanged);
+		}
 	}
 
-public:
-	TMap<UPaperSprite*, const FAtlasedTextureSlot*> SpriteToSlotMap;
-};
-
-//////////////////////////////////////////////////////////////////////////
-// FPaperAtlasGenerator
-
-void FPaperAtlasGenerator::HandleAssetChangedEvent(UPaperSpriteAtlas* Atlas)
-{
-	Atlas->MaxWidth = FMath::Clamp(Atlas->MaxWidth, 16, 4096);
-	Atlas->MaxHeight = FMath::Clamp(Atlas->MaxHeight, 16, 4096);
-
-	// Find all sprites that reference the atlas and force them loaded
-	FAssetRegistryModule& AssetRegistryModule = FModuleManager::LoadModuleChecked<FAssetRegistryModule>(TEXT("AssetRegistry"));
-
-	TArray<FAssetData> AssetList;
-	TMultiMap<FName, FString> TagsAndValues;
-	TagsAndValues.Add(TEXT("AtlasGroupGUID"), Atlas->AtlasGUID.ToString(EGuidFormats::Digits));
-	AssetRegistryModule.Get().GetAssetsByTagValues(TagsAndValues, AssetList);
-
-	TIndirectArray<FSingleTexturePaperAtlas> AtlasGenerators;
-
-	// Start off with one page
-	new (AtlasGenerators) FSingleTexturePaperAtlas(Atlas->MaxWidth, Atlas->MaxHeight);
-
-	for (const FAssetData& SpriteRef : AssetList)
+	// Test for improvement if necessary
+	// An "improvement" is defined as less atlases overall, but could be extended to check for atlas area once we support resizing atlases
+	if (bTestForAtlasImprovement)
 	{
-		if (UPaperSprite* Sprite = Cast<UPaperSprite>(SpriteRef.GetAsset()))
+		if (NumTexturesUsedInAtlasSlots(ImprovementTestAtlas) < NumTexturesUsedInAtlasSlots(Atlas->AtlasSlots))
 		{
-			//@TODO: Use the tight bounds instead of the source bounds
-			const FVector2D SpriteSizeFloat = Sprite->GetSourceSize();
-			const FIntPoint SpriteSize(FMath::TruncToInt(SpriteSizeFloat.X), FMath::TruncToInt(SpriteSizeFloat.Y));
-
-			if (Sprite->GetSourceTexture() == nullptr)
+			const EAppReturnType::Type Choice = FMessageDialog::Open(EAppMsgType::YesNo, NSLOCTEXT("PaperEditor", "AtlasPackingImprovement", "Atlas packing can be improved significantly by repacking the entire atlas. This will require re-saving most or all sprites in this atlas.\nDo you want to do this now?"));
+			if (Choice == EAppReturnType::Yes)
 			{
-				UE_LOG(LogPaper2DEditor, Error, TEXT("Sprite %s has no source texture and cannot be packed"), *(Sprite->GetPathName()));
-				continue;
+				// Likely to mark most sprites dirty
+				Atlas->AtlasSlots = ImprovementTestAtlas;
 			}
+		}
+	}
 
-			//@TODO: Take padding into account (push this into the generator)
-			if ((SpriteSize.X > Atlas->MaxWidth) || (SpriteSize.Y >= Atlas->MaxHeight))
+	// Update atlas textures
+	TArray<UTexture*> RemappedAtlasTextures; // Will only contain valid and used textures after this
+	TArray<bool> RemappedAtlasForceDirty; // If any atlases were missing (due to user deleting bits), all the dependent sprites are considered dirty
+	TArray<int32> AtlasLookupIndex; // To correct mismatched atlas numbers, gaps in arrays, etc
+	for (FPaperSpriteAtlasSlot& Slot : Atlas->AtlasSlots)
+	{
+		// Get the "correct" atlas index
+		while (!AtlasLookupIndex.IsValidIndex(Slot.AtlasIndex))
+		{
+			AtlasLookupIndex.Add(-1); // add an uninitialized lookup entry
+			RemappedAtlasForceDirty.Add(false);
+		}
+
+		if (AtlasLookupIndex[Slot.AtlasIndex] == -1)
+		{
+			AtlasLookupIndex[Slot.AtlasIndex] = RemappedAtlasTextures.Num();
+			if (Atlas->GeneratedTextures.IsValidIndex(Slot.AtlasIndex) && Atlas->GeneratedTextures[Slot.AtlasIndex] != nullptr)
 			{
-				// This sprite cannot ever fit into an atlas page
-				UE_LOG(LogPaper2DEditor, Error, TEXT("Sprite %s (%d x %d) can never fit into atlas %s (%d x %d) due to maximum page size restrictions"),
-					*(Sprite->GetPathName()),
-					SpriteSize.X,
-					SpriteSize.Y,
-					*(Atlas->GetPathName()),
-					Atlas->MaxWidth,
-					Atlas->MaxHeight);
+				RemappedAtlasTextures.Add(Atlas->GeneratedTextures[Slot.AtlasIndex]);
+				RemappedAtlasForceDirty.Add(false);
 			}
 			else
 			{
-				const FAtlasedTextureSlot* Slot = nullptr;
-
-				// Does it fit in any existing pages?
-				for (auto& Generator : AtlasGenerators)
-				{
-					Slot = Generator.AddSprite(Sprite);
-
-					if (Slot != nullptr)
-					{
-						break;
-					}
-				}
-
-				if (Slot == nullptr)
-				{
-					// Doesn't fit in any current pages, make a new one
-					FSingleTexturePaperAtlas* NewGenerator = new (AtlasGenerators) FSingleTexturePaperAtlas(Atlas->MaxWidth, Atlas->MaxHeight);
-					Slot = NewGenerator->AddSprite(Sprite);
-				}
-
-				if (Slot != nullptr)
-				{
-					UE_LOG(LogPaper2DEditor, Warning, TEXT("Allocated %s to (%d,%d)"), *(Sprite->GetPathName()), Slot->X, Slot->Y);
-				}
-				else
-				{
-					// ERROR: Didn't fit into a brand new page, maybe padding was wrong
-					UE_LOG(LogPaper2DEditor, Error, TEXT("Failed to allocate %s into a brand new page"), *(Sprite->GetPathName()));
-				}
+				// The texture never existed - all sprites referencing this MUST be dirty and MUST be updated
+				RemappedAtlasTextures.Add(NewNamedObject<UTexture2D>(Atlas, NAME_None, RF_Public));
+				RemappedAtlasForceDirty.Add(true);
 			}
 		}
+
+		// Now atlas index refers to the RemappedAtlasTexures
+		Slot.AtlasIndex = AtlasLookupIndex[Slot.AtlasIndex];
 	}
 
-	// Turn the generators back into textures
-	Atlas->GeneratedTextures.Empty(AtlasGenerators.Num());
-	for (int32 GeneratorIndex = 0; GeneratorIndex < AtlasGenerators.Num(); ++GeneratorIndex)
+	// Now fill the atlases and update sprite data where needed
+	for (int32 AtlasIndex = 0; AtlasIndex < RemappedAtlasTextures.Num(); ++AtlasIndex)
 	{
-		FSingleTexturePaperAtlas& AtlasPage = AtlasGenerators[GeneratorIndex];
+		UTexture2D* AtlasTexture = Cast<UTexture2D>(RemappedAtlasTextures[AtlasIndex]);
+		check(AtlasTexture); // this should not be null
+
+		// An atlas is ALSO forced dirty if the dimensions have changed
+		// We're just grabbing a fixed atlas size here for now
+		int32 AtlasWidth = Atlas->MaxWidth;
+		int32 AtlasHeight = Atlas->MaxHeight;
+		int32 BytesPerPixel = sizeof(FColor);
+
+		bool bAtlasDirty = RemappedAtlasForceDirty[AtlasIndex] || AtlasTexture->GetSizeX() != AtlasWidth || AtlasTexture->GetSizeY() != AtlasHeight;
+
+		// Propagate texture settings
+		AtlasTexture->CompressionSettings = Atlas->CompressionSettings;
+		AtlasTexture->Filter = Atlas->Filter;
+		AtlasTexture->AddressX = TextureAddress::TA_Clamp;
+		AtlasTexture->AddressY = TextureAddress::TA_Clamp;
+		AtlasTexture->MipGenSettings = TextureMipGenSettings::TMGS_LeaveExistingMips;
+
+		// Allocate enough space for all mips
+		TArray<uint8> AtlasTextureData;
+		int32 CurrentAtlasWidth = AtlasWidth;
+		int32 CurrentAtlasHeight = AtlasHeight;
+		int32 TotalPixels = 0;
+		for (int32 MipIndex = 0; MipIndex < Atlas->MipCount; ++MipIndex)
+		{
+			TotalPixels += CurrentAtlasHeight * CurrentAtlasWidth;
+		}
+		AtlasTextureData.AddZeroed(TotalPixels * BytesPerPixel);
 
-		UTexture2D* Texture = NewNamedObject<UTexture2D>(Atlas, NAME_None, RF_Public);
-		Atlas->GeneratedTextures.Add(Texture);
+		// Atlas sprites
+		TArray<FPaperSpriteAtlasSlot> SlotsForAtlas;
+		for (FPaperSpriteAtlasSlot& Slot : Atlas->AtlasSlots)
+		{
+			UPaperSprite* SpriteBeingBuilt = Slot.SpriteRef.Get();
+
+			// Only for the sprites in this atlas
+			if (Slot.AtlasIndex == AtlasIndex && SpriteBeingBuilt != nullptr)
+			{
+				SlotsForAtlas.Add(Slot);
+				CopySpriteToAtlasTextureData(AtlasTextureData, AtlasWidth, AtlasHeight, BytesPerPixel, Atlas->PaddingType, Atlas->Padding, SpriteBeingBuilt, Slot);
+			}
+		}
+
+		// Mipmaps
+		if (Atlas->MipCount > 1)
+		{
+			GenerateMipChainARGB(SlotsForAtlas, AtlasTextureData, Atlas->MipCount, AtlasWidth, AtlasHeight);
+		}
 
-		AtlasPage.CopyToTexture(Texture);
 
-		// Now update the baked data for all the sprites to point there
-		for (auto& SpriteToSlot : AtlasPage.SpriteToSlotMap)
+		AtlasTexture->Source.Init(AtlasWidth, AtlasHeight, /*NumSlices=*/ 1, Atlas->MipCount, ETextureSourceFormat::TSF_BGRA8, AtlasTextureData.GetData());
+		AtlasTexture->UpdateResource();
+		AtlasTexture->PostEditChange();
+	}
+
+	// Rebuild sprites that have changed position in the atlas
+	for (FPaperSpriteAtlasSlot& Slot : Atlas->AtlasSlots)
+	{
+		UPaperSprite* SpriteBeingBuilt = Slot.SpriteRef.Get();
+		if (SpriteBeingBuilt)
 		{
-			UPaperSprite* Sprite = SpriteToSlot.Key;
-			Sprite->Modify();
-			const FAtlasedTextureSlot* Slot = SpriteToSlot.Value;
-			Sprite->BakedSourceTexture = Texture;
-			Sprite->BakedSourceUV = FVector2D(Slot->X, Slot->Y);
-			Sprite->RebuildRenderData();
+			UTexture2D* BakedSourceTexture = Cast<UTexture2D>(RemappedAtlasTextures[Slot.AtlasIndex]);
+			FVector2D BakedSourceUV = FVector2D(Slot.X + Atlas->Padding, Slot.Y + Atlas->Padding);
+
+			if (SpriteBeingBuilt->BakedSourceTexture != BakedSourceTexture || SpriteBeingBuilt->BakedSourceUV != BakedSourceUV || bAtlasDimensionsChanged)
+			{
+				SpriteBeingBuilt->Modify();
+				SpriteBeingBuilt->BakedSourceTexture = Cast<UTexture2D>(RemappedAtlasTextures[Slot.AtlasIndex]);
+				SpriteBeingBuilt->BakedSourceUV = FVector2D(Slot.X + Atlas->Padding, Slot.Y + Atlas->Padding);
+				SpriteBeingBuilt->RebuildRenderData();
+
+				// Propagate changes to sprites in scene
+				SpriteBeingBuilt->PostEditChange();
+			}
 		}
 	}
 
-	//@TODO: Adjust the atlas rebuild code so that it tries to preserve existing sprites (optimize for the 'just added one to the end' case, preventing dirtying lots of assets unnecessarily)
-	//@TODO: invoke this code when a sprite has the atlas group property modified
+	// Finalize changes
+	Atlas->GeneratedTextures = RemappedAtlasTextures;
+
+	// Dirty
+	Atlas->MarkPackageDirty();
 }
diff --git a/Engine/Plugins/2D/Paper2D/Source/Paper2DEditor/Private/Atlasing/PaperAtlasHelpers.h b/Engine/Plugins/2D/Paper2D/Source/Paper2DEditor/Private/Atlasing/PaperAtlasHelpers.h
new file mode 100644
--- /dev/null
+++ b/Engine/Plugins/2D/Paper2D/Source/Paper2DEditor/Private/Atlasing/PaperAtlasHelpers.h
@@ -0,0 +1,210 @@
+// Copyright 1998-2015 Epic Games, Inc. All Rights Reserved.
+
+#pragma once
+
+// Loads and returns a list of sprites referenced in the atlas group (i.e. sprites used in the previous build)
+static void LoadAllReferencedSprites(UPaperSpriteAtlas* AtlasGroup, TArray<UPaperSprite*>& AllSprites)
+{
+	for (FPaperSpriteAtlasSlot& Slot : AtlasGroup->AtlasSlots)
+	{
+		FStringAssetReference SpriteStringRef = Slot.SpriteRef.ToStringReference();
+		if (!SpriteStringRef.ToString().IsEmpty())
+		{
+			UPaperSprite* Sprite = Cast<UPaperSprite>(StaticLoadObject(UPaperSprite::StaticClass(), nullptr, *SpriteStringRef.ToString(), nullptr, LOAD_None, nullptr));
+			if (Sprite != nullptr)
+			{
+				AllSprites.Add(Sprite);
+			}
+		}
+	}
+}
+
+// Find all sprites that reference the atlas and force them loaded
+static void LoadAllSpritesWithAtlasGroupGUID(UPaperSpriteAtlas* Atlas, TArray<UPaperSprite*>& SpritesInNewAtlas)
+{
+	FAssetRegistryModule& AssetRegistryModule = FModuleManager::LoadModuleChecked<FAssetRegistryModule>(TEXT("AssetRegistry"));
+
+	TArray<FAssetData> AssetList;
+	TMultiMap<FName, FString> TagsAndValues;
+	TagsAndValues.Add(TEXT("AtlasGroupGUID"), Atlas->AtlasGUID.ToString(EGuidFormats::Digits));
+	AssetRegistryModule.Get().GetAssetsByTagValues(TagsAndValues, AssetList);
+
+	for (const FAssetData& SpriteRef : AssetList)
+	{
+		if (UPaperSprite* Sprite = Cast<UPaperSprite>(SpriteRef.GetAsset()))
+		{
+			if (Sprite->GetAtlasGroup() == Atlas)
+			{
+				SpritesInNewAtlas.Add(Sprite);
+			}
+		}
+	}
+}
+
+// Create an empty slot and return the slot index
+// Returns -1 if the slot can't be created for any reason (eg. width and height too small)
+static int32 CreateEmptySlot(TArray<FPaperSpriteAtlasSlot>& Slots, int32 AtlasIndex, int32 X, int32 Y, int32 Width, int32 Height)
+{
+	if (Width > 0 && Height > 0)
+	{
+		FPaperSpriteAtlasSlot* NewSlot = new(Slots)FPaperSpriteAtlasSlot();
+		NewSlot->AtlasIndex = AtlasIndex;
+		NewSlot->X = X;
+		NewSlot->Y = Y;
+		NewSlot->Width = Width;
+		NewSlot->Height = Height;
+		return Slots.Num() - 1;
+	}
+	else
+	{
+		return -1;
+	}
+}
+
+// Insert a sprite into a specified
+// The sprite must be known to fit in the provided slot
+static void InsertSpriteIntoSlot(TArray<FPaperSpriteAtlasSlot>& Slots, int SlotIndexToInsertInto, UPaperSprite* Sprite, int32 Width, int32 Height)
+{
+	FPaperSpriteAtlasSlot& Slot = Slots[SlotIndexToInsertInto];
+	int32 RemainingWidth = Slot.Width - Width;
+	int32 RemainingHeight = Slot.Height - Height;
+	int32 OriginalSlotWidth = Slot.Width;
+	int32 OriginalSlotHeight = Slot.Height;
+	Slot.SpriteRef = Sprite;
+	Slot.Width = Width;
+	Slot.Height = Height;
+	check(RemainingWidth >= 0 && RemainingHeight >= 0);
+
+	if (RemainingHeight <= RemainingWidth)
+	{
+		CreateEmptySlot(Slots, Slot.AtlasIndex, Slot.X, Slot.Y + Height, Width, RemainingHeight);
+		CreateEmptySlot(Slots, Slot.AtlasIndex, Slot.X + Width, Slot.Y, RemainingWidth, OriginalSlotHeight);
+	}
+	else
+	{
+		CreateEmptySlot(Slots, Slot.AtlasIndex, Slot.X, Slot.Y + Height, OriginalSlotWidth, RemainingHeight);
+		CreateEmptySlot(Slots, Slot.AtlasIndex, Slot.X + Width, Slot.Y, RemainingWidth, Height);
+	}
+}
+
+static void MergeAdjacentRects(UPaperSpriteAtlas* AtlasGroup)
+{
+	//@TODO: run through and merge
+}
+
+static void RemoveTextureSlotWithSprite(UPaperSpriteAtlas* AtlasGroup, UPaperSprite* Sprite)
+{
+	TArray<FPaperSpriteAtlasSlot>& AtlasSlots = AtlasGroup->AtlasSlots;
+	for (int32 SlotIndex = 0; SlotIndex < AtlasSlots.Num(); ++SlotIndex)
+	{
+		FPaperSpriteAtlasSlot& AtlasSlot = AtlasSlots[SlotIndex];
+		if (AtlasSlot.SpriteRef.Get() == Sprite)
+		{
+			AtlasSlot.SpriteRef = (UPaperSprite*)nullptr;
+		}
+	}
+}
+
+// Calculate the number of textures used in the atlas slots (atlas indices could be sparse) 
+static int NumTexturesUsedInAtlasSlots(const TArray<FPaperSpriteAtlasSlot>& AtlasSlots)
+{
+	TSet<int> TextureIndex;
+	for (const FPaperSpriteAtlasSlot& Slot : AtlasSlots)
+	{
+		TextureIndex.Add(Slot.AtlasIndex);
+	}
+
+	return TextureIndex.Num();
+}
+
+// FindBestSlotForTexture
+// NOTE: Assumes all UPaperSprites have been loaded
+static const FPaperSpriteAtlasSlot& FindBestSlotForTexture(TArray<FPaperSpriteAtlasSlot>& AtlasSlots, int32 AtlasWidth, int32 AtlasHeight, UPaperSprite* Sprite, int32 Width, int32 Height, bool& bOutSpriteChanged)
+{
+	int32 BestMatchIndex = -1;
+	int32 BestMatchScore = MAX_int32; // Lower score = better match
+	bool bImpossibleToFit = (Width > AtlasWidth || Height > AtlasHeight); // If impossible to fit, still continue to free up a slot if it already exists
+	bOutSpriteChanged = false;
+
+	// Asset pointer for the sprite we're currently trying to pack
+	TAssetPtr<UPaperSprite> SpriteAssetRef = Sprite;
+
+	// 1. Find matching texture and see if we can fit it in the existing slot
+	for (int32 SlotIndex = 0; SlotIndex < AtlasSlots.Num(); ++SlotIndex)
+	{
+		FPaperSpriteAtlasSlot& AtlasSlot = AtlasSlots[SlotIndex];
+		if (AtlasSlot.SpriteRef == SpriteAssetRef)
+		{
+			if (Width <= AtlasSlot.Width && Height <= AtlasSlot.Height && !bImpossibleToFit)
+			{
+				// This is the best slot for an incremental update
+				// The slot width and height should not be changed in case it grows to its original size later
+				BestMatchIndex = SlotIndex;
+				BestMatchScore = -1.0f;
+
+				// The sprite UVs need updating if the actual dimensions have changed
+				bOutSpriteChanged = Width != AtlasSlot.Width || Height != AtlasSlot.Height;
+				break;
+			}
+			else
+			{
+				// The sprite used to be in this slot, but its too big for the slot now
+				// Free the slot to be used later
+				AtlasSlot.SpriteRef = (UPaperSprite*)nullptr;
+				break;
+			}
+		}
+	}
+
+	// 2. Not found, try to fit in any empty slot
+	if (BestMatchIndex == -1 && !bImpossibleToFit)
+	{
+		for (int32 SlotIndex = 0; SlotIndex < AtlasSlots.Num(); ++SlotIndex)
+		{
+			FPaperSpriteAtlasSlot& AtlasSlot = AtlasSlots[SlotIndex];
+			if (Width <= AtlasSlot.Width && Height <= AtlasSlot.Height && AtlasSlot.SpriteRef.Get() == nullptr)
+			{
+				// The one with the least amount of wastage is the best bet?
+				int32 RemainingWidth = AtlasSlot.Width - Width;
+				int32 RemainingHeight = AtlasSlot.Height - Height;
+				int32 CurrentScore = RemainingWidth * RemainingHeight;
+				if (CurrentScore < BestMatchScore)
+				{
+					BestMatchIndex = SlotIndex;
+					BestMatchScore = CurrentScore;
+				}
+			}
+		}
+
+		// Found a slot, insert the sprite
+		// Split up and insert split new slots
+		if (BestMatchIndex != -1)
+		{
+			InsertSpriteIntoSlot(AtlasSlots, BestMatchIndex, Sprite, Width, Height);
+			bOutSpriteChanged = true;
+		}
+	}
+
+	// Can't fit in any of the slots, need a new atlas
+	if (BestMatchIndex == -1 && !bImpossibleToFit)
+	{
+		// Insert a new atlas at the end of the list, empty slots are cleared out later
+		int32 LargestAtlasIndex = -1;
+		for (int32 SlotIndex = 0; SlotIndex < AtlasSlots.Num(); ++SlotIndex)
+		{
+			const FPaperSpriteAtlasSlot& AtlasSlot = AtlasSlots[SlotIndex];
+			LargestAtlasIndex = FMath::Max(LargestAtlasIndex, AtlasSlot.AtlasIndex);
+		}
+
+		// Insert an empty slot for the new atlas texture
+		BestMatchIndex = CreateEmptySlot(AtlasSlots, LargestAtlasIndex + 1, 0, 0, AtlasWidth, AtlasHeight);
+		check(BestMatchIndex != -1);
+
+		// Insert the sprite into this newly created slot
+		InsertSpriteIntoSlot(AtlasSlots, BestMatchIndex, Sprite, Width, Height);
+		bOutSpriteChanged = true;
+	}
+
+	return AtlasSlots[BestMatchIndex];
+}
+
diff --git a/Engine/Plugins/2D/Paper2D/Source/Paper2DEditor/Private/Atlasing/PaperAtlasTextureHelpers.h b/Engine/Plugins/2D/Paper2D/Source/Paper2DEditor/Private/Atlasing/PaperAtlasTextureHelpers.h
new file mode 100644
--- /dev/null
+++ b/Engine/Plugins/2D/Paper2D/Source/Paper2DEditor/Private/Atlasing/PaperAtlasTextureHelpers.h
@@ -0,0 +1,228 @@
+// Copyright 1998-2015 Epic Games, Inc. All Rights Reserved.
+
+#pragma once
+
+// Copy texture data from one 32bpp image to another
+static void CopyTextureData(const uint8* Source, uint8* Dest, uint32 SizeX, uint32 SizeY, uint32 BytesPerPixel, uint32 SourceStride, uint32 DestStride)
+{
+	const uint32 NumBytesPerRow = SizeX * BytesPerPixel;
+
+	for (uint32 Y = 0; Y < SizeY; ++Y)
+	{
+		FMemory::Memcpy(
+			Dest + (DestStride * Y),
+			Source + (SourceStride * Y),
+			NumBytesPerRow
+			);
+	}
+}
+
+// Read sprite data into the 4 bytes per pixel target array
+static bool ReadSpriteTexture(TArray<uint8>& TargetBuffer, FIntPoint& OutSpriteSize, UPaperSprite* Sprite)
+{
+	const int32 BytesPerPixel = 4;
+	const FVector2D SpriteSizeFloat = Sprite->GetSourceSize();
+	const FIntPoint SpriteSize(FMath::TruncToInt(SpriteSizeFloat.X), FMath::TruncToInt(SpriteSizeFloat.Y));
+	TargetBuffer.Empty();
+	TargetBuffer.AddZeroed(SpriteSize.X * SpriteSize.Y * BytesPerPixel);
+
+	check(Sprite->GetSourceTexture());
+	FTextureSource& SourceData = Sprite->GetSourceTexture()->Source;
+	if (SourceData.GetFormat() == TSF_BGRA8)
+	{
+		uint32 BytesPerPixel = SourceData.GetBytesPerPixel();
+		uint8* OffsetSource = SourceData.LockMip(0) + (FMath::TruncToInt(Sprite->GetSourceUV().X) + FMath::TruncToInt(Sprite->GetSourceUV().Y) * SourceData.GetSizeX()) * BytesPerPixel;
+		uint8* OffsetDest = TargetBuffer.GetData();
+		CopyTextureData(OffsetSource, OffsetDest, SpriteSize.X, SpriteSize.Y, BytesPerPixel, SourceData.GetSizeX() * BytesPerPixel, SpriteSize.X * BytesPerPixel);
+		SourceData.UnlockMip(0);
+	}
+	else
+	{
+		UE_LOG(LogPaper2DEditor, Error, TEXT("Sprite %s is not BGRA8, which isn't supported in atlases yet"), *(Sprite->GetPathName()));
+	}
+
+	OutSpriteSize = SpriteSize;
+	return true;
+}
+
+// Fills the padding space with the correct values
+static void PadSprite(const FPaperSpriteAtlasSlot& Slot, EPaperSpriteAtlasPadding PaddingType, int32 Padding, FIntPoint &SpriteSize, int32 AtlasWidth, int32 AtlasBytesPerPixel, TArray<uint8>& TextureData)
+{
+#define PIXEL(PIXEL__X, PIXEL__Y) TextureData[((Slot.Y + (PIXEL__Y)) * AtlasWidth + (Slot.X + (PIXEL__X))) * AtlasBytesPerPixel + (PixelIndex)]
+	if (PaddingType == EPaperSpriteAtlasPadding::DilateBorder)
+	{
+		for (int32 X = 0; X < Padding; ++X)
+		{
+			for (int32 Y = 0; Y < SpriteSize.Y + Padding * 2; ++Y)
+			{
+				for (int32 PixelIndex = 0; PixelIndex < AtlasBytesPerPixel; ++PixelIndex)
+				{
+					int32 ClampedY = Padding + FMath::Clamp(Y - Padding, 0, SpriteSize.Y - 1);
+					PIXEL(X, Y) = PIXEL(Padding, ClampedY);
+					PIXEL(Padding + SpriteSize.X + X, Y) = PIXEL(Padding + SpriteSize.X - 1, ClampedY);
+				}
+			}
+		}
+		for (int32 Y = 0; Y < Padding; ++Y)
+		{
+			for (int32 X = 0; X < SpriteSize.X + Padding * 2; ++X)
+			{
+				for (int32 PixelIndex = 0; PixelIndex < AtlasBytesPerPixel; ++PixelIndex)
+				{
+					int32 ClampedX = Padding + FMath::Clamp(X - Padding, 0, SpriteSize.X - 1);
+					PIXEL(X, Y) = PIXEL(ClampedX, Padding);
+					PIXEL(X, Padding + SpriteSize.Y + Y) = PIXEL(ClampedX, Padding + SpriteSize.Y - 1);
+				}
+			}
+		}
+	}
+	else if (PaddingType == EPaperSpriteAtlasPadding::PadWithZero)
+	{
+		for (int32 X = 0; X < Padding; ++X)
+		{
+			for (int32 Y = 0; Y < SpriteSize.Y + Padding * 2; ++Y)
+			{
+				for (int32 PixelIndex = 0; PixelIndex < AtlasBytesPerPixel; ++PixelIndex)
+				{
+					PIXEL(X, Y) = 0;
+					PIXEL(Padding + SpriteSize.X + X, Y) = 0;
+				}
+			}
+		}
+		for (int32 Y = 0; Y < Padding; ++Y)
+		{
+			for (int32 X = 0; X < SpriteSize.X + Padding * 2; ++X)
+			{
+				for (int32 PixelIndex = 0; PixelIndex < AtlasBytesPerPixel; ++PixelIndex)
+				{
+					PIXEL(X, Y) = 0;
+					PIXEL(X, Padding + SpriteSize.Y + Y) = 0;
+				}
+			}
+		}
+	}
+#undef PIXEL
+}
+
+// Copies the sprite texture data at the position described by slot
+static void CopySpriteToAtlasTextureData(TArray<uint8>& TextureData, int32 AtlasWidth, int32 AtlasHeight, int32 AtlasBytesPerPixel, EPaperSpriteAtlasPadding PaddingType, int32 Padding, UPaperSprite* Sprite, const FPaperSpriteAtlasSlot& Slot)
+{
+	FIntPoint SpriteSize;
+	TArray<uint8> DummyBuffer;
+	ReadSpriteTexture(DummyBuffer, /*out*/SpriteSize, Sprite);
+
+	// Copy the texture into the texture buffer
+	for (int32 Y = 0; Y < SpriteSize.Y; ++Y)
+	{
+		for (int32 X = 0; X < SpriteSize.X; ++X)
+		{
+			for (int32 PixelIndex = 0; PixelIndex < AtlasBytesPerPixel; ++PixelIndex)
+			{
+				TextureData[((Slot.Y + Y + Padding) * AtlasWidth + (Slot.X + X + Padding)) * AtlasBytesPerPixel + PixelIndex] = DummyBuffer[(Y * SpriteSize.X + X) * AtlasBytesPerPixel + PixelIndex];
+			}
+		}
+	}
+
+	// Padding
+	PadSprite(Slot, PaddingType, Padding, SpriteSize, AtlasWidth, AtlasBytesPerPixel, TextureData);
+}
+
+static int32 ClampMips(int Width, int Height, int MipCount)
+{
+	int32 NumMips = 1;
+	while (MipCount > 1 && (Width % 2) == 0 && (Height % 2) == 0)
+	{
+		Width /= 2;
+		Height /= 2;
+		MipCount--;
+		NumMips++;
+	}
+	return NumMips;
+}
+
+static void GenerateMipChainARGB(const TArray<FPaperSpriteAtlasSlot>& Slots, TArray<uint8>& AtlasTextureData, int32 MipCount, int32 Width, int32 Height)
+{
+	int32 SourceMipOffset = 0;
+	int32 SourceMipWidth = Width;
+	int32 SourceMipHeight = Height;
+
+	// Mask bitmap stores all valid pixels in the image, i.e. of the rects that contain data
+	TArray<int8> MaskBitmap;
+	MaskBitmap.AddZeroed(SourceMipWidth * SourceMipHeight);
+	for (const FPaperSpriteAtlasSlot& Slot : Slots)
+	{
+		for (int32 Y = 0; Y < Slot.Height; ++Y)
+		{
+			for (int32 X = 0; X < Slot.Width; ++X)
+			{
+				MaskBitmap[(Slot.Y + Y) * SourceMipWidth + (Slot.X + X)] = 1;
+			}
+		}
+	}
+
+	int32 BytesPerPixel = 4;
+	check(BytesPerPixel == 4); // Only support 4 bytes per pixel
+
+	// Mipmap offset to mip level being written into
+	int32 TargetMipOffset = SourceMipWidth * SourceMipHeight * BytesPerPixel;
+	int32 TargetMipWidth = SourceMipWidth / 2;
+	int32 TargetMipHeight = SourceMipHeight / 2;
+
+	for (int32 MipIndex = 1; MipIndex < MipCount; ++MipIndex)
+	{
+		int32 MipLevelSize = TargetMipHeight * TargetMipWidth * BytesPerPixel;
+
+		for (int32 Y = 0; Y < TargetMipHeight; ++Y)
+		{
+			for (int32 X = 0; X < TargetMipWidth; ++X)
+			{
+				int32 TotalChannelValues[4] = { 0, 0, 0, 0 };
+				check(BytesPerPixel == 4); // refer to above
+
+				int32 ValidPixelCount = 0;
+				int32 Mask0 = MaskBitmap[(Y * 2 + 0) * SourceMipWidth + (X * 2 + 0)];
+				int32 Mask1 = MaskBitmap[(Y * 2 + 0) * SourceMipWidth + (X * 2 + 1)];
+				int32 Mask2 = MaskBitmap[(Y * 2 + 1) * SourceMipWidth + (X * 2 + 0)];
+				int32 Mask3 = MaskBitmap[(Y * 2 + 1) * SourceMipWidth + (X * 2 + 1)];
+				ValidPixelCount = Mask0 + Mask1 + Mask2 + Mask3;
+
+				for (int32 PixelIndex = 0; PixelIndex < BytesPerPixel; ++PixelIndex)
+				{
+					TotalChannelValues[PixelIndex] += AtlasTextureData[SourceMipOffset + ((Y * 2 + 0) * SourceMipWidth + (X * 2 + 0)) * BytesPerPixel + PixelIndex];
+					TotalChannelValues[PixelIndex] += AtlasTextureData[SourceMipOffset + ((Y * 2 + 0) * SourceMipWidth + (X * 2 + 1)) * BytesPerPixel + PixelIndex];
+					TotalChannelValues[PixelIndex] += AtlasTextureData[SourceMipOffset + ((Y * 2 + 1) * SourceMipWidth + (X * 2 + 0)) * BytesPerPixel + PixelIndex];
+					TotalChannelValues[PixelIndex] += AtlasTextureData[SourceMipOffset + ((Y * 2 + 1) * SourceMipWidth + (X * 2 + 1)) * BytesPerPixel + PixelIndex];
+				}
+
+				for (int32 PixelIndex = 0; PixelIndex < BytesPerPixel; ++PixelIndex)
+				{
+					uint8 TargetPixelValue = (ValidPixelCount > 0) ? (TotalChannelValues[PixelIndex] / ValidPixelCount) : 0;
+					AtlasTextureData[TargetMipOffset + (Y * TargetMipWidth + X) * BytesPerPixel + PixelIndex] = TargetPixelValue;
+				}
+			}
+		}
+
+		// Downsample mask in place, if any of the 4 mask pixels == 1, the target mask should be 1
+		for (int32 Y = 0; Y < TargetMipHeight; ++Y)
+		{
+			for (int32 X = 0; X < TargetMipWidth; ++X)
+			{
+				uint8 TargetPixelValue = 0;
+				TargetPixelValue |= MaskBitmap[(Y * 2 + 0) * SourceMipWidth + (X * 2 + 0)];
+				TargetPixelValue |= MaskBitmap[(Y * 2 + 0) * SourceMipWidth + (X * 2 + 1)];
+				TargetPixelValue |= MaskBitmap[(Y * 2 + 1) * SourceMipWidth + (X * 2 + 0)];
+				TargetPixelValue |= MaskBitmap[(Y * 2 + 1) * SourceMipWidth + (X * 2 + 1)];
+				MaskBitmap[Y * TargetMipWidth + X] = TargetPixelValue;
+			}
+		}
+
+		// Update for next mip
+		SourceMipOffset = TargetMipOffset;
+		SourceMipWidth = TargetMipWidth;
+		SourceMipHeight = TargetMipHeight;
+		TargetMipOffset += MipLevelSize;
+		TargetMipWidth /= 2;
+		TargetMipHeight /= 2;
+	}
+}
+
