[INTEGRATE] Change 2426246 by [REDACTED_PERSON][REDACTED_HANDLE] on 2015/01/30 15:17:29

	Paper2D: Implement DoesSocketExist interface for flipbook and sprite components [UE-7642]

diff --git a/Engine/Plugins/2D/Paper2D/Source/Paper2D/Classes/PaperFlipbook.h b/Engine/Plugins/2D/Paper2D/Source/Paper2D/Classes/PaperFlipbook.h
--- a/Engine/Plugins/2D/Paper2D/Source/Paper2D/Classes/PaperFlipbook.h
+++ b/Engine/Plugins/2D/Paper2D/Source/Paper2D/Classes/PaperFlipbook.h
@@ -127,6 +127,9 @@ public:
 	// Returns true if the flipbook has any sockets
 	bool HasAnySockets() const;
 
+	// Returns true if the flipbook has a specific named socket
+	bool DoesSocketExist(FName SocketName) const;
+
 	// Returns a list of all of the sockets
 	void QuerySupportedSockets(TArray<FComponentSocketDescription>& OutSockets) const;
 
diff --git a/Engine/Plugins/2D/Paper2D/Source/Paper2D/Classes/PaperFlipbookComponent.h b/Engine/Plugins/2D/Paper2D/Source/Paper2D/Classes/PaperFlipbookComponent.h
--- a/Engine/Plugins/2D/Paper2D/Source/Paper2D/Classes/PaperFlipbookComponent.h
+++ b/Engine/Plugins/2D/Paper2D/Source/Paper2D/Classes/PaperFlipbookComponent.h
@@ -174,6 +174,7 @@ public:
 
 	// USceneComponent interface
 	virtual bool HasAnySockets() const override;
+	virtual bool DoesSocketExist(FName InSocketName) const override;
 	virtual FTransform GetSocketTransform(FName InSocketName, ERelativeTransformSpace TransformSpace = RTS_World) const override;
 	virtual void QuerySupportedSockets(TArray<FComponentSocketDescription>& OutSockets) const override;
 	// End of USceneComponent interface
diff --git a/Engine/Plugins/2D/Paper2D/Source/Paper2D/Classes/PaperSpriteComponent.h b/Engine/Plugins/2D/Paper2D/Source/Paper2D/Classes/PaperSpriteComponent.h
--- a/Engine/Plugins/2D/Paper2D/Source/Paper2D/Classes/PaperSpriteComponent.h
+++ b/Engine/Plugins/2D/Paper2D/Source/Paper2D/Classes/PaperSpriteComponent.h
@@ -70,6 +70,7 @@ public:
 
 	// USceneComponent interface
 	virtual bool HasAnySockets() const override;
+	virtual bool DoesSocketExist(FName InSocketName) const override;
 	virtual FTransform GetSocketTransform(FName InSocketName, ERelativeTransformSpace TransformSpace = RTS_World) const override;
 	virtual void QuerySupportedSockets(TArray<FComponentSocketDescription>& OutSockets) const override;
 	// End of USceneComponent interface
diff --git a/Engine/Plugins/2D/Paper2D/Source/Paper2D/Private/PaperFlipbook.cpp b/Engine/Plugins/2D/Paper2D/Source/Paper2D/Private/PaperFlipbook.cpp
--- a/Engine/Plugins/2D/Paper2D/Source/Paper2D/Private/PaperFlipbook.cpp
+++ b/Engine/Plugins/2D/Paper2D/Source/Paper2D/Private/PaperFlipbook.cpp
@@ -181,6 +181,22 @@ bool UPaperFlipbook::HasAnySockets() const
 	return false;
 }
 
+bool UPaperFlipbook::DoesSocketExist(FName SocketName) const
+{
+	for (const FPaperFlipbookKeyFrame& KeyFrame : KeyFrames)
+	{
+		if (KeyFrame.Sprite != nullptr)
+		{
+			if (KeyFrame.Sprite->FindSocket(SocketName))
+			{
+				return true;
+			}
+		}
+	}
+
+	return false;
+}
+
 void UPaperFlipbook::QuerySupportedSockets(TArray<FComponentSocketDescription>& OutSockets) const
 {
 	TSet<FName> SocketNames;
diff --git a/Engine/Plugins/2D/Paper2D/Source/Paper2D/Private/PaperFlipbookComponent.cpp b/Engine/Plugins/2D/Paper2D/Source/Paper2D/Private/PaperFlipbookComponent.cpp
--- a/Engine/Plugins/2D/Paper2D/Source/Paper2D/Private/PaperFlipbookComponent.cpp
+++ b/Engine/Plugins/2D/Paper2D/Source/Paper2D/Private/PaperFlipbookComponent.cpp
@@ -589,6 +589,20 @@ bool UPaperFlipbookComponent::HasAnySockets() const
 	return false;
 }
 
+bool UPaperFlipbookComponent::DoesSocketExist(FName InSocketName) const
+{
+	if (SourceFlipbook != nullptr)
+	{
+		FTransform SocketLocalTransform;
+		if (SourceFlipbook->DoesSocketExist(InSocketName))
+		{
+			return true;
+		}
+	}
+
+	return false;
+}
+
 FTransform UPaperFlipbookComponent::GetSocketTransform(FName InSocketName, ERelativeTransformSpace TransformSpace) const
 {
 	if (SourceFlipbook != nullptr)
diff --git a/Engine/Plugins/2D/Paper2D/Source/Paper2D/Private/PaperSpriteComponent.cpp b/Engine/Plugins/2D/Paper2D/Source/Paper2D/Private/PaperSpriteComponent.cpp
--- a/Engine/Plugins/2D/Paper2D/Source/Paper2D/Private/PaperSpriteComponent.cpp
+++ b/Engine/Plugins/2D/Paper2D/Source/Paper2D/Private/PaperSpriteComponent.cpp
@@ -133,6 +133,19 @@ bool UPaperSpriteComponent::HasAnySockets() const
 	return false;
 }
 
+bool UPaperSpriteComponent::DoesSocketExist(FName InSocketName) const
+{
+	if (SourceSprite != nullptr)
+	{
+		if (FPaperSpriteSocket* Socket = SourceSprite->FindSocket(InSocketName))
+		{
+			return true;
+		}
+	}
+
+	return false;
+}
+
 FTransform UPaperSpriteComponent::GetSocketTransform(FName InSocketName, ERelativeTransformSpace TransformSpace) const
 {
 	if (SourceSprite != nullptr)
