fix for UE-4743 LIVE: Crash when moving this asset to another folder

diff --git a/Engine/Plugins/2D/Paper2D/Source/PaperJsonImporter/Private/PaperJsonSpriteSheetImporter.cpp b/Engine/Plugins/2D/Paper2D/Source/PaperJsonImporter/Private/PaperJsonSpriteSheetImporter.cpp
--- a/Engine/Plugins/2D/Paper2D/Source/PaperJsonImporter/Private/PaperJsonSpriteSheetImporter.cpp
+++ b/Engine/Plugins/2D/Paper2D/Source/PaperJsonImporter/Private/PaperJsonSpriteSheetImporter.cpp
@@ -8,6 +8,7 @@
 #include "AssetToolsModule.h"
 #include "AssetRegistryModule.h"
 #include "PackageTools.h"
+#include "ObjectTools.h"
 #include "PaperSpriteSheet.h"
 
 TSharedPtr<FJsonObject> ParseJSON(const FString& FileContents, const FString& NameForErrors)
@@ -478,7 +479,8 @@ bool FPaperJsonSpriteSheetImporter::PerformImport(const FString& LongPackagePath
 		UObject* OuterForFrame = nullptr; // @TODO: Use this if we don't want them to be individual assets - Flipbook;
 
 		// Create a unique package name and asset name for the frame
-		const FString TentativePackagePath = PackageTools::SanitizePackageName(TargetSubPath + TEXT("/") + Frame.FrameName.ToString());
+		const FString SanitizedFrameName = ObjectTools::SanitizeObjectName(Frame.FrameName.ToString());
+		const FString TentativePackagePath = PackageTools::SanitizePackageName(TargetSubPath + TEXT("/") + SanitizedFrameName);
 		FString DefaultSuffix;
 		FString AssetName;
 		FString PackageName;
