Paper2D: Add support for Margin and Padding to the tile set selection rectangle [UE-9988]

diff --git a/Engine/Plugins/2D/Paper2D/Source/Paper2D/Classes/PaperTileSet.h b/Engine/Plugins/2D/Paper2D/Source/Paper2D/Classes/PaperTileSet.h
--- a/Engine/Plugins/2D/Paper2D/Source/Paper2D/Classes/PaperTileSet.h
+++ b/Engine/Plugins/2D/Paper2D/Source/Paper2D/Classes/PaperTileSet.h
@@ -55,4 +55,8 @@ public:
 	int32 GetTileCountY() const;
 
 	bool GetTileUV(int32 TileIndex, /*out*/ FVector2D& Out_TileUV) const;
+	FIntPoint GetTileUVFromTileXY(const FIntPoint& TileXY) const;
+
+	// Converts the texture coordinates into tile coordinates
+	FIntPoint GetTileXYFromTextureUV(const FVector2D& TextureUV, bool bRoundUp) const;
 };
diff --git a/Engine/Plugins/2D/Paper2D/Source/Paper2D/Private/PaperTileSet.cpp b/Engine/Plugins/2D/Paper2D/Source/Paper2D/Private/PaperTileSet.cpp
--- a/Engine/Plugins/2D/Paper2D/Source/Paper2D/Private/PaperTileSet.cpp
+++ b/Engine/Plugins/2D/Paper2D/Source/Paper2D/Private/PaperTileSet.cpp
@@ -67,38 +67,39 @@ int32 UPaperTileSet::GetTileCountY() const
 
 bool UPaperTileSet::GetTileUV(int32 TileIndex, /*out*/ FVector2D& Out_TileUV) const
 {
-	//@TODO: Performance: should cache this stuff
-	if (TileSheet != nullptr)
+	const int32 NumCells = GetTileCount();
+
+	if ((TileIndex < 0) || (TileIndex >= NumCells))
 	{
-		checkSlow((TileWidth > 0) && (TileHeight > 0));
-		const int32 TextureWidth = TileSheet->GetSizeX() - (Margin * 2) + Spacing;
-		const int32 TextureHeight = TileSheet->GetSizeY() - (Margin * 2) + Spacing;
-
-		const int32 CellsX = TextureWidth / (TileWidth + Spacing);
-		const int32 CellsY = TextureHeight / (TileHeight + Spacing);
-
-		const int32 NumCells = CellsX * CellsY;
-
-		if ((TileIndex < 0) || (TileIndex >= NumCells))
-		{
-			return false;
-		}
-		else
-		{
-			const int32 X = TileIndex % CellsX;
-			const int32 Y = TileIndex / CellsX;
-
-			Out_TileUV.X = X * (TileWidth + Spacing) + Margin;
-			Out_TileUV.Y = Y * (TileHeight + Spacing) + Margin;
-			return true;
-		}
+		return false;
 	}
 	else
 	{
-		return false;
+		const int32 CellsX = GetTileCountX();
+
+		const FIntPoint XY(TileIndex % CellsX, TileIndex / CellsX);
+
+		Out_TileUV = GetTileUVFromTileXY(XY);
+		return true;
 	}
 }
 
+FIntPoint UPaperTileSet::GetTileUVFromTileXY(const FIntPoint& TileXY) const
+{
+	return FIntPoint(TileXY.X * (TileWidth + Spacing) + Margin, TileXY.Y * (TileHeight + Spacing) + Margin);
+}
+
+FIntPoint UPaperTileSet::GetTileXYFromTextureUV(const FVector2D& TextureUV, bool bRoundUp) const
+{
+	const float DividendX = TextureUV.X - Margin;
+	const float DividendY = TextureUV.Y - Margin;
+	const float DivisorX = TileWidth + Spacing;
+	const float DivisorY = TileHeight + Spacing;
+	const int32 X = bRoundUp ? FMath::DivideAndRoundUp<int32>(DividendX, DivisorX) : FMath::DivideAndRoundDown<int32>(DividendX, DivisorX);
+	const int32 Y = bRoundUp ? FMath::DivideAndRoundUp<int32>(DividendY, DivisorY) : FMath::DivideAndRoundDown<int32>(DividendY, DivisorY);
+	return FIntPoint(X, Y);
+}
+
 #if WITH_EDITOR
 
 #include "PaperTileMapComponent.h"
diff --git a/Engine/Plugins/2D/Paper2D/Source/Paper2DEditor/Private/TileSetEditor.cpp b/Engine/Plugins/2D/Paper2D/Source/Paper2DEditor/Private/TileSetEditor.cpp
--- a/Engine/Plugins/2D/Paper2D/Source/Paper2DEditor/Private/TileSetEditor.cpp
+++ b/Engine/Plugins/2D/Paper2D/Source/Paper2DEditor/Private/TileSetEditor.cpp
@@ -216,14 +216,16 @@ void STileSetSelectorViewport::OnSelectionChanged(FMarqueeOperation Marquee, boo
 	const FVector2D BottomRightUnrounded = Marquee.Rect.GetLowerRight();
 	if ((TopLeftUnrounded != FVector2D::ZeroVector) || Marquee.IsValid())
 	{
+		const int32 TileCountX = TileSetBeingEdited->GetTileCountX();
+		const int32 TileCountY = TileSetBeingEdited->GetTileCountY();
+
 		// Round down the top left corner
-		SelectionTopLeft.X = FMath::Clamp<int>((int)(TopLeftUnrounded.X / TileSetBeingEdited->TileWidth), 0, TileSetBeingEdited->GetTileCountX());
-		SelectionTopLeft.Y = FMath::Clamp<int>((int)(TopLeftUnrounded.Y / TileSetBeingEdited->TileHeight), 0, TileSetBeingEdited->GetTileCountY());
+		const FIntPoint TileTopLeft = TileSetBeingEdited->GetTileXYFromTextureUV(TopLeftUnrounded, false);
+		SelectionTopLeft = FIntPoint(FMath::Clamp<int32>(TileTopLeft.X, 0, TileCountX), FMath::Clamp<int32>(TileTopLeft.Y, 0, TileCountY));
 
 		// Round up the bottom right corner
-		FIntPoint SelectionBottomRight;
-		SelectionBottomRight.X = FMath::Clamp<int>(FMath::DivideAndRoundUp((int)BottomRightUnrounded.X, TileSetBeingEdited->TileWidth), 0, TileSetBeingEdited->GetTileCountX());
-		SelectionBottomRight.Y = FMath::Clamp<int>(FMath::DivideAndRoundUp((int)BottomRightUnrounded.Y, TileSetBeingEdited->TileHeight), 0, TileSetBeingEdited->GetTileCountY());
+		const FIntPoint TileBottomRight = TileSetBeingEdited->GetTileXYFromTextureUV(BottomRightUnrounded, true);
+		const FIntPoint SelectionBottomRight(FMath::Clamp<int32>(TileBottomRight.X, 0, TileCountX), FMath::Clamp<int32>(TileBottomRight.Y, 0, TileCountY));
 
 		// Compute the new selection dimensions
 		SelectionDimensions = SelectionBottomRight - SelectionTopLeft;
@@ -266,11 +268,14 @@ void STileSetSelectorViewport::RefreshSelectionRectangle()
 		const int32 TileIndex = (bHasSelection && (TileSetBeingEdited != nullptr)) ? (SelectionTopLeft.X + SelectionTopLeft.Y * TileSetBeingEdited->GetTileCountX()) : INDEX_NONE;
 		Client->TileIndex = TileIndex;
 
-		if (bHasSelection)
+		if (bHasSelection && (TileSetBeingEdited != nullptr))
 		{
 			Client->ValidPaintRectangle.Color = FLinearColor::White;
-			Client->ValidPaintRectangle.Dimensions = FVector2D(SelectionDimensions.X * TileSetBeingEdited->TileWidth, SelectionDimensions.Y * TileSetBeingEdited->TileHeight);
-			Client->ValidPaintRectangle.TopLeft = FVector2D(SelectionTopLeft.X * TileSetBeingEdited->TileWidth, SelectionTopLeft.Y * TileSetBeingEdited->TileHeight);
+
+			const FIntPoint TopLeft = TileSetBeingEdited->GetTileUVFromTileXY(SelectionTopLeft);
+			const FIntPoint BottomRight = TileSetBeingEdited->GetTileUVFromTileXY(SelectionTopLeft + SelectionDimensions);
+			Client->ValidPaintRectangle.Dimensions = BottomRight - TopLeft;
+			Client->ValidPaintRectangle.TopLeft = TopLeft;
 		}
 	}
 }
