Paper2D: Add support for specifying additional textures for a sprite
- This is useful for providing masks, normal maps, etc...
- The position of the sprite in each additional texture must match the relative position in the Source sprite
- Sprite texture sampler nodes can now be configured to point to one of the additional textures
[UE-2086]

diff --git a/Engine/Plugins/2D/Paper2D/Source/Paper2D/Classes/PaperSprite.h b/Engine/Plugins/2D/Paper2D/Source/Paper2D/Classes/PaperSprite.h
--- a/Engine/Plugins/2D/Paper2D/Source/Paper2D/Classes/PaperSprite.h
+++ b/Engine/Plugins/2D/Paper2D/Source/Paper2D/Classes/PaperSprite.h
@@ -32,6 +32,8 @@ struct FPaperSpriteSocket
 	FName SocketName;
 };
 
+typedef TArray<class UTexture2D*, TInlineAllocator<4>> FAdditionalSpriteTextureArray;
+
 /**
  * Sprite Asset
  *
@@ -85,6 +87,10 @@ protected:
 	UPROPERTY(Category=Sprite, EditAnywhere, AssetRegistrySearchable)
 	UTexture2D* SourceTexture;
 
+	// Additional source textures for other slots
+	UPROPERTY(Category=Sprite, EditAnywhere, AssetRegistrySearchable, meta=(DisplayName="Additional Textures"))
+	TArray<UTexture2D*> AdditionalSourceTextures;
+
 	UPROPERTY()
 	UTexture2D* BakedSourceTexture;
 
@@ -262,6 +268,9 @@ public:
 	// Returns the texture this should be rendered with
 	UTexture2D* GetBakedTexture() const;
 
+	// Returns the list of additional source textures this should be rendered with
+	void GetBakedAdditionalSourceTextures(FAdditionalSpriteTextureArray& OutTextureList) const;
+
 	// Return the default material for this sprite
 	UMaterialInterface* GetDefaultMaterial() const { return DefaultMaterial; }
 
diff --git a/Engine/Plugins/2D/Paper2D/Source/Paper2D/Classes/SpriteDrawCall.h b/Engine/Plugins/2D/Paper2D/Source/Paper2D/Classes/SpriteDrawCall.h
--- a/Engine/Plugins/2D/Paper2D/Source/Paper2D/Classes/SpriteDrawCall.h
+++ b/Engine/Plugins/2D/Paper2D/Source/Paper2D/Classes/SpriteDrawCall.h
@@ -15,7 +15,9 @@ public:
 	FVector Destination;
 
 	UPROPERTY(Category=Sprite, EditAnywhere)
-	UTexture2D* Texture;
+	UTexture2D* BaseTexture;
+
+	FAdditionalSpriteTextureArray AdditionalTextures;
 
 	UPROPERTY(Category=Sprite, EditAnywhere)
 	FLinearColor Color;
@@ -27,8 +29,13 @@ public:
 	void BuildFromSprite(const class UPaperSprite* Sprite);
 
 	FSpriteDrawCallRecord()
-		: Texture(nullptr)
+		: BaseTexture(nullptr)
 		, Color(FLinearColor::White)
 	{
 	}
+
+	bool IsValid() const
+	{
+		return (RenderVerts.Num() > 0) && (BaseTexture != nullptr) && (BaseTexture->Resource != nullptr);
+	}
 };
diff --git a/Engine/Plugins/2D/Paper2D/Source/Paper2D/Private/PaperRenderSceneProxy.cpp b/Engine/Plugins/2D/Paper2D/Source/Paper2D/Private/PaperRenderSceneProxy.cpp
--- a/Engine/Plugins/2D/Paper2D/Source/Paper2D/Private/PaperRenderSceneProxy.cpp
+++ b/Engine/Plugins/2D/Paper2D/Source/Paper2D/Private/PaperRenderSceneProxy.cpp
@@ -90,26 +90,26 @@ public:
 //static TGlobalResource<FPaperSpriteVertexFactory> GPaperSpriteVertexFactory;
 
 //////////////////////////////////////////////////////////////////////////
-// FTextureOverrideRenderProxy
+// FSpriteTextureOverrideRenderProxy
 
 /**
- * A material render proxy which overrides a named texture parameter.
+ * A material render proxy which overrides various named texture parameters.
  */
-class FTextureOverrideRenderProxy : public FDynamicPrimitiveResource, public FMaterialRenderProxy
+class FSpriteTextureOverrideRenderProxy : public FDynamicPrimitiveResource, public FMaterialRenderProxy
 {
 public:
 	const FMaterialRenderProxy* const Parent;
-	const UTexture* Texture;
-	FName TextureParameterName;
+	const UTexture* BaseTexture;
+	FAdditionalSpriteTextureArray AdditionalTextures;
 
 	/** Initialization constructor. */
-	FTextureOverrideRenderProxy(const FMaterialRenderProxy* InParent, const UTexture* InTexture, const FName& InParameterName)
+	FSpriteTextureOverrideRenderProxy(const FMaterialRenderProxy* InParent, const UTexture* InBaseTexture, FAdditionalSpriteTextureArray InAdditionalTextures)
 		: Parent(InParent)
-		, Texture(InTexture)
-		, TextureParameterName(InParameterName)
+		, BaseTexture(InBaseTexture)
+		, AdditionalTextures(InAdditionalTextures)
 	{}
 
-	virtual ~FTextureOverrideRenderProxy()
+	virtual ~FSpriteTextureOverrideRenderProxy()
 	{
 	}
 
@@ -142,16 +142,30 @@ public:
 	{
 		if (ParameterName == TextureParameterName)
 		{
-			*OutValue = Texture;
+			*OutValue = BaseTexture;
 			return true;
 		}
-		else
+		else if (ParameterName.GetComparisonIndex() == AdditionalTextureParameterRootName.GetComparisonIndex())
 		{
-			return Parent->GetTextureValue(ParameterName, OutValue, Context);
+			const int32 AdditionalSlotIndex = ParameterName.GetNumber() - 1;
+			if (AdditionalTextures.IsValidIndex(AdditionalSlotIndex))
+			{
+				*OutValue = AdditionalTextures[AdditionalSlotIndex];
+				return true;
+			}
 		}
+		
+		return Parent->GetTextureValue(ParameterName, OutValue, Context);
 	}
+
+protected:
+	static const FName TextureParameterName;
+	static const FName AdditionalTextureParameterRootName;
 };
 
+const FName FSpriteTextureOverrideRenderProxy::TextureParameterName(TEXT("SpriteTexture"));
+const FName FSpriteTextureOverrideRenderProxy::AdditionalTextureParameterRootName(TEXT("SpriteAdditionalTexture"));
+
 //////////////////////////////////////////////////////////////////////////
 // FPaperRenderSceneProxy
 
@@ -330,9 +344,8 @@ void FPaperRenderSceneProxy::GetNewBatchMeshes(const FSceneView* View, bool bUse
 	{
 		const FLinearColor EffectiveWireframeColor = (Batch.Material->GetBlendMode() != BLEND_Opaque) ? WireframeColor : FLinearColor::Green;
 		FMaterialRenderProxy* ParentMaterialProxy = Batch.Material->GetRenderProxy((View->Family->EngineShowFlags.Selection) && IsSelected(), IsHovered());
-		FTexture* TextureResource = (Batch.Texture != nullptr) ? Batch.Texture->Resource : nullptr;
 
-		if ((TextureResource != nullptr) && (Batch.NumVertices > 0))
+		if (Batch.IsValid())
 		{
 			FMeshBatch& Mesh = Collector.AllocateMesh();
 
@@ -357,7 +370,7 @@ void FPaperRenderSceneProxy::GetNewBatchMeshes(const FSceneView* View, bool bUse
 			}
 
 			// Create a texture override material proxy and register it as a dynamic resource so that it won't be deleted until the rendering thread has finished with it
-			FTextureOverrideRenderProxy* TextureOverrideMaterialProxy = new FTextureOverrideRenderProxy(ParentMaterialProxy, Batch.Texture, TEXT("SpriteTexture"));
+			FSpriteTextureOverrideRenderProxy* TextureOverrideMaterialProxy = new FSpriteTextureOverrideRenderProxy(ParentMaterialProxy, Batch.BaseTexture, Batch.AdditionalTextures);
 			Collector.RegisterOneFrameMaterialProxy(TextureOverrideMaterialProxy);
 
 			Mesh.UseDynamicData = true;
@@ -406,8 +419,7 @@ void FPaperRenderSceneProxy::GetBatchMesh(const FSceneView* View, bool bUseOverr
 
 	for (const FSpriteDrawCallRecord& Record : Batch)
 	{
-		FTexture* TextureResource = (Record.Texture != nullptr) ? Record.Texture->Resource : nullptr;
-		if ((TextureResource != nullptr) && (Record.RenderVerts.Num() > 0))
+		if (Record.IsValid())
 		{
 			const FLinearColor SpriteColor = bUseOverrideColor ? OverrideColor : Record.Color;
 			const FVector EffectiveOrigin = Record.Destination;
@@ -454,7 +466,7 @@ void FPaperRenderSceneProxy::GetBatchMesh(const FSceneView* View, bool bUseOverr
 			}
 
 			// Create a texture override material proxy and register it as a dynamic resource so that it won't be deleted until the rendering thread has finished with it
-			FTextureOverrideRenderProxy* TextureOverrideMaterialProxy = new FTextureOverrideRenderProxy(ParentMaterialProxy, Record.Texture, TEXT("SpriteTexture"));
+			FSpriteTextureOverrideRenderProxy* TextureOverrideMaterialProxy = new FSpriteTextureOverrideRenderProxy(ParentMaterialProxy, Record.BaseTexture, Record.AdditionalTextures);
 			Collector.RegisterOneFrameMaterialProxy(TextureOverrideMaterialProxy);
 
 			Mesh.MaterialRenderProxy = TextureOverrideMaterialProxy;
diff --git a/Engine/Plugins/2D/Paper2D/Source/Paper2D/Private/PaperRenderSceneProxy.h b/Engine/Plugins/2D/Paper2D/Source/Paper2D/Private/PaperRenderSceneProxy.h
--- a/Engine/Plugins/2D/Paper2D/Source/Paper2D/Private/PaperRenderSceneProxy.h
+++ b/Engine/Plugins/2D/Paper2D/Source/Paper2D/Private/PaperRenderSceneProxy.h
@@ -37,18 +37,29 @@ struct PAPER2D_API FPaperSpriteVertex
 
 struct PAPER2D_API FSpriteRenderSection
 {
+	UMaterialInterface* Material;
+	UTexture2D* BaseTexture;
+	FAdditionalSpriteTextureArray AdditionalTextures;
+
 	int32 VertexOffset;
 	int32 NumVertices;
 
-	UMaterialInterface* Material;
-	UTexture2D* Texture;
-
 	FSpriteRenderSection()
 		: VertexOffset(INDEX_NONE)
 		, NumVertices(0)
 	{
 	}
 
+	class FTexture* GetBaseTextureResource() const
+	{
+		return (BaseTexture != nullptr) ? BaseTexture->Resource : nullptr;
+	}
+
+	bool IsValid() const
+	{
+		return (NumVertices > 0) && (GetBaseTextureResource() != nullptr);
+	}
+
 #if 0
 	template <typename SourceArrayType>
 	void AddTriangles(const FSpriteDrawCallRecord& Record, SourceArrayType& Vertices)
diff --git a/Engine/Plugins/2D/Paper2D/Source/Paper2D/Private/PaperSprite.cpp b/Engine/Plugins/2D/Paper2D/Source/Paper2D/Private/PaperSprite.cpp
--- a/Engine/Plugins/2D/Paper2D/Source/Paper2D/Private/PaperSprite.cpp
+++ b/Engine/Plugins/2D/Paper2D/Source/Paper2D/Private/PaperSprite.cpp
@@ -312,7 +312,9 @@ void FSpriteDrawCallRecord::BuildFromSprite(const UPaperSprite* Sprite)
 	if (Sprite != nullptr)
 	{
 		Destination = FVector::ZeroVector;
-		Texture = Sprite->GetBakedTexture();
+		BaseTexture = Sprite->GetBakedTexture();
+		Sprite->GetBakedAdditionalSourceTextures(/*out*/ AdditionalTextures);
+
 		Color = FLinearColor::White;
 
 		RenderVerts = Sprite->BakedRenderData;
@@ -1932,6 +1934,11 @@ UTexture2D* UPaperSprite::GetBakedTexture() const
 	return (BakedSourceTexture != nullptr) ? BakedSourceTexture : SourceTexture;
 }
 
+void UPaperSprite::GetBakedAdditionalSourceTextures(FAdditionalSpriteTextureArray& OutTextureList) const
+{
+	OutTextureList = AdditionalSourceTextures;
+}
+
 UMaterialInterface* UPaperSprite::GetMaterial(int32 MaterialIndex) const
 {
 	if (MaterialIndex == 0)
diff --git a/Engine/Plugins/2D/Paper2D/Source/Paper2D/Private/PaperSpriteSceneProxy.cpp b/Engine/Plugins/2D/Paper2D/Source/Paper2D/Private/PaperSpriteSceneProxy.cpp
--- a/Engine/Plugins/2D/Paper2D/Source/Paper2D/Private/PaperSpriteSceneProxy.cpp
+++ b/Engine/Plugins/2D/Paper2D/Source/Paper2D/Private/PaperSpriteSceneProxy.cpp
@@ -96,12 +96,14 @@ void FPaperSpriteSceneProxy::SetSprite_RenderThread(const FSpriteDrawCallRecord&
 		
 		Record.Color = NewDynamicData.Color;
 		Record.Destination = NewDynamicData.Destination;
-		Record.Texture = NewDynamicData.Texture;
+		Record.BaseTexture = NewDynamicData.BaseTexture;
+		Record.AdditionalTextures = NewDynamicData.AdditionalTextures;
 		Record.RenderVerts.Append(NewDynamicData.RenderVerts.GetData(), SplitIndex);
 
 		AltRecord.Color = NewDynamicData.Color;
 		AltRecord.Destination = NewDynamicData.Destination;
-		AltRecord.Texture = NewDynamicData.Texture;
+		AltRecord.BaseTexture = NewDynamicData.BaseTexture;
+		AltRecord.AdditionalTextures = NewDynamicData.AdditionalTextures;
 		AltRecord.RenderVerts.Append(NewDynamicData.RenderVerts.GetData() + SplitIndex, NewDynamicData.RenderVerts.Num() - SplitIndex);
 	}
 	else
diff --git a/Engine/Plugins/2D/Paper2D/Source/Paper2D/Private/PaperTileMapComponent.cpp b/Engine/Plugins/2D/Paper2D/Source/Paper2D/Private/PaperTileMapComponent.cpp
--- a/Engine/Plugins/2D/Paper2D/Source/Paper2D/Private/PaperTileMapComponent.cpp
+++ b/Engine/Plugins/2D/Paper2D/Source/Paper2D/Private/PaperTileMapComponent.cpp
@@ -291,7 +291,9 @@ void UPaperTileMapComponent::RebuildRenderData(FPaperTileMapRenderSceneProxy* Pr
 					if ((SourceTexture != LastSourceTexture) || (CurrentBatch == nullptr))
 					{
 						CurrentBatch = (new (BatchedSprites) FSpriteDrawCallRecord());
-						CurrentBatch->Texture = SourceTexture;
+						CurrentBatch->BaseTexture = SourceTexture;
+						//CurrentBatch->AdditionalTextures = ?; //@TODO: PAPER2D: Need to add multi-texture support to tile sets / tile maps
+						// Probably also need to change the batch check here to TileSet changing to avoid checking each texture in the array
 						CurrentBatch->Color = DrawColor;
 						CurrentBatch->Destination = TopLeftCornerOfTile.ProjectOnTo(PaperAxisZ);
 					}
diff --git a/Engine/Plugins/2D/Paper2D/Source/Paper2D/Private/Rendering/MaterialExpressionSpriteTextureSampler.cpp b/Engine/Plugins/2D/Paper2D/Source/Paper2D/Private/Rendering/MaterialExpressionSpriteTextureSampler.cpp
--- a/Engine/Plugins/2D/Paper2D/Source/Paper2D/Private/Rendering/MaterialExpressionSpriteTextureSampler.cpp
+++ b/Engine/Plugins/2D/Paper2D/Source/Paper2D/Private/Rendering/MaterialExpressionSpriteTextureSampler.cpp
@@ -12,12 +12,30 @@ UMaterialExpressionSpriteTextureSampler::UMaterialExpressionSpriteTextureSampler
 	: Super(ObjectInitializer)
 {
 	ParameterName = TEXT("SpriteTexture");
+	bSampleAdditionalTextures = false;
+	AdditionalSlotIndex = 0;
 }
 
 void UMaterialExpressionSpriteTextureSampler::GetCaption(TArray<FString>& OutCaptions) const
 {
-	OutCaptions.Add(TEXT("Paper2D"));
-	OutCaptions.Add(TEXT("SpriteTexture"));
+	OutCaptions.Add(TEXT("Paper2D Sprite"));
+
+	if (!SlotDisplayName.IsEmpty())
+	{
+		OutCaptions.Add(SlotDisplayName.ToString());
+	}
+
+	if (bSampleAdditionalTextures)
+	{
+		FNumberFormattingOptions NoCommas;
+		NoCommas.UseGrouping = false;
+		const FText SlotDesc = FText::Format(LOCTEXT("SpriteSamplerTitle_AdditionalSlot", "Additional Texture #{0}"), FText::AsNumber(AdditionalSlotIndex, &NoCommas));
+		OutCaptions.Add(SlotDesc.ToString());
+	}
+	else
+	{
+		OutCaptions.Add(LOCTEXT("SpriteSamplerTitle_BasicSlot", "Source Texture").ToString());
+	}
 }
 
 #if WITH_EDITOR
@@ -35,13 +53,25 @@ bool UMaterialExpressionSpriteTextureSampler::CanRenameNode() const
 
 void UMaterialExpressionSpriteTextureSampler::GetExpressionToolTip(TArray<FString>& OutToolTip)
 {
-	OutToolTip.Add(LOCTEXT("SpriteTextureSamplerTooltip", "This is a texture sampler 2D with the parameter name fixed as 'SpriteTexture'. The texture specified here will be replaced by the SourceTexture of a Paper2D sprite if this material is used on a sprite.").ToString());
+	const FText MyTooltip = GetClass()->GetToolTipText(/*bShortTooltip=*/ true);
+	OutToolTip.Add(MyTooltip.ToString());
 }
 
 void UMaterialExpressionSpriteTextureSampler::PostEditChangeProperty(FPropertyChangedEvent& PropertyChangedEvent)
 {
+	// Clamp the slot index to something reasonably sane
+	AdditionalSlotIndex = FMath::Clamp<int32>(AdditionalSlotIndex, 0, 127);
+
 	// Ensure that the parameter name never changes
-	ParameterName = TEXT("SpriteTexture");
+	if (bSampleAdditionalTextures)
+	{
+		ParameterName = FName(TEXT("SpriteAdditionalTexture"), AdditionalSlotIndex + 1);
+	}
+	else
+	{
+		ParameterName = TEXT("SpriteTexture");
+	}
+
 	Super::PostEditChangeProperty(PropertyChangedEvent);
 }
 
diff --git a/Engine/Plugins/2D/Paper2D/Source/Paper2D/Public/MaterialExpressionSpriteTextureSampler.h b/Engine/Plugins/2D/Paper2D/Source/Paper2D/Public/MaterialExpressionSpriteTextureSampler.h
--- a/Engine/Plugins/2D/Paper2D/Source/Paper2D/Public/MaterialExpressionSpriteTextureSampler.h
+++ b/Engine/Plugins/2D/Paper2D/Source/Paper2D/Public/MaterialExpressionSpriteTextureSampler.h
@@ -5,12 +5,24 @@
 #include "Runtime/Engine/Classes/Materials/MaterialExpressionTextureSampleParameter2D.h"
 #include "MaterialExpressionSpriteTextureSampler.generated.h"
 
-// This is a texture sampler 2D with the parameter name fixed as 'SpriteTexture'. The texture specified here will be replaced by the SourceTexture of a Paper2D sprite if this material is used on a sprite.
+// This is a texture sampler 2D with a special automatically defined parameter name. The texture specified here will be replaced by the SourceTexture or an AdditionalSourceTextures entry of a Paper2D sprite if this material is used on a sprite.
 UCLASS(hideCategories=(MaterialExpressionTextureSampleParameter, MaterialExpressionSpriteTextureSampler))
 class PAPER2D_API UMaterialExpressionSpriteTextureSampler : public UMaterialExpressionTextureSampleParameter2D
 {
 	GENERATED_UCLASS_BODY()
 
+	// Is this a sampler for the default SourceTexture or the AdditionalSourceTextures list?
+	UPROPERTY(EditAnywhere, Category=Paper2D)
+	bool bSampleAdditionalTextures;
+
+	// This is the slot index into the AdditionalSourceTextures array
+	UPROPERTY(EditAnywhere, Category=Paper2D, meta=(EditCondition=bSampleAdditionalTextures))
+	int32 AdditionalSlotIndex;
+
+	// Friendly label for the texture slot, displayed in the Sprite Editor if not empty
+	UPROPERTY(EditAnywhere, Category=Paper2D)
+	FText SlotDisplayName;
+
 	// UMaterialExpression interface
 	virtual void GetCaption(TArray<FString>& OutCaptions) const override;
 #if WITH_EDITOR
