Paper2D: Add name customizations for sprite texture slots, which are pulled from the default material (so you can tag a slot as 'Normal Map' or 'Base Map' and it will display that way in the Sprite Editor)

diff --git a/Engine/Plugins/2D/Paper2D/Source/Paper2DEditor/Private/SpriteEditor/SpriteDetailsCustomization.cpp b/Engine/Plugins/2D/Paper2D/Source/Paper2DEditor/Private/SpriteEditor/SpriteDetailsCustomization.cpp
--- a/Engine/Plugins/2D/Paper2D/Source/Paper2DEditor/Private/SpriteEditor/SpriteDetailsCustomization.cpp
+++ b/Engine/Plugins/2D/Paper2D/Source/Paper2DEditor/Private/SpriteEditor/SpriteDetailsCustomization.cpp
@@ -7,6 +7,8 @@
 #include "Editor/Documentation/Public/IDocumentation.h"
 #include "PhysicsEngine/PhysicsSettings.h"
 #include "PropertyRestriction.h"
+#include "PropertyCustomizationHelpers.h"
+#include "MaterialExpressionSpriteTextureSampler.h"
 
 #define LOCTEXT_NAMESPACE "SpriteEditor"
 
@@ -55,7 +57,7 @@ FDetailWidgetRow& FSpriteDetailsCustomization::GenerateWarningRow(IDetailCategor
 void FSpriteDetailsCustomization::CustomizeDetails(IDetailLayoutBuilder& DetailLayout)
 {
 	// Make sure sprite properties are near the top
-	IDetailCategoryBuilder& SpriteCategory = DetailLayout.EditCategory("Sprite", FText::GetEmpty(), ECategoryPriority::TypeSpecific);
+	IDetailCategoryBuilder& SpriteCategory = DetailLayout.EditCategory("Sprite", FText::GetEmpty(), ECategoryPriority::Important);
 	BuildSpriteSection(SpriteCategory, DetailLayout);
 
 	// Build the rendering category
@@ -72,7 +74,8 @@ void FSpriteDetailsCustomization::BuildSpriteSection(IDetailCategoryBuilder& Spr
 	// Show other normal properties in the sprite category so that desired ordering doesn't get messed up
 	SpriteCategory.AddProperty(GET_MEMBER_NAME_CHECKED(UPaperSprite, SourceUV));
 	SpriteCategory.AddProperty(GET_MEMBER_NAME_CHECKED(UPaperSprite, SourceDimension));
-	SpriteCategory.AddProperty(GET_MEMBER_NAME_CHECKED(UPaperSprite, SourceTexture));
+	BuildTextureSection(SpriteCategory, DetailLayout);
+
 	SpriteCategory.AddProperty(GET_MEMBER_NAME_CHECKED(UPaperSprite, DefaultMaterial));
 	SpriteCategory.AddProperty(GET_MEMBER_NAME_CHECKED(UPaperSprite, Sockets));
 	SpriteCategory.AddProperty(GET_MEMBER_NAME_CHECKED(UPaperSprite, PixelsPerUnrealUnit));
@@ -189,6 +192,116 @@ void FSpriteDetailsCustomization::BuildCollisionSection(IDetailCategoryBuilder&
 	}
 }
 
+void FSpriteDetailsCustomization::BuildTextureSection(IDetailCategoryBuilder& SpriteCategory, IDetailLayoutBuilder& DetailLayout)
+{
+	// Grab information about the material
+	TSharedPtr<IPropertyHandle> DefaultMaterialProperty = DetailLayout.GetProperty(GET_MEMBER_NAME_CHECKED(UPaperSprite, DefaultMaterial));
+
+	FText SourceTextureOverrideLabel;
+	if (DefaultMaterialProperty.IsValid())
+	{
+		UObject* DefaultMaterialAsObject;
+		if (DefaultMaterialProperty->GetValue(/*out*/ DefaultMaterialAsObject) == FPropertyAccess::Success)
+		{
+			if (UMaterialInterface* DefaultMaterialInterface = Cast<UMaterialInterface>(DefaultMaterialAsObject))
+			{
+				if (UMaterial* DefaultMaterial = DefaultMaterialInterface->GetMaterial())
+				{
+					// Get a list of sprite samplers
+					TArray<const UMaterialExpressionSpriteTextureSampler*> SpriteSamplerExpressions;
+					DefaultMaterial->GetAllExpressionsOfType(/*inout*/ SpriteSamplerExpressions);
+
+					// Turn that into a set of labels
+					for (const UMaterialExpressionSpriteTextureSampler* Sampler : SpriteSamplerExpressions)
+					{
+						if (!Sampler->SlotDisplayName.IsEmpty())
+						{
+							if (Sampler->bSampleAdditionalTextures)
+							{
+								AdditionalTextureLabels.FindOrAdd(Sampler->AdditionalSlotIndex) = Sampler->SlotDisplayName;
+							}
+							else
+							{
+								SourceTextureOverrideLabel = Sampler->SlotDisplayName;
+							}
+						}
+					}
+				}
+			}
+		}
+	}
+
+	// Create the base texture widget
+	TSharedPtr<IPropertyHandle> SourceTextureProperty = DetailLayout.GetProperty(GET_MEMBER_NAME_CHECKED(UPaperSprite, SourceTexture));
+	DetailLayout.HideProperty(SourceTextureProperty);
+	SpriteCategory.AddCustomRow(SourceTextureProperty->GetPropertyDisplayName())
+		.NameContent()
+		[
+			CreateTextureNameWidget(SourceTextureProperty, SourceTextureOverrideLabel)
+		]
+		.ValueContent()
+		[
+			SourceTextureProperty->CreatePropertyValueWidget()
+		];
+
+	// Create the additional textures widget
+	TSharedPtr<IPropertyHandle> AdditionalSourceTexturesProperty = DetailLayout.GetProperty(GET_MEMBER_NAME_CHECKED(UPaperSprite, AdditionalSourceTextures));
+	TSharedRef<FDetailArrayBuilder> AdditionalSourceTexturesBuilder = MakeShareable(new FDetailArrayBuilder(AdditionalSourceTexturesProperty.ToSharedRef()));
+	AdditionalSourceTexturesBuilder->OnGenerateArrayElementWidget(FOnGenerateArrayElementWidget::CreateSP(this, &FSpriteDetailsCustomization::GenerateAdditionalTextureWidget));
+	SpriteCategory.AddCustomBuilder(AdditionalSourceTexturesBuilder);
+}
+
+void FSpriteDetailsCustomization::GenerateAdditionalTextureWidget(TSharedRef<IPropertyHandle> PropertyHandle, int32 ArrayIndex, IDetailChildrenBuilder& ChildrenBuilder)
+{
+	IDetailPropertyRow& TextureRow = ChildrenBuilder.AddChildProperty(PropertyHandle);
+
+	FText ExtraText;
+	if (FText* pExtraText = AdditionalTextureLabels.Find(ArrayIndex))
+	{
+		ExtraText = *pExtraText;
+	}
+
+	FNumberFormattingOptions NoCommas;
+	NoCommas.UseGrouping = false;
+	const FText SlotDesc = FText::Format(LOCTEXT("AdditionalTextureSlotIndex", "Slot #{0}"), FText::AsNumber(ArrayIndex, &NoCommas));
+
+	TextureRow.DisplayName(SlotDesc);
+
+	TextureRow.CustomWidget(false)
+		.NameContent()
+		[
+			CreateTextureNameWidget(PropertyHandle, ExtraText)
+		]
+		.ValueContent()
+		[
+			PropertyHandle->CreatePropertyValueWidget()
+		];
+}
+
+TSharedRef<SWidget> FSpriteDetailsCustomization::CreateTextureNameWidget(TSharedPtr<IPropertyHandle> PropertyHandle, const FText& OverrideText)
+{
+	TSharedRef<SWidget> PropertyNameWidget = PropertyHandle->CreatePropertyNameWidget();
+	if (OverrideText.IsEmpty())
+	{
+		return PropertyNameWidget;
+	}
+	else
+	{
+		return SNew(SVerticalBox)
+			+SVerticalBox::Slot()
+			[
+				PropertyNameWidget
+			]
+			+SVerticalBox::Slot()
+			.Padding(FMargin(0.0f, 4.0f, 0.0f, 0.0f))
+			[
+				SNew(STextBlock)
+				.Font(IDetailLayoutBuilder::GetDetailFont())
+				.Text(OverrideText)
+			];
+	}
+}
+
 EVisibility FSpriteDetailsCustomization::PhysicsModeMatches(TSharedPtr<IPropertyHandle> Property, ESpriteCollisionMode::Type DesiredMode) const
 {
 	if (Property.IsValid())
diff --git a/Engine/Plugins/2D/Paper2D/Source/Paper2DEditor/Private/SpriteEditor/SpriteDetailsCustomization.h b/Engine/Plugins/2D/Paper2D/Source/Paper2DEditor/Private/SpriteEditor/SpriteDetailsCustomization.h
--- a/Engine/Plugins/2D/Paper2D/Source/Paper2DEditor/Private/SpriteEditor/SpriteDetailsCustomization.h
+++ b/Engine/Plugins/2D/Paper2D/Source/Paper2DEditor/Private/SpriteEditor/SpriteDetailsCustomization.h
@@ -30,4 +30,14 @@ protected:
 	void BuildSpriteSection(IDetailCategoryBuilder& SpriteCategory, IDetailLayoutBuilder& DetailLayout);
 	void BuildCollisionSection(IDetailCategoryBuilder& CollisionCategory, IDetailLayoutBuilder& DetailLayout);
 	void BuildRenderingSection(IDetailCategoryBuilder& RenderingCategory, IDetailLayoutBuilder& DetailLayout);
+
+	void BuildTextureSection(IDetailCategoryBuilder& SpriteCategory, IDetailLayoutBuilder& DetailLayout);
+	void GenerateAdditionalTextureWidget(TSharedRef<IPropertyHandle> PropertyHandle, int32 ArrayIndex, IDetailChildrenBuilder& ChildrenBuilder);
+
+	// Creates a name widget with optional captioning label
+	TSharedRef<SWidget> CreateTextureNameWidget(TSharedPtr<IPropertyHandle> PropertyHandle, const FText& OverrideText);
+
+protected:
+	// These are labels harvested from the material
+	TMap<int32, FText> AdditionalTextureLabels;
 };
