Paper2D: Fixed a regression where sprites wouldn't cast shadows even when bCastShadows=true [UE-10406]

diff --git a/Engine/Plugins/2D/Paper2D/Source/Paper2D/Private/PaperRenderSceneProxy.cpp b/Engine/Plugins/2D/Paper2D/Source/Paper2D/Private/PaperRenderSceneProxy.cpp
--- a/Engine/Plugins/2D/Paper2D/Source/Paper2D/Private/PaperRenderSceneProxy.cpp
+++ b/Engine/Plugins/2D/Paper2D/Source/Paper2D/Private/PaperRenderSceneProxy.cpp
@@ -472,7 +472,7 @@ void FPaperRenderSceneProxy::GetBatchMesh(const FSceneView* View, bool bUseOverr
 			Mesh.MaterialRenderProxy = TextureOverrideMaterialProxy;
 			Mesh.LCI = nullptr;
 			Mesh.ReverseCulling = IsLocalToWorldDeterminantNegative() ? true : false;
-			Mesh.CastShadow = false;
+			Mesh.CastShadow = bCastShadow;
 			Mesh.DepthPriorityGroup = DPG;
 			Mesh.Type = PT_TriangleList;
 			Mesh.bDisableBackfaceCulling = true;
