Flipbook Editor: Added key frame menu options to act on the associated sprite:
- Find in the Content Browser
- Open the sprite for editing
- Replace the sprite with a different one
Flipbook Editor: Added the current keyframe # in the section header
[UE-10804]

diff --git a/Engine/Plugins/2D/Paper2D/Source/Paper2DEditor/Private/FlipbookEditor/FlipbookEditorCommands.cpp b/Engine/Plugins/2D/Paper2D/Source/Paper2DEditor/Private/FlipbookEditor/FlipbookEditorCommands.cpp
--- a/Engine/Plugins/2D/Paper2D/Source/Paper2DEditor/Private/FlipbookEditor/FlipbookEditorCommands.cpp
+++ b/Engine/Plugins/2D/Paper2D/Source/Paper2DEditor/Private/FlipbookEditor/FlipbookEditorCommands.cpp
@@ -17,4 +17,8 @@ void FFlipbookEditorCommands::RegisterCommands()
 	UI_COMMAND(AddNewFrame, "Add Key Frame", "Adds a new key frame to the flipbook.", EUserInterfaceActionType::Button, FInputGesture());
 	UI_COMMAND(AddNewFrameBefore, "Insert Key Frame Before", "Adds a new key frame to the flipbook before the selection.", EUserInterfaceActionType::Button, FInputGesture());
 	UI_COMMAND(AddNewFrameAfter, "Insert Key Frame After", "Adds a new key frame to the flipbook after the selection.", EUserInterfaceActionType::Button, FInputGesture());
+
+	UI_COMMAND(PickNewSpriteFrame, "Pick New Sprite", "Picks a new sprite for this key frame.", EUserInterfaceActionType::Button, FInputGesture());
+	UI_COMMAND(EditSpriteFrame, "Edit Sprite", "Opens the sprite for this key frame in the Sprite Editor.", EUserInterfaceActionType::Button, FInputGesture());
+	UI_COMMAND(ShowInContentBrowser, "Show in Content Browser", "Shows the sprite for this key frame in the Content Browser.", EUserInterfaceActionType::Button, FInputGesture());
 }
diff --git a/Engine/Plugins/2D/Paper2D/Source/Paper2DEditor/Private/FlipbookEditor/FlipbookEditorCommands.h b/Engine/Plugins/2D/Paper2D/Source/Paper2DEditor/Private/FlipbookEditor/FlipbookEditorCommands.h
--- a/Engine/Plugins/2D/Paper2D/Source/Paper2DEditor/Private/FlipbookEditor/FlipbookEditorCommands.h
+++ b/Engine/Plugins/2D/Paper2D/Source/Paper2DEditor/Private/FlipbookEditor/FlipbookEditorCommands.h
@@ -33,4 +33,9 @@ public:
 	TSharedPtr<FUICommandInfo> AddNewFrame;
 	TSharedPtr<FUICommandInfo> AddNewFrameBefore;
 	TSharedPtr<FUICommandInfo> AddNewFrameAfter;
+
+	// Asset commands
+	TSharedPtr<FUICommandInfo> PickNewSpriteFrame;
+	TSharedPtr<FUICommandInfo> EditSpriteFrame;
+	TSharedPtr<FUICommandInfo> ShowInContentBrowser;
 };
\ No newline at end of file
diff --git a/Engine/Plugins/2D/Paper2D/Source/Paper2DEditor/Private/FlipbookEditor/SFlipbookTimeline.cpp b/Engine/Plugins/2D/Paper2D/Source/Paper2DEditor/Private/FlipbookEditor/SFlipbookTimeline.cpp
--- a/Engine/Plugins/2D/Paper2D/Source/Paper2DEditor/Private/FlipbookEditor/SFlipbookTimeline.cpp
+++ b/Engine/Plugins/2D/Paper2D/Source/Paper2DEditor/Private/FlipbookEditor/SFlipbookTimeline.cpp
@@ -19,7 +19,7 @@
 //////////////////////////////////////////////////////////////////////////
 // SFlipbookTimeline
 
-void SFlipbookTimeline::Construct(const FArguments& InArgs, TSharedPtr<const FUICommandList> InCommandList)
+void SFlipbookTimeline::Construct(const FArguments& InArgs, TSharedPtr<FUICommandList> InCommandList)
 {
 	FlipbookBeingEdited = InArgs._FlipbookBeingEdited;
 	PlayTime = InArgs._PlayTime;
diff --git a/Engine/Plugins/2D/Paper2D/Source/Paper2DEditor/Private/FlipbookEditor/SFlipbookTimeline.h b/Engine/Plugins/2D/Paper2D/Source/Paper2DEditor/Private/FlipbookEditor/SFlipbookTimeline.h
--- a/Engine/Plugins/2D/Paper2D/Source/Paper2DEditor/Private/FlipbookEditor/SFlipbookTimeline.h
+++ b/Engine/Plugins/2D/Paper2D/Source/Paper2DEditor/Private/FlipbookEditor/SFlipbookTimeline.h
@@ -20,7 +20,7 @@ public:
 		SLATE_EVENT(FOnFlipbookKeyframeSelectionChanged, OnSelectionChanged)
 	SLATE_END_ARGS()
 
-	void Construct(const FArguments& InArgs, TSharedPtr<const FUICommandList> InCommandList);
+	void Construct(const FArguments& InArgs, TSharedPtr<FUICommandList> InCommandList);
 
 	// SWidget interface
 	virtual void OnDragEnter(const FGeometry& MyGeometry, const FDragDropEvent& DragDropEvent) override;
@@ -38,7 +38,7 @@ private:
 private:
 	TAttribute<class UPaperFlipbook*> FlipbookBeingEdited;
 	TAttribute<float> PlayTime;
-	TSharedPtr<const FUICommandList> CommandList;
+	TSharedPtr<FUICommandList> CommandList;
 	FOnFlipbookKeyframeSelectionChanged OnSelectionChanged;
 	int32 SlateUnitsPerFrame;
 };
diff --git a/Engine/Plugins/2D/Paper2D/Source/Paper2DEditor/Private/FlipbookEditor/STimelineTrack.cpp b/Engine/Plugins/2D/Paper2D/Source/Paper2DEditor/Private/FlipbookEditor/STimelineTrack.cpp
--- a/Engine/Plugins/2D/Paper2D/Source/Paper2DEditor/Private/FlipbookEditor/STimelineTrack.cpp
+++ b/Engine/Plugins/2D/Paper2D/Source/Paper2DEditor/Private/FlipbookEditor/STimelineTrack.cpp
@@ -6,36 +6,66 @@
 #include "FlipbookEditorCommands.h"
 #include "SFlipbookTrackHandle.h"
 #include "AssetDragDropOp.h"
+#include "PropertyCustomizationHelpers.h"
 
 #define LOCTEXT_NAMESPACE "FlipbookEditor"
 
+//////////////////////////////////////////////////////////////////////////
+// SFlipbookKeyframeWidget
+
 TSharedRef<SWidget> SFlipbookKeyframeWidget::GenerateContextMenu()
 {
+	const FFlipbookEditorCommands& Commands = FFlipbookEditorCommands::Get();
+
 	OnSelectionChanged.ExecuteIfBound(FrameIndex);
 
 	FMenuBuilder MenuBuilder(true, CommandList);
-	MenuBuilder.BeginSection("KeyframeActions", LOCTEXT("KeyframeActionsSectionHeader", "Keyframe Actions"));
+	{
+		FNumberFormattingOptions NoCommas;
+		NoCommas.UseGrouping = false;
+		
+		const FText KeyframeSectionTitle = FText::Format(LOCTEXT("KeyframeActionsSectionHeader", "Keyframe #{0} Actions"), FText::AsNumber(FrameIndex, &NoCommas));
+		MenuBuilder.BeginSection("KeyframeActions", KeyframeSectionTitle);
+
+		// 		MenuBuilder.AddMenuEntry(FGenericCommands::Get().Cut);
+		// 		MenuBuilder.AddMenuEntry(FGenericCommands::Get().Copy);
+		// 		MenuBuilder.AddMenuEntry(FGenericCommands::Get().Paste);
+		MenuBuilder.AddMenuEntry(FGenericCommands::Get().Duplicate);
+		MenuBuilder.AddMenuEntry(FGenericCommands::Get().Delete);
+
+		MenuBuilder.AddMenuSeparator();
 
-	// 		MenuBuilder.AddMenuEntry(FGenericCommands::Get().Cut);
-	// 		MenuBuilder.AddMenuEntry(FGenericCommands::Get().Copy);
-	// 		MenuBuilder.AddMenuEntry(FGenericCommands::Get().Paste);
-	MenuBuilder.AddMenuEntry(FGenericCommands::Get().Duplicate);
-	MenuBuilder.AddMenuEntry(FGenericCommands::Get().Delete);
+		MenuBuilder.AddMenuEntry(Commands.AddNewFrameBefore);
+		MenuBuilder.AddMenuEntry(Commands.AddNewFrameAfter);
 
-	MenuBuilder.AddMenuSeparator();
+		MenuBuilder.EndSection();
+	}
+
+	CommandList->MapAction(Commands.ShowInContentBrowser, FExecuteAction::CreateSP(this, &SFlipbookKeyframeWidget::ShowInContentBrowser));
+	CommandList->MapAction(Commands.EditSpriteFrame, FExecuteAction::CreateSP(this, &SFlipbookKeyframeWidget::EditKeyFrame));
+
+	{
+		TAttribute<FText> CurrentAssetTitle = TAttribute<FText>::Create(TAttribute<FText>::FGetter::CreateSP(this, &SFlipbookKeyframeWidget::GetKeyframeAssetName));
+		MenuBuilder.BeginSection("KeyframeAssetActions", CurrentAssetTitle);
 
-	MenuBuilder.AddMenuEntry(FFlipbookEditorCommands::Get().AddNewFrameBefore);
-	MenuBuilder.AddMenuEntry(FFlipbookEditorCommands::Get().AddNewFrameAfter);
+		MenuBuilder.AddMenuEntry(Commands.ShowInContentBrowser);
+		MenuBuilder.AddMenuEntry(Commands.EditSpriteFrame);
 
-	MenuBuilder.EndSection();
+		MenuBuilder.AddSubMenu(
+			Commands.PickNewSpriteFrame->GetLabel(),
+			Commands.PickNewSpriteFrame->GetDescription(),
+			FNewMenuDelegate::CreateSP(this, &SFlipbookKeyframeWidget::OpenSpritePickerMenu));
 
+		MenuBuilder.EndSection();
+	}
 	return MenuBuilder.MakeWidget();
 }
 
-void SFlipbookKeyframeWidget::Construct(const FArguments& InArgs, int32 InFrameIndex, TSharedPtr<const FUICommandList> InCommandList)
+void SFlipbookKeyframeWidget::Construct(const FArguments& InArgs, int32 InFrameIndex, TSharedPtr<FUICommandList> InCommandList)
 {
 	FrameIndex = InFrameIndex;
-	CommandList = InCommandList;
+	CommandList = MakeShareable(new FUICommandList);
+	CommandList->Append(InCommandList.ToSharedRef());
 	SlateUnitsPerFrame = InArgs._SlateUnitsPerFrame;
 	FlipbookBeingEdited = InArgs._FlipbookBeingEdited;
 	OnSelectionChanged = InArgs._OnSelectionChanged;
@@ -154,19 +184,41 @@ FReply SFlipbookKeyframeWidget::KeyframeOnMouseButtonUp(const FGeometry& MyGeome
 	return FReply::Unhandled();
 }
 
-FText SFlipbookKeyframeWidget::GetKeyframeTooltip() const
+// Can return null
+const FPaperFlipbookKeyFrame* SFlipbookKeyframeWidget::GetKeyFrameData() const
 {
 	UPaperFlipbook* Flipbook = FlipbookBeingEdited.Get();
 	if ((Flipbook != nullptr) && Flipbook->IsValidKeyFrameIndex(FrameIndex))
 	{
-		const FPaperFlipbookKeyFrame& KeyFrame = Flipbook->GetKeyFrameChecked(FrameIndex);
+		return &(Flipbook->GetKeyFrameChecked(FrameIndex));
+	}
 
-		FText SpriteLine = (KeyFrame.Sprite != nullptr) ? FText::FromString(KeyFrame.Sprite->GetName()) : LOCTEXT("NoSprite", "(none)");
+	return nullptr;
+}
+
+FText SFlipbookKeyframeWidget::GetKeyframeAssetName() const
+{
+	if (const FPaperFlipbookKeyFrame* KeyFrame = GetKeyFrameData())
+	{
+		const FText SpriteLine = (KeyFrame->Sprite != nullptr) ? FText::FromString(KeyFrame->Sprite->GetName()) : LOCTEXT("NoSprite", "(none)");
+		return FText::Format(LOCTEXT("KeyFrameAssetName", "Current Asset: {0}"), SpriteLine);
+	}
+	else
+	{
+		return LOCTEXT("KeyFrameAssetName_None", "Current Asset: (none)");
+	}
+}
+
+FText SFlipbookKeyframeWidget::GetKeyframeTooltip() const
+{
+	if (const FPaperFlipbookKeyFrame* KeyFrame = GetKeyFrameData())
+	{
+		const FText SpriteLine = (KeyFrame->Sprite != nullptr) ? FText::FromString(KeyFrame->Sprite->GetName()) : LOCTEXT("NoSprite", "(none)");
 
 		return FText::Format(LOCTEXT("KeyFrameTooltip", "Sprite: {0}\nIndex: {1}\nDuration: {2} frame(s)"),
 			SpriteLine,
 			FText::AsNumber(FrameIndex),
-			FText::AsNumber(KeyFrame.FrameRun));
+			FText::AsNumber(KeyFrame->FrameRun));
 	}
 	else
 	{
@@ -176,11 +228,9 @@ FText SFlipbookKeyframeWidget::GetKeyframeTooltip() const
 
 FOptionalSize SFlipbookKeyframeWidget::GetFrameWidth() const
 {
-	UPaperFlipbook* Flipbook = FlipbookBeingEdited.Get();
-	if (Flipbook && Flipbook->IsValidKeyFrameIndex(FrameIndex))
+	if (const FPaperFlipbookKeyFrame* KeyFrame = GetKeyFrameData())
 	{
-		const FPaperFlipbookKeyFrame& KeyFrame = Flipbook->GetKeyFrameChecked(FrameIndex);
-		return FMath::Max<float>(0, KeyFrame.FrameRun * SlateUnitsPerFrame.Get() - FFlipbookUIConstants::HandleWidth);
+		return FMath::Max<float>(0, KeyFrame->FrameRun * SlateUnitsPerFrame.Get() - FFlipbookUIConstants::HandleWidth);
 	}
 	else
 	{
@@ -188,6 +238,77 @@ FOptionalSize SFlipbookKeyframeWidget::GetFrameWidth() const
 	}
 }
 
+void SFlipbookKeyframeWidget::OpenSpritePickerMenu(FMenuBuilder& MenuBuilder)
+{
+	const bool bAllowClear = true;
+
+	TArray<const UClass*> AllowedClasses;
+	AllowedClasses.Add(UPaperSprite::StaticClass());
+
+	FAssetData CurrentAssetData;
+	if (const FPaperFlipbookKeyFrame* KeyFrame = GetKeyFrameData())
+	{
+		CurrentAssetData = FAssetData(KeyFrame->Sprite);
+	}
+
+	TSharedRef<SWidget> AssetPickerWidget = PropertyCustomizationHelpers::MakeAssetPickerWithMenu(CurrentAssetData,
+		bAllowClear,
+		AllowedClasses,
+		PropertyCustomizationHelpers::GetNewAssetFactoriesForClasses(AllowedClasses),
+		FOnShouldFilterAsset(),
+		FOnAssetSelected::CreateSP(this, &SFlipbookKeyframeWidget::OnAssetSelected),
+		FSimpleDelegate::CreateSP(this, &SFlipbookKeyframeWidget::CloseMenu));
+
+	MenuBuilder.AddWidget(AssetPickerWidget, FText::GetEmpty(), /*bNoIndent=*/ true);
+}
+
+void SFlipbookKeyframeWidget::CloseMenu()
+{
+	FSlateApplication::Get().DismissAllMenus();
+}
+
+void SFlipbookKeyframeWidget::OnAssetSelected(const FAssetData& AssetData)
+{
+	if (UPaperFlipbook* Flipbook = FlipbookBeingEdited.Get())
+	{
+		FScopedFlipbookMutator EditLock(Flipbook);
+
+		if (EditLock.KeyFrames.IsValidIndex(FrameIndex))
+		{
+			UPaperSprite* NewSprite = Cast<UPaperSprite>(AssetData.GetAsset());
+
+			EditLock.KeyFrames[FrameIndex].Sprite = NewSprite;
+		}
+	}
+}
+
+void SFlipbookKeyframeWidget::ShowInContentBrowser()
+{
+	if (const FPaperFlipbookKeyFrame* KeyFrame = GetKeyFrameData())
+	{
+		if (KeyFrame->Sprite != nullptr)
+		{
+			TArray<UObject*> ObjectsToSync;
+			ObjectsToSync.Add(KeyFrame->Sprite);
+			GEditor->SyncBrowserToObjects(ObjectsToSync);
+		}
+	}
+}
+
+void SFlipbookKeyframeWidget::EditKeyFrame()
+{
+	if (const FPaperFlipbookKeyFrame* KeyFrame = GetKeyFrameData())
+	{
+		if (KeyFrame->Sprite != nullptr)
+		{
+			FAssetEditorManager::Get().OpenEditorForAsset(KeyFrame->Sprite);
+		}
+	}
+}
+
+//////////////////////////////////////////////////////////////////////////
+// FFlipbookKeyFrameDragDropOp
+
 TSharedPtr<SWidget> FFlipbookKeyFrameDragDropOp::GetDefaultDecorator() const
 {
 	const FLinearColor BorderColor = (KeyFrameData.Sprite != nullptr) ? FLinearColor::White : FLinearColor::Black;
@@ -268,7 +389,10 @@ FFlipbookKeyFrameDragDropOp::FFlipbookKeyFrameDragDropOp() : Transaction(LOCTEXT
 
 }
 
-void SFlipbookTimelineTrack::Construct(const FArguments& InArgs, TSharedPtr<const FUICommandList> InCommandList)
+//////////////////////////////////////////////////////////////////////////
+// SFlipbookTimelineTrack
+
+void SFlipbookTimelineTrack::Construct(const FArguments& InArgs, TSharedPtr<FUICommandList> InCommandList)
 {
 	CommandList = InCommandList;
 	SlateUnitsPerFrame = InArgs._SlateUnitsPerFrame;
diff --git a/Engine/Plugins/2D/Paper2D/Source/Paper2DEditor/Private/FlipbookEditor/STimelineTrack.h b/Engine/Plugins/2D/Paper2D/Source/Paper2DEditor/Private/FlipbookEditor/STimelineTrack.h
--- a/Engine/Plugins/2D/Paper2D/Source/Paper2DEditor/Private/FlipbookEditor/STimelineTrack.h
+++ b/Engine/Plugins/2D/Paper2D/Source/Paper2DEditor/Private/FlipbookEditor/STimelineTrack.h
@@ -66,28 +66,39 @@ public:
 
 	SLATE_END_ARGS()
 
-	void Construct(const FArguments& InArgs, int32 InFrameIndex, TSharedPtr<const FUICommandList> InCommandList);
+	void Construct(const FArguments& InArgs, int32 InFrameIndex, TSharedPtr<FUICommandList> InCommandList);
 
+	// SWidget interface
 	virtual FReply OnMouseButtonDown(const FGeometry& MyGeometry, const FPointerEvent& MouseEvent) override;
-
 	virtual FReply OnDragDetected(const FGeometry& MyGeometry, const FPointerEvent& MouseEvent) override;
+	virtual FReply OnDrop(const FGeometry& MyGeometry, const FDragDropEvent& DragDropEvent) override;
+	// End of SWidget interface
 
-	FReply OnDrop(const FGeometry& MyGeometry, const FDragDropEvent& DragDropEvent) override;
-
+protected:
 	FReply KeyframeOnMouseButtonUp(const FGeometry& MyGeometry, const FPointerEvent& MouseEvent);
 
+	FText GetKeyframeAssetName() const;
 	FText GetKeyframeTooltip() const;
 
 	TSharedRef<SWidget> GenerateContextMenu();
 
 	FOptionalSize GetFrameWidth() const;
 
+	void OpenSpritePickerMenu(class FMenuBuilder& MenuBuilder);
+	void OnAssetSelected(const class FAssetData& AssetData);
+	void CloseMenu();
+	void ShowInContentBrowser();
+	void EditKeyFrame();
+
+	// Can return null
+	const struct FPaperFlipbookKeyFrame* GetKeyFrameData() const;
+
 protected:
 	int32 FrameIndex;
 	TAttribute<float> SlateUnitsPerFrame;
 	TAttribute<class UPaperFlipbook*> FlipbookBeingEdited;
 	FOnFlipbookKeyframeSelectionChanged OnSelectionChanged;
-	TSharedPtr<const FUICommandList> CommandList;
+	TSharedPtr<FUICommandList> CommandList;
 };
 
 //////////////////////////////////////////////////////////////////////////
@@ -107,7 +118,7 @@ public:
 
 	SLATE_END_ARGS()
 
-	void Construct(const FArguments& InArgs, TSharedPtr<const FUICommandList> InCommandList);
+	void Construct(const FArguments& InArgs, TSharedPtr<FUICommandList> InCommandList);
 
 	virtual void Tick( const FGeometry& AllottedGeometry, const double InCurrentTime, const float InDeltaTime ) override
 	{
@@ -132,5 +143,5 @@ private:
 	float HandleWidth;
 
 	FOnFlipbookKeyframeSelectionChanged OnSelectionChanged;
-	TSharedPtr<const FUICommandList> CommandList;
+	TSharedPtr<FUICommandList> CommandList;
 };
