Paper2D: Generalize socket selection/editing code from the sprite editor and allowed sockets to be seen in the flipbook editor as well
- Added a shadow to the in-viewport display of the socket name
- Added a show menu with the basics to the flipbook editor
- Sockets from sprite keyframes are now visible in the flipbook editor (can be toggled in the show menu)

diff --git a/Engine/Plugins/2D/Paper2D/Source/Paper2DEditor/Private/FlipbookEditor/FlipbookEditor.cpp b/Engine/Plugins/2D/Paper2D/Source/Paper2DEditor/Private/FlipbookEditor/FlipbookEditor.cpp
--- a/Engine/Plugins/2D/Paper2D/Source/Paper2DEditor/Private/FlipbookEditor/FlipbookEditor.cpp
+++ b/Engine/Plugins/2D/Paper2D/Source/Paper2DEditor/Private/FlipbookEditor/FlipbookEditor.cpp
@@ -121,6 +121,12 @@ void SFlipbookEditorViewport::BindCommands()
 		FExecuteAction::CreateSP( EditorViewportClientRef, &FFlipbookEditorViewportClient::ToggleShowPivot ),
 		FCanExecuteAction(),
 		FIsActionChecked::CreateSP( EditorViewportClientRef, &FFlipbookEditorViewportClient::IsShowPivotChecked ) );
+
+	CommandList->MapAction(
+		Commands.SetShowSockets,
+		FExecuteAction::CreateSP( EditorViewportClientRef, &FFlipbookEditorViewportClient::ToggleShowSockets ),
+		FCanExecuteAction(),
+		FIsActionChecked::CreateSP( EditorViewportClientRef, &FFlipbookEditorViewportClient::IsShowSocketsChecked ) );
 }
 
 TSharedRef<FEditorViewportClient> SFlipbookEditorViewport::MakeEditorViewportClient()
diff --git a/Engine/Plugins/2D/Paper2D/Source/Paper2DEditor/Private/FlipbookEditor/FlipbookEditorCommands.cpp b/Engine/Plugins/2D/Paper2D/Source/Paper2DEditor/Private/FlipbookEditor/FlipbookEditorCommands.cpp
--- a/Engine/Plugins/2D/Paper2D/Source/Paper2DEditor/Private/FlipbookEditor/FlipbookEditorCommands.cpp
+++ b/Engine/Plugins/2D/Paper2D/Source/Paper2DEditor/Private/FlipbookEditor/FlipbookEditorCommands.cpp
@@ -13,6 +13,7 @@ void FFlipbookEditorCommands::RegisterCommands()
 	UI_COMMAND(SetShowCollision, "Collision", "Toggles display of the simplified collision mesh of the static mesh, if one has been assigned.", EUserInterfaceActionType::ToggleButton, FInputGesture(EKeys::C, EModifierKey::Alt));
 
 	UI_COMMAND(SetShowPivot, "Show Pivot", "Display the pivot location of the static mesh.", EUserInterfaceActionType::ToggleButton, FInputGesture());
+	UI_COMMAND(SetShowSockets, "Sockets", "Displays the flipbook sockets.", EUserInterfaceActionType::ToggleButton, FInputGesture());
 
 	UI_COMMAND(AddNewFrame, "Add Key Frame", "Adds a new key frame to the flipbook.", EUserInterfaceActionType::Button, FInputGesture());
 	UI_COMMAND(AddNewFrameBefore, "Insert Key Frame Before", "Adds a new key frame to the flipbook before the selection.", EUserInterfaceActionType::Button, FInputGesture());
diff --git a/Engine/Plugins/2D/Paper2D/Source/Paper2DEditor/Private/FlipbookEditor/FlipbookEditorCommands.h b/Engine/Plugins/2D/Paper2D/Source/Paper2DEditor/Private/FlipbookEditor/FlipbookEditorCommands.h
--- a/Engine/Plugins/2D/Paper2D/Source/Paper2DEditor/Private/FlipbookEditor/FlipbookEditorCommands.h
+++ b/Engine/Plugins/2D/Paper2D/Source/Paper2DEditor/Private/FlipbookEditor/FlipbookEditorCommands.h
@@ -28,6 +28,7 @@ public:
 
 	// View Menu Commands
 	TSharedPtr<FUICommandInfo> SetShowPivot;
+	TSharedPtr<FUICommandInfo> SetShowSockets;
 
 	// Timeline commands
 	TSharedPtr<FUICommandInfo> AddNewFrame;
diff --git a/Engine/Plugins/2D/Paper2D/Source/Paper2DEditor/Private/FlipbookEditor/FlipbookEditorViewportClient.cpp b/Engine/Plugins/2D/Paper2D/Source/Paper2DEditor/Private/FlipbookEditor/FlipbookEditorViewportClient.cpp
--- a/Engine/Plugins/2D/Paper2D/Source/Paper2DEditor/Private/FlipbookEditor/FlipbookEditorViewportClient.cpp
+++ b/Engine/Plugins/2D/Paper2D/Source/Paper2DEditor/Private/FlipbookEditor/FlipbookEditorViewportClient.cpp
@@ -7,6 +7,7 @@
 #include "PreviewScene.h"
 #include "ScopedTransaction.h"
 #include "CanvasTypes.h"
+#include "PaperEditorShared/SocketEditing.h"
 
 #define LOCTEXT_NAMESPACE "FlipbookEditor"
 
@@ -30,6 +31,7 @@ FFlipbookEditorViewportClient::FFlipbookEditorViewportClient(const TAttribute<UP
 	PreviewScene->AddComponent(AnimatedRenderComponent.Get(), FTransform::Identity);
 
 	bShowPivot = false;
+	bShowSockets = true;
 	bDeferZoomToSprite = true;
 	DrawHelper.bDrawGrid = false;
 
@@ -58,6 +60,11 @@ void FFlipbookEditorViewportClient::DrawCanvas(FViewport& Viewport, FSceneView&
 		TextItem.Draw(&Canvas);
 		YPos += 36;
 	}
+
+	if (bShowSockets)
+	{
+		FSocketEditingHelper::DrawSocketNames(AnimatedRenderComponent.Get(), Viewport, View, Canvas);
+	}
 }
 
 void FFlipbookEditorViewportClient::Draw(FViewport* Viewport, FCanvas* Canvas)
@@ -73,6 +80,11 @@ void FFlipbookEditorViewportClient::Draw(const FSceneView* View, FPrimitiveDrawI
 	{
 		FUnrealEdUtils::DrawWidget(View, PDI, AnimatedRenderComponent->ComponentToWorld.ToMatrixWithScale(), 0, 0, EAxisList::Screen, EWidgetMovementMode::WMM_Translate);
 	}
+
+	if (bShowSockets)
+	{
+		FSocketEditingHelper::DrawSockets(AnimatedRenderComponent.Get(), View, PDI);
+	}
 }
 
 void FFlipbookEditorViewportClient::Tick(float DeltaSeconds)
diff --git a/Engine/Plugins/2D/Paper2D/Source/Paper2DEditor/Private/FlipbookEditor/FlipbookEditorViewportClient.h b/Engine/Plugins/2D/Paper2D/Source/Paper2DEditor/Private/FlipbookEditor/FlipbookEditorViewportClient.h
--- a/Engine/Plugins/2D/Paper2D/Source/Paper2DEditor/Private/FlipbookEditor/FlipbookEditorViewportClient.h
+++ b/Engine/Plugins/2D/Paper2D/Source/Paper2DEditor/Private/FlipbookEditor/FlipbookEditorViewportClient.h
@@ -34,10 +34,14 @@ public:
 	void ToggleShowPivot() { bShowPivot = !bShowPivot; Invalidate(); }
 	bool IsShowPivotChecked() const { return bShowPivot; }
 
+	void ToggleShowSockets() { bShowSockets = !bShowSockets; Invalidate(); }
+	bool IsShowSocketsChecked() const { return bShowSockets; }
+
 	UPaperFlipbookComponent* GetPreviewComponent() const
 	{
 		return AnimatedRenderComponent.Get();
 	}
+
 private:
 
 	// The preview scene
@@ -55,6 +59,9 @@ private:
 	// Should we show the sprite pivot?
 	bool bShowPivot;
 
+	// Should we show sockets?
+	bool bShowSockets;
+
 	// Should we zoom to the sprite next tick?
 	bool bDeferZoomToSprite;
 };
diff --git a/Engine/Plugins/2D/Paper2D/Source/Paper2DEditor/Private/FlipbookEditor/SFlipbookEditorViewportToolbar.cpp b/Engine/Plugins/2D/Paper2D/Source/Paper2DEditor/Private/FlipbookEditor/SFlipbookEditorViewportToolbar.cpp
--- a/Engine/Plugins/2D/Paper2D/Source/Paper2DEditor/Private/FlipbookEditor/SFlipbookEditorViewportToolbar.cpp
+++ b/Engine/Plugins/2D/Paper2D/Source/Paper2DEditor/Private/FlipbookEditor/SFlipbookEditorViewportToolbar.cpp
@@ -13,4 +13,29 @@ void SFlipbookEditorViewportToolbar::Construct(const FArguments& InArgs, TShared
 	SCommonEditorViewportToolbarBase::Construct(SCommonEditorViewportToolbarBase::FArguments(), InInfoProvider);
 }
 
+TSharedRef<SWidget> SFlipbookEditorViewportToolbar::GenerateShowMenu() const
+{
+	GetInfoProvider().OnFloatingButtonClicked();
+
+	TSharedRef<SEditorViewport> ViewportRef = GetInfoProvider().GetViewportWidget();
+
+	const bool bInShouldCloseWindowAfterMenuSelection = true;
+	FMenuBuilder ShowMenuBuilder(bInShouldCloseWindowAfterMenuSelection, ViewportRef->GetCommandList());
+	{
+		ShowMenuBuilder.AddMenuEntry(FFlipbookEditorCommands::Get().SetShowSockets);
+		ShowMenuBuilder.AddMenuEntry(FFlipbookEditorCommands::Get().SetShowPivot);
+
+		ShowMenuBuilder.AddMenuSeparator();
+
+		ShowMenuBuilder.AddMenuEntry(FFlipbookEditorCommands::Get().SetShowGrid);
+		ShowMenuBuilder.AddMenuEntry(FFlipbookEditorCommands::Get().SetShowBounds);
+
+		ShowMenuBuilder.AddMenuSeparator();
+
+		ShowMenuBuilder.AddMenuEntry(FFlipbookEditorCommands::Get().SetShowCollision);
+	}
+
+	return ShowMenuBuilder.MakeWidget();
+}
+
 #undef LOCTEXT_NAMESPACE
diff --git a/Engine/Plugins/2D/Paper2D/Source/Paper2DEditor/Private/FlipbookEditor/SFlipbookEditorViewportToolbar.h b/Engine/Plugins/2D/Paper2D/Source/Paper2DEditor/Private/FlipbookEditor/SFlipbookEditorViewportToolbar.h
--- a/Engine/Plugins/2D/Paper2D/Source/Paper2DEditor/Private/FlipbookEditor/SFlipbookEditorViewportToolbar.h
+++ b/Engine/Plugins/2D/Paper2D/Source/Paper2DEditor/Private/FlipbookEditor/SFlipbookEditorViewportToolbar.h
@@ -12,4 +12,8 @@ public:
 	SLATE_END_ARGS()
 
 	void Construct(const FArguments& InArgs, TSharedPtr<class ICommonEditorViewportToolbarInfoProvider> InInfoProvider);
+
+	// SCommonEditorViewportToolbarBase interface
+	virtual TSharedRef<SWidget> GenerateShowMenu() const override;
+	// End of SCommonEditorViewportToolbarBase
 };
diff --git a/Engine/Plugins/2D/Paper2D/Source/Paper2DEditor/Private/PaperEditorShared/AssetEditorSelectedItem.h b/Engine/Plugins/2D/Paper2D/Source/Paper2DEditor/Private/PaperEditorShared/AssetEditorSelectedItem.h
new file mode 100644
--- /dev/null
+++ b/Engine/Plugins/2D/Paper2D/Source/Paper2DEditor/Private/PaperEditorShared/AssetEditorSelectedItem.h
@@ -0,0 +1,65 @@
+// Copyright 1998-2015 Epic Games, Inc. All Rights Reserved.
+
+#pragma once
+
+//////////////////////////////////////////////////////////////////////////
+// FSelectedItem
+
+class FSelectedItem
+{
+protected:
+	FSelectedItem(FName InTypeName)
+		: TypeName(InTypeName)
+	{
+	}
+
+protected:
+	FName TypeName;
+public:
+	virtual bool IsA(FName TestType) const
+	{
+		return TestType == TypeName;
+	}
+
+	virtual uint32 GetTypeHash() const
+	{
+		return 0;
+	}
+
+	virtual bool Equals(const FSelectedItem& OtherItem) const
+	{
+		return false;
+	}
+
+	virtual void ApplyDelta(const FVector2D& Delta)
+	{
+	}
+
+	//@TODO: Doesn't belong here in base!
+	virtual void SplitEdge()
+	{
+	}
+
+	virtual FVector GetWorldPos() const
+	{
+		return FVector::ZeroVector;
+	}
+
+	virtual const class FSpriteSelectedVertex* CastSelectedVertex() const
+	{
+		return nullptr;
+	}
+
+	virtual ~FSelectedItem() {}
+};
+
+inline uint32 GetTypeHash(const FSelectedItem& Item)
+{
+	return Item.GetTypeHash();
+}
+
+inline bool operator==(const FSelectedItem& V1, const FSelectedItem& V2)
+{
+	return V1.Equals(V2);
+}
+
diff --git a/Engine/Plugins/2D/Paper2D/Source/Paper2DEditor/Private/PaperEditorShared/SocketEditing.cpp b/Engine/Plugins/2D/Paper2D/Source/Paper2DEditor/Private/PaperEditorShared/SocketEditing.cpp
new file mode 100644
--- /dev/null
+++ b/Engine/Plugins/2D/Paper2D/Source/Paper2DEditor/Private/PaperEditorShared/SocketEditing.cpp
@@ -0,0 +1,148 @@
+// Copyright 1998-2015 Epic Games, Inc. All Rights Reserved.
+
+#include "Paper2DEditorPrivatePCH.h"
+#include "SocketEditing.h"
+
+//////////////////////////////////////////////////////////////////////////
+// HSpriteSelectableObjectHitProxy
+
+IMPLEMENT_HIT_PROXY(HSpriteSelectableObjectHitProxy, HHitProxy);
+
+//////////////////////////////////////////////////////////////////////////
+// FSocketEditingHelper
+
+const FName FSpriteSelectedSocket::SocketTypeID(TEXT("Socket"));
+
+FSpriteSelectedSocket::FSpriteSelectedSocket()
+	: FSelectedItem(SocketTypeID)
+{
+}
+
+bool FSpriteSelectedSocket::Equals(const FSelectedItem& OtherItem) const
+{
+	if (OtherItem.IsA(FSpriteSelectedSocket::SocketTypeID))
+	{
+		const FSpriteSelectedSocket& S1 = *this;
+		const FSpriteSelectedSocket& S2 = *(FSpriteSelectedSocket*)(&OtherItem);
+
+		return (S1.SocketName == S2.SocketName) && (S1.PreviewComponentPtr == S2.PreviewComponentPtr);
+	}
+	else
+	{
+		return false;
+	}
+}
+
+void FSpriteSelectedSocket::ApplyDelta(const FVector2D& Delta)
+{
+	if (UPrimitiveComponent* PreviewComponent = PreviewComponentPtr.Get())
+	{
+		UObject* AssociatedAsset = const_cast<UObject*>(PreviewComponent->AdditionalStatObject());
+		if (UPaperSprite* Sprite = Cast<UPaperSprite>(AssociatedAsset))
+		{
+			if (FPaperSpriteSocket* Socket = Sprite->FindSocket(SocketName))
+			{
+				//@TODO: Currently sockets are in unflipped pivot space,
+				const FVector Delta3D_UU = (PaperAxisX * Delta.X) + (PaperAxisY * -Delta.Y);
+				const FVector Delta3D = Delta3D_UU * Sprite->GetPixelsPerUnrealUnit();
+				Socket->LocalTransform.SetLocation(Socket->LocalTransform.GetLocation() + Delta3D);
+			}
+		}
+	}
+}
+
+FVector FSpriteSelectedSocket::GetWorldPos() const
+{
+	if (UPrimitiveComponent* PreviewComponent = PreviewComponentPtr.Get())
+	{
+		return PreviewComponent->GetSocketLocation(SocketName);
+	}
+
+	return FVector::ZeroVector;
+}
+
+void FSpriteSelectedSocket::SplitEdge()
+{
+	// Nonsense operation on a socket, do nothing
+}
+
+//////////////////////////////////////////////////////////////////////////
+// FSocketEditingHelper
+
+void FSocketEditingHelper::DrawSockets(UPrimitiveComponent* PreviewComponent, const FSceneView* View, FPrimitiveDrawInterface* PDI)
+{
+	if (PreviewComponent != nullptr)
+	{
+		const bool bIsHitTesting = PDI->IsHitTesting();
+
+		const float DiamondSize = 5.0f;
+		const FColor DiamondColor(255, 128, 128);
+
+		TArray<FComponentSocketDescription> SocketList;
+		PreviewComponent->QuerySupportedSockets(/*out*/ SocketList);
+
+		for (const FComponentSocketDescription& Socket : SocketList)
+		{
+			if (bIsHitTesting)
+			{
+				TSharedPtr<FSpriteSelectedSocket> Data = MakeShareable(new FSpriteSelectedSocket);
+				Data->PreviewComponentPtr = PreviewComponent;
+				Data->SocketName = Socket.Name;
+				PDI->SetHitProxy(new HSpriteSelectableObjectHitProxy(Data));
+			}
+
+			const FMatrix SocketTM = PreviewComponent->GetSocketTransform(Socket.Name, RTS_World).ToMatrixWithScale();
+			DrawWireDiamond(PDI, SocketTM, DiamondSize, DiamondColor, SDPG_Foreground);
+
+			if (bIsHitTesting)
+			{
+				PDI->SetHitProxy(nullptr);
+			}
+		}
+	}
+}
+
+void FSocketEditingHelper::DrawSocketNames(UPrimitiveComponent* PreviewComponent, FViewport& Viewport, FSceneView& View, FCanvas& Canvas)
+{
+	if (PreviewComponent != nullptr)
+	{
+		const int32 HalfX = Viewport.GetSizeXY().X / 2;
+		const int32 HalfY = Viewport.GetSizeXY().Y / 2;
+
+		const FColor SocketNameColor(255, 196, 196);
+
+		TArray<FComponentSocketDescription> SocketList;
+		PreviewComponent->QuerySupportedSockets(/*out*/ SocketList);
+
+		for (const FComponentSocketDescription& Socket : SocketList)
+		{
+			const FVector SocketWorldPos = PreviewComponent->GetSocketLocation(Socket.Name);
+
+			const FPlane Proj = View.Project(SocketWorldPos);
+			if (Proj.W > 0.f)
+			{
+				const int32 XPos = HalfX + (HalfX * Proj.X);
+				const int32 YPos = HalfY + (HalfY * (-Proj.Y));
+
+				FCanvasTextItem Msg(FVector2D(XPos, YPos), FText::FromString(Socket.Name.ToString()), GEngine->GetMediumFont(), SocketNameColor);
+				Msg.EnableShadow(FLinearColor::Black);
+				Canvas.DrawItem(Msg);
+
+// 				//@TODO: Draws the current value of the rotation (probably want to keep this!)
+// 				if (bManipulating && StaticMeshEditorPtr.Pin()->GetSelectedSocket() == Socket)
+// 				{
+// 					// Figure out the text height
+// 					FTextSizingParameters Parameters(GEngine->GetSmallFont(), 1.0f, 1.0f);
+// 					UCanvas::CanvasStringSize(Parameters, *Socket->SocketName.ToString());
+// 					int32 YL = FMath::TruncToInt(Parameters.DrawYL);
+// 
+// 					DrawAngles(&Canvas, XPos, YPos + YL, 
+// 						Widget->GetCurrentAxis(), 
+// 						GetWidgetMode(),
+// 						Socket->RelativeRotation,
+// 						Socket->RelativeLocation);
+// 				}
+			}
+		}
+	}
+}
diff --git a/Engine/Plugins/2D/Paper2D/Source/Paper2DEditor/Private/PaperEditorShared/SocketEditing.h b/Engine/Plugins/2D/Paper2D/Source/Paper2DEditor/Private/PaperEditorShared/SocketEditing.h
new file mode 100644
--- /dev/null
+++ b/Engine/Plugins/2D/Paper2D/Source/Paper2DEditor/Private/PaperEditorShared/SocketEditing.h
@@ -0,0 +1,53 @@
+// Copyright 1998-2015 Epic Games, Inc. All Rights Reserved.
+
+#pragma once
+
+#include "AssetEditorSelectedItem.h"
+
+//////////////////////////////////////////////////////////////////////////
+// HSpriteSelectableObjectHitProxy
+
+struct HSpriteSelectableObjectHitProxy : public HHitProxy
+{
+	DECLARE_HIT_PROXY(PAPER2DEDITOR_API);
+
+	TSharedPtr<FSelectedItem> Data;
+
+	HSpriteSelectableObjectHitProxy(TSharedPtr<FSelectedItem> InData)
+		: HHitProxy(HPP_UI)
+		, Data(InData)
+	{
+	}
+};
+
+//////////////////////////////////////////////////////////////////////////
+// FSpriteSelectedSocket
+
+class FSpriteSelectedSocket : public FSelectedItem
+{
+public:
+	FName SocketName;
+	TWeakObjectPtr<UPrimitiveComponent> PreviewComponentPtr;
+
+	static const FName SocketTypeID;
+
+public:
+	FSpriteSelectedSocket();
+
+	// FSelectedItem interface
+	virtual bool Equals(const FSelectedItem& OtherItem) const override;
+	virtual void ApplyDelta(const FVector2D& Delta) override;
+	FVector GetWorldPos() const override;
+	virtual void SplitEdge() override;
+	// End of FSelectedItem interface
+};
+
+//////////////////////////////////////////////////////////////////////////
+// FSocketEditingHelper
+
+class FSocketEditingHelper
+{
+public:
+	static void DrawSockets(UPrimitiveComponent* PreviewComponent, const FSceneView* View, FPrimitiveDrawInterface* PDI);
+	static void DrawSocketNames(UPrimitiveComponent* PreviewComponent, FViewport& Viewport, FSceneView& View, FCanvas& Canvas);
+};
diff --git a/Engine/Plugins/2D/Paper2D/Source/Paper2DEditor/Private/SpriteEditor/SpriteEditorCommands.cpp b/Engine/Plugins/2D/Paper2D/Source/Paper2DEditor/Private/SpriteEditor/SpriteEditorCommands.cpp
--- a/Engine/Plugins/2D/Paper2D/Source/Paper2DEditor/Private/SpriteEditor/SpriteEditorCommands.cpp
+++ b/Engine/Plugins/2D/Paper2D/Source/Paper2DEditor/Private/SpriteEditor/SpriteEditorCommands.cpp
@@ -14,7 +14,7 @@ void FSpriteEditorCommands::RegisterCommands()
 	UI_COMMAND(SetShowBounds, "Bounds", "Toggles display of the bounds of the static mesh.", EUserInterfaceActionType::ToggleButton, FInputGesture());
 	UI_COMMAND(SetShowCollision, "Collision", "Toggles display of the simplified collision mesh of the static mesh, if one has been assigned.", EUserInterfaceActionType::ToggleButton, FInputGesture(EKeys::C, EModifierKey::Alt));
 
-	UI_COMMAND(SetShowSockets, "Sockets", "Displays the static mesh sockets.", EUserInterfaceActionType::ToggleButton, FInputGesture());
+	UI_COMMAND(SetShowSockets, "Sockets", "Displays the sprite sockets.", EUserInterfaceActionType::ToggleButton, FInputGesture());
 
 	UI_COMMAND(SetShowNormals, "Normals", "Toggles display of vertex normals in the Preview Pane.", EUserInterfaceActionType::ToggleButton, FInputGesture());
 	UI_COMMAND(SetShowPivot, "Pivot", "Display the pivot location of the sprite.", EUserInterfaceActionType::ToggleButton, FInputGesture());
diff --git a/Engine/Plugins/2D/Paper2D/Source/Paper2DEditor/Private/SpriteEditor/SpriteEditorSelections.h b/Engine/Plugins/2D/Paper2D/Source/Paper2DEditor/Private/SpriteEditor/SpriteEditorSelections.h
--- a/Engine/Plugins/2D/Paper2D/Source/Paper2DEditor/Private/SpriteEditor/SpriteEditorSelections.h
+++ b/Engine/Plugins/2D/Paper2D/Source/Paper2DEditor/Private/SpriteEditor/SpriteEditorSelections.h
@@ -2,6 +2,8 @@
 
 #pragma once
 
+#include "PaperEditorShared/AssetEditorSelectedItem.h"
+
 //////////////////////////////////////////////////////////////////////////
 // FSelectionTypes
 
@@ -11,75 +13,11 @@ public:
 	static const FName Vertex;
 	static const FName Edge;
 	static const FName Pivot;
-	static const FName Socket;
 	static const FName SourceRegion;
 private:
 	FSelectionTypes() {}
 };
 
-class FSpriteSelectedVertex;
-
-//////////////////////////////////////////////////////////////////////////
-// FSelectedItem
-
-class FSelectedItem
-{
-protected:
-	FSelectedItem(FName InTypeName)
-		: TypeName(InTypeName)
-	{
-	}
-
-protected:
-	FName TypeName;
-public:
-	virtual bool IsA(FName TestType) const
-	{
-		return TestType == TypeName;
-	}
-
-	virtual uint32 GetTypeHash() const
-	{
-		return 0;
-	}
-
-	virtual bool Equals(const FSelectedItem& OtherItem) const
-	{
-		return false;
-	}
-
-	virtual void ApplyDelta(const FVector2D& Delta)
-	{
-	}
-
-	//@TODO: Doesn't belong here in base!
-	virtual void SplitEdge()
-	{
-	}
-
-	virtual FVector GetWorldPos() const
-	{
-		return FVector::ZeroVector;
-	}
-
-	virtual const FSpriteSelectedVertex* CastSelectedVertex() const
-	{
-		return nullptr;
-	}
-
-	virtual ~FSelectedItem() {}
-};
-
-inline uint32 GetTypeHash(const FSelectedItem& Vertex)
-{
-	return Vertex.GetTypeHash();
-}
-
-inline bool operator==(const FSelectedItem& V1, const FSelectedItem& V2)
-{
-	return V1.Equals(V2);
-}
-
 //////////////////////////////////////////////////////////////////////////
 // FSpriteSelectedSourceRegion
 
@@ -438,69 +376,3 @@ public:
 		}
 	}
 };
-
-
-
-//////////////////////////////////////////////////////////////////////////
-// FSpriteSelectedSocket
-
-class FSpriteSelectedSocket : public FSelectedItem
-{
-public:
-	FName SocketName;
-	TWeakObjectPtr<UPaperSprite> SpritePtr;
-
-public:
-	FSpriteSelectedSocket()
-		: FSelectedItem(FSelectionTypes::Socket)
-	{
-	}
-
-	virtual bool Equals(const FSelectedItem& OtherItem) const override
-	{
-		if (OtherItem.IsA(FSelectionTypes::Socket))
-		{
-			const FSpriteSelectedSocket& S1 = *this;
-			const FSpriteSelectedSocket& S2 = *(FSpriteSelectedSocket*)(&OtherItem);
-
-			return (S1.SocketName == S2.SocketName) && (S1.SpritePtr == S2.SpritePtr);
-		}
-		else
-		{
-			return false;
-		}
-	}
-
-	//@TODO: Currently sockets are in unflipped pivot space,
-	virtual void ApplyDelta(const FVector2D& Delta) override
-	{
-		if (UPaperSprite* Sprite = SpritePtr.Get())
-		{
-			if (FPaperSpriteSocket* Socket = Sprite->FindSocket(SocketName))
-			{
-				const FVector Delta3D_UU = (PaperAxisX * Delta.X) + (PaperAxisY * -Delta.Y);
-				const FVector Delta3D = Delta3D_UU * Sprite->GetPixelsPerUnrealUnit();
-				Socket->LocalTransform.SetLocation(Socket->LocalTransform.GetLocation() + Delta3D);
-			}
-		}
-	}
-
-	FVector GetWorldPos() const override
-	{
-		if (UPaperSprite* Sprite = SpritePtr.Get())
-		{
-			if (FPaperSpriteSocket* Socket = Sprite->FindSocket(SocketName))
-			{
- 				const FVector PivotSpacePos = Socket->LocalTransform.GetLocation();
-				return Sprite->GetPivotToWorld().TransformPosition(PivotSpacePos);
-			}
-		}
-
-		return FVector::ZeroVector;
-	}
-
-	virtual void SplitEdge() override
-	{
-		// Nonsense operation on a socket, do nothing
-	}
-};
diff --git a/Engine/Plugins/2D/Paper2D/Source/Paper2DEditor/Private/SpriteEditor/SpriteEditorViewportClient.cpp b/Engine/Plugins/2D/Paper2D/Source/Paper2DEditor/Private/SpriteEditor/SpriteEditorViewportClient.cpp
--- a/Engine/Plugins/2D/Paper2D/Source/Paper2DEditor/Private/SpriteEditor/SpriteEditorViewportClient.cpp
+++ b/Engine/Plugins/2D/Paper2D/Source/Paper2DEditor/Private/SpriteEditor/SpriteEditorViewportClient.cpp
@@ -13,6 +13,7 @@
 #include "AssetRegistryModule.h"
 #include "CanvasTypes.h"
 #include "CanvasItem.h"
+#include "PaperEditorShared/SocketEditing.h"
 
 #define LOCTEXT_NAMESPACE "SpriteEditor"
 
@@ -41,26 +42,8 @@ static int32 GetVertexIndexFromHash(int32 PolygonVertexHash)
 const FName FSelectionTypes::Vertex(TEXT("Vertex"));
 const FName FSelectionTypes::Edge(TEXT("Edge"));
 const FName FSelectionTypes::Pivot(TEXT("Pivot"));
-const FName FSelectionTypes::Socket(TEXT("Socket"));
 const FName FSelectionTypes::SourceRegion(TEXT("SourceRegion"));
 
-//////////////////////////////////////////////////////////////////////////
-// HSpriteSelectableObjectHitProxy
-
-struct HSpriteSelectableObjectHitProxy : public HHitProxy
-{
-	DECLARE_HIT_PROXY( PAPER2DEDITOR_API );
-
-	TSharedPtr<FSelectedItem> Data;
-
-	HSpriteSelectableObjectHitProxy(TSharedPtr<FSelectedItem> InData)
-		: HHitProxy(HPP_UI)
-		, Data(InData)
-	{
-	}
-};
-IMPLEMENT_HIT_PROXY(HSpriteSelectableObjectHitProxy, HHitProxy);
-
 //////////////////////////////////////////////////////////////////////////
 // FSpriteEditorViewportClient
 
@@ -593,79 +576,6 @@ void FSpriteEditorViewportClient::DrawBoundsAsText(FViewport& InViewport, FScene
 	YPos += 18;
 }
 
-void FSpriteEditorViewportClient::DrawSockets(const FSceneView* View, FPrimitiveDrawInterface* PDI)
-{
-	UPaperSprite* Sprite = GetSpriteBeingEdited();
-
-	const bool bIsHitTesting = PDI->IsHitTesting();
-
-	const float DiamondSize = 5.0f;
-	const FColor DiamondColor(255, 128, 128);
-
-	const FTransform SpritePivotToWorld = Sprite->GetPivotToWorld();
-	for (int32 SocketIndex = 0; SocketIndex < Sprite->Sockets.Num(); SocketIndex++)
-	{
-		const FPaperSpriteSocket& Socket = Sprite->Sockets[SocketIndex];
-		const FMatrix SocketTM = (Socket.LocalTransform * SpritePivotToWorld).ToMatrixWithScale();
-		
-		if (bIsHitTesting)
-		{
-			TSharedPtr<FSpriteSelectedSocket> Data = MakeShareable(new FSpriteSelectedSocket);
-			Data->SpritePtr = Sprite;
-			Data->SocketName = Socket.SocketName;
-			PDI->SetHitProxy(new HSpriteSelectableObjectHitProxy(Data));
-		}
-		
-		DrawWireDiamond(PDI, SocketTM, DiamondSize, DiamondColor, SDPG_Foreground);
-		
-		PDI->SetHitProxy(nullptr);
-	}
-}
-
-void FSpriteEditorViewportClient::DrawSocketNames(FViewport& InViewport, FSceneView& View, FCanvas& Canvas)
-{
-	UPaperSprite* Sprite = GetSpriteBeingEdited();
-
-	const int32 HalfX = Viewport->GetSizeXY().X / 2;
-	const int32 HalfY = Viewport->GetSizeXY().Y / 2;
-
-	const FColor SocketNameColor(255,196,196);
-
-	// Draw socket names if desired.
-	const FTransform SpritePivotToWorld = Sprite->GetPivotToWorld();
-	for (int32 SocketIndex = 0; SocketIndex < Sprite->Sockets.Num(); SocketIndex++)
-	{
-		const FPaperSpriteSocket& Socket = Sprite->Sockets[SocketIndex];
-
-		const FVector SocketWorldPos = SpritePivotToWorld.TransformPosition(Socket.LocalTransform.GetLocation());
-		
-		const FPlane Proj = View.Project(SocketWorldPos);
-		if (Proj.W > 0.f)
-		{
-			const int32 XPos = HalfX + (HalfX * Proj.X);
-			const int32 YPos = HalfY + (HalfY * (-Proj.Y));
-
-			FCanvasTextItem Msg(FVector2D(XPos, YPos), FText::FromString(Socket.SocketName.ToString()), GEngine->GetMediumFont(), SocketNameColor);
-			Canvas.DrawItem(Msg);
-
-			//@TODO: Draws the current value of the rotation (probably want to keep this!)
-// 			if (bManipulating && StaticMeshEditorPtr.Pin()->GetSelectedSocket() == Socket)
-// 			{
-// 				// Figure out the text height
-// 				FTextSizingParameters Parameters(GEngine->GetSmallFont(), 1.0f, 1.0f);
-// 				UCanvas::CanvasStringSize(Parameters, *Socket->SocketName.ToString());
-// 				int32 YL = FMath::TruncToInt(Parameters.DrawYL);
-// 
-// 				DrawAngles(&Canvas, XPos, YPos + YL, 
-// 					Widget->GetCurrentAxis(), 
-// 					GetWidgetMode(),
-// 					Socket->RelativeRotation,
-// 					Socket->RelativeLocation);
-// 			}
-		}
-	}
-}
-
 void FSpriteEditorViewportClient::DrawCanvas(FViewport& Viewport, FSceneView& View, FCanvas& Canvas)
 {
 	FEditorViewportClient::DrawCanvas(Viewport, View, Canvas);
@@ -815,7 +725,7 @@ void FSpriteEditorViewportClient::DrawCanvas(FViewport& Viewport, FSceneView& Vi
 
 	if (bShowSockets)
 	{
-		DrawSocketNames(Viewport, View, Canvas);
+		FSocketEditingHelper::DrawSocketNames(RenderSpriteComponent, Viewport, View, Canvas);
 	}
 }
 
@@ -832,7 +742,7 @@ void FSpriteEditorViewportClient::Draw(const FSceneView* View, FPrimitiveDrawInt
 
 	if (bShowSockets)
 	{
-		DrawSockets(View, PDI);
+		FSocketEditingHelper::DrawSockets(RenderSpriteComponent, View, PDI);
 	}
 }
 
diff --git a/Engine/Plugins/2D/Paper2D/Source/Paper2DEditor/Private/SpriteEditor/SpriteEditorViewportClient.h b/Engine/Plugins/2D/Paper2D/Source/Paper2DEditor/Private/SpriteEditor/SpriteEditorViewportClient.h
--- a/Engine/Plugins/2D/Paper2D/Source/Paper2DEditor/Private/SpriteEditor/SpriteEditorViewportClient.h
+++ b/Engine/Plugins/2D/Paper2D/Source/Paper2DEditor/Private/SpriteEditor/SpriteEditorViewportClient.h
@@ -216,8 +216,6 @@ private:
 	void DrawGeometry(FViewport& InViewport, FSceneView& View, FCanvas& Canvas, const FSpritePolygonCollection& Geometry, const FLinearColor& GeometryVertexColor, const FLinearColor& NegativeGeometryVertexColor, bool bIsRenderGeometry);
 	void DrawGeometryStats(FViewport& InViewport, FSceneView& View, FCanvas& Canvas, const FSpritePolygonCollection& Geometry, bool bIsRenderGeometry, int32& YPos);
 	void DrawRenderStats(FViewport& InViewport, FSceneView& View, FCanvas& Canvas, class UPaperSprite* Sprite, int32& YPos);
-	void DrawSockets(const FSceneView* View, FPrimitiveDrawInterface* PDI);
-	void DrawSocketNames(FViewport& InViewport, FSceneView& View, FCanvas& Canvas);
 	void DrawSourceRegion(FViewport& InViewport, FSceneView& View, FCanvas& Canvas, const FLinearColor& GeometryVertexColor);
 	void DrawRelatedSprites(FViewport& InViewport, FSceneView& View, FCanvas& Canvas, const FLinearColor& GeometryVertexColor);
 	void DrawMarquee(FViewport& InViewport, FSceneView& View, FCanvas& Canvas, const FLinearColor& MarqueeColor);
