Paper2D: Add support for rotation and scaling sockets via the widget in the sprite editor viewport
[UE-10904]

diff --git a/Engine/Plugins/2D/Paper2D/Source/Paper2DEditor/Private/PaperEditorShared/AssetEditorSelectedItem.h b/Engine/Plugins/2D/Paper2D/Source/Paper2DEditor/Private/PaperEditorShared/AssetEditorSelectedItem.h
--- a/Engine/Plugins/2D/Paper2D/Source/Paper2DEditor/Private/PaperEditorShared/AssetEditorSelectedItem.h
+++ b/Engine/Plugins/2D/Paper2D/Source/Paper2DEditor/Private/PaperEditorShared/AssetEditorSelectedItem.h
@@ -31,7 +31,7 @@ public:
 		return false;
 	}
 
-	virtual void ApplyDelta(const FVector2D& Delta)
+	virtual void ApplyDelta(const FVector2D& Delta, const FRotator& Rotation, const FVector& Scale3D, FWidget::EWidgetMode MoveMode)
 	{
 	}
 
diff --git a/Engine/Plugins/2D/Paper2D/Source/Paper2DEditor/Private/PaperEditorShared/SocketEditing.cpp b/Engine/Plugins/2D/Paper2D/Source/Paper2DEditor/Private/PaperEditorShared/SocketEditing.cpp
--- a/Engine/Plugins/2D/Paper2D/Source/Paper2DEditor/Private/PaperEditorShared/SocketEditing.cpp
+++ b/Engine/Plugins/2D/Paper2D/Source/Paper2DEditor/Private/PaperEditorShared/SocketEditing.cpp
@@ -1,4 +1,4 @@
-// Copyright 1998-2015 Epic Games, Inc. All Rights Reserved.
+	// Copyright 1998-2015 Epic Games, Inc. All Rights Reserved.
 
 #include "Paper2DEditorPrivatePCH.h"
 #include "SocketEditing.h"
@@ -33,7 +33,7 @@ bool FSpriteSelectedSocket::Equals(const FSelectedItem& OtherItem) const
 	}
 }
 
-void FSpriteSelectedSocket::ApplyDelta(const FVector2D& Delta)
+void FSpriteSelectedSocket::ApplyDelta(const FVector2D& Delta, const FRotator& Rotation, const FVector& Scale3D, FWidget::EWidgetMode MoveMode)
 {
 	if (UPrimitiveComponent* PreviewComponent = PreviewComponentPtr.Get())
 	{
@@ -42,10 +42,42 @@ void FSpriteSelectedSocket::ApplyDelta(const FVector2D& Delta)
 		{
 			if (FPaperSpriteSocket* Socket = Sprite->FindSocket(SocketName))
 			{
-				//@TODO: Currently sockets are in unflipped pivot space,
-				const FVector Delta3D_UU = (PaperAxisX * Delta.X) + (PaperAxisY * -Delta.Y);
-				const FVector Delta3D = Delta3D_UU * Sprite->GetPixelsPerUnrealUnit();
-				Socket->LocalTransform.SetLocation(Socket->LocalTransform.GetLocation() + Delta3D);
+				const bool bDoRotation = (MoveMode == FWidget::WM_Rotate) || (MoveMode == FWidget::WM_TranslateRotateZ);
+				const bool bDoTranslation = (MoveMode == FWidget::WM_Translate) || (MoveMode == FWidget::WM_TranslateRotateZ);
+				const bool bDoScale = MoveMode == FWidget::WM_Scale;
+
+				if (bDoTranslation)
+				{
+					//@TODO: Currently sockets are in unflipped pivot space,
+					const FVector Delta3D_UU = (PaperAxisX * Delta.X) + (PaperAxisY * -Delta.Y);
+					const FVector Delta3D = Delta3D_UU * Sprite->GetPixelsPerUnrealUnit();
+					Socket->LocalTransform.SetLocation(Socket->LocalTransform.GetLocation() + Delta3D);
+				}
+
+				if (bDoRotation)
+				{
+					const FRotator CurrentRot = Socket->LocalTransform.GetRotation().Rotator();
+					FRotator SocketWinding;
+					FRotator SocketRotRemainder;
+					CurrentRot.GetWindingAndRemainder(SocketWinding, SocketRotRemainder);
+
+					const FQuat ActorQ = SocketRotRemainder.Quaternion();
+					const FQuat DeltaQ = Rotation.Quaternion();
+					const FQuat ResultQ = DeltaQ * ActorQ;
+					const FRotator NewSocketRotRem = FRotator( ResultQ );
+					FRotator DeltaRot = NewSocketRotRem - SocketRotRemainder;
+					DeltaRot.Normalize();
+
+					const FRotator NewRotation(CurrentRot + DeltaRot);
+					Socket->LocalTransform.SetRotation(NewRotation.Quaternion());
+				}
+				
+				if (bDoScale)
+				{
+					const FVector4 LocalSpaceScaleOffset = Socket->LocalTransform.TransformVector(Scale3D);
+
+					Socket->LocalTransform.SetScale3D(Socket->LocalTransform.GetScale3D() + LocalSpaceScaleOffset);
+				}
 			}
 		}
 	}
diff --git a/Engine/Plugins/2D/Paper2D/Source/Paper2DEditor/Private/PaperEditorShared/SocketEditing.h b/Engine/Plugins/2D/Paper2D/Source/Paper2DEditor/Private/PaperEditorShared/SocketEditing.h
--- a/Engine/Plugins/2D/Paper2D/Source/Paper2DEditor/Private/PaperEditorShared/SocketEditing.h
+++ b/Engine/Plugins/2D/Paper2D/Source/Paper2DEditor/Private/PaperEditorShared/SocketEditing.h
@@ -36,7 +36,7 @@ public:
 
 	// FSelectedItem interface
 	virtual bool Equals(const FSelectedItem& OtherItem) const override;
-	virtual void ApplyDelta(const FVector2D& Delta) override;
+	virtual void ApplyDelta(const FVector2D& Delta, const FRotator& Rotation, const FVector& Scale3D, FWidget::EWidgetMode MoveMode) override;
 	FVector GetWorldPos() const override;
 	virtual void SplitEdge() override;
 	// End of FSelectedItem interface
diff --git a/Engine/Plugins/2D/Paper2D/Source/Paper2DEditor/Private/SpriteEditor/SpriteEditorSelections.h b/Engine/Plugins/2D/Paper2D/Source/Paper2DEditor/Private/SpriteEditor/SpriteEditorSelections.h
--- a/Engine/Plugins/2D/Paper2D/Source/Paper2DEditor/Private/SpriteEditor/SpriteEditorSelections.h
+++ b/Engine/Plugins/2D/Paper2D/Source/Paper2DEditor/Private/SpriteEditor/SpriteEditorSelections.h
@@ -171,9 +171,12 @@ public:
 		return Result;
 	}
 
-	virtual void ApplyDelta(const FVector2D& Delta) override
+	virtual void ApplyDelta(const FVector2D& Delta, const FRotator& Rotation, const FVector& Scale3D, FWidget::EWidgetMode MoveMode) override
 	{
-		ApplyDeltaIndexed(Delta, VertexIndex);
+		if (MoveMode == FWidget::WM_Translate)
+		{
+			ApplyDeltaIndexed(Delta, VertexIndex);
+		}
 	}
 
 	FVector GetWorldPos() const override
@@ -294,9 +297,12 @@ public:
 		return Result;
 	}
 
-	virtual void ApplyDelta(const FVector2D& Delta) override
+	virtual void ApplyDelta(const FVector2D& Delta, const FRotator& Rotation, const FVector& Scale3D, FWidget::EWidgetMode MoveMode) override
 	{
-		ApplyDeltaIndexed(Delta, VertexIndex);
+		if (MoveMode == FWidget::WM_Translate)
+		{
+			ApplyDeltaIndexed(Delta, VertexIndex);
+		}
 	}
 
 	FVector GetWorldPos() const override
@@ -340,7 +346,7 @@ public:
 		}
 	}
 
-	virtual void ApplyDelta(const FVector2D& Delta) override
+	virtual void ApplyDelta(const FVector2D& Delta, const FRotator& Rotation, const FVector& Scale3D, FWidget::EWidgetMode MoveMode) override
 	{
 		ApplyDeltaIndexed(Delta, VertexIndex);
 		ApplyDeltaIndexed(Delta, VertexIndex+1);
diff --git a/Engine/Plugins/2D/Paper2D/Source/Paper2DEditor/Private/SpriteEditor/SpriteEditorViewportClient.cpp b/Engine/Plugins/2D/Paper2D/Source/Paper2DEditor/Private/SpriteEditor/SpriteEditorViewportClient.cpp
--- a/Engine/Plugins/2D/Paper2D/Source/Paper2DEditor/Private/SpriteEditor/SpriteEditorViewportClient.cpp
+++ b/Engine/Plugins/2D/Paper2D/Source/Paper2DEditor/Private/SpriteEditor/SpriteEditorViewportClient.cpp
@@ -58,7 +58,7 @@ FSpriteEditorViewportClient::FSpriteEditorViewportClient(TWeakPtr<FSpriteEditor>
 
 	SetRealtime(true);
 
-	WidgetMode = FWidget::WM_Translate;
+	DesiredWidgetMode = FWidget::WM_Translate;
 	bManipulating = false;
 	bManipulationDirtiedSomething = false;
 	ScopedTransaction = nullptr;
@@ -1223,6 +1223,8 @@ bool FSpriteEditorViewportClient::InputWidgetDelta(FViewport* Viewport, EAxisLis
 	{
 		bHandled = true;
 
+		const FWidget::EWidgetMode MoveMode = GetWidgetMode();
+
 		// Negate Y because vertices are in source texture space, not world space
 		const FVector2D Drag2D(FVector::DotProduct(Drag, PaperAxisX), -FVector::DotProduct(Drag, PaperAxisY));
 
@@ -1230,7 +1232,7 @@ bool FSpriteEditorViewportClient::InputWidgetDelta(FViewport* Viewport, EAxisLis
 		for (auto SelectionIt = SelectionSet.CreateConstIterator(); SelectionIt; ++SelectionIt)
 		{
 			TSharedPtr<FSelectedItem> SelectedItem = *SelectionIt;
-			SelectedItem->ApplyDelta(Drag2D);
+			SelectedItem->ApplyDelta(Drag2D, Rot, Scale, MoveMode);
 		}
 
 		if (SelectionSet.Num() > 0)
@@ -1272,7 +1274,26 @@ void FSpriteEditorViewportClient::TrackingStopped()
 
 FWidget::EWidgetMode FSpriteEditorViewportClient::GetWidgetMode() const
 {
-	return (SelectionSet.Num() > 0) ? FWidget::WM_Translate : FWidget::WM_None;
+	return (SelectionSet.Num() > 0) ? DesiredWidgetMode : FWidget::WM_None;
+}
+
+void FSpriteEditorViewportClient::SetWidgetMode(FWidget::EWidgetMode NewMode)
+{
+	DesiredWidgetMode = NewMode;
+}
+
+bool FSpriteEditorViewportClient::CanSetWidgetMode(FWidget::EWidgetMode NewMode) const
+{
+	return CanCycleWidgetMode();
+}
+
+bool FSpriteEditorViewportClient::CanCycleWidgetMode() const
+{
+	if (!Widget->IsDragging())
+	{
+		return (SelectionSet.Num() > 0);
+	}
+	return false;
 }
 
 FVector FSpriteEditorViewportClient::GetWidgetLocation() const
diff --git a/Engine/Plugins/2D/Paper2D/Source/Paper2DEditor/Private/SpriteEditor/SpriteEditorViewportClient.h b/Engine/Plugins/2D/Paper2D/Source/Paper2DEditor/Private/SpriteEditor/SpriteEditorViewportClient.h
--- a/Engine/Plugins/2D/Paper2D/Source/Paper2DEditor/Private/SpriteEditor/SpriteEditorViewportClient.h
+++ b/Engine/Plugins/2D/Paper2D/Source/Paper2DEditor/Private/SpriteEditor/SpriteEditorViewportClient.h
@@ -61,6 +61,9 @@ public:
 	virtual void TrackingStarted(const struct FInputEventState& InInputState, bool bIsDragging, bool bNudge) override;
 	virtual void TrackingStopped() override;
 	virtual FWidget::EWidgetMode GetWidgetMode() const override;
+	virtual void SetWidgetMode(FWidget::EWidgetMode NewMode) override;
+	virtual bool CanSetWidgetMode(FWidget::EWidgetMode NewMode) const override;
+	virtual bool CanCycleWidgetMode() const override;
 	virtual FVector GetWidgetLocation() const;
 	virtual FMatrix GetWidgetCoordSystem() const;
 	virtual ECoordSystem GetWidgetCoordSystemSpace() const;
@@ -147,7 +150,7 @@ private:
 	UPaperSpriteComponent* RenderSpriteComponent;
 
 	// Widget mode
-	FWidget::EWidgetMode WidgetMode;
+	FWidget::EWidgetMode DesiredWidgetMode;
 
 	// Are we currently manipulating something?
 	bool bManipulating;
