Paper2D: Added some missing includes to SFlipbookEditorViewportToolbar to fix non-unity build

diff --git a/Engine/Plugins/2D/Paper2D/Source/Paper2DEditor/Private/FlipbookEditor/SFlipbookEditorViewportToolbar.cpp b/Engine/Plugins/2D/Paper2D/Source/Paper2DEditor/Private/FlipbookEditor/SFlipbookEditorViewportToolbar.cpp
--- a/Engine/Plugins/2D/Paper2D/Source/Paper2DEditor/Private/FlipbookEditor/SFlipbookEditorViewportToolbar.cpp
+++ b/Engine/Plugins/2D/Paper2D/Source/Paper2DEditor/Private/FlipbookEditor/SFlipbookEditorViewportToolbar.cpp
@@ -2,6 +2,8 @@
 
 #include "Paper2DEditorPrivatePCH.h"
 #include "SFlipbookEditorViewportToolbar.h"
+#include "SEditorViewport.h"
+#include "FlipbookEditorCommands.h"
 
 #define LOCTEXT_NAMESPACE "SFlipbookEditorViewportToolbar"
 
