Paper2D: Fix display of custom texture rows (they were getting clipped and had an extra options button)

diff --git a/Engine/Plugins/2D/Paper2D/Source/Paper2DEditor/Private/SpriteEditor/SpriteDetailsCustomization.cpp b/Engine/Plugins/2D/Paper2D/Source/Paper2DEditor/Private/SpriteEditor/SpriteDetailsCustomization.cpp
--- a/Engine/Plugins/2D/Paper2D/Source/Paper2DEditor/Private/SpriteEditor/SpriteDetailsCustomization.cpp
+++ b/Engine/Plugins/2D/Paper2D/Source/Paper2DEditor/Private/SpriteEditor/SpriteDetailsCustomization.cpp
@@ -240,6 +240,7 @@ void FSpriteDetailsCustomization::BuildTextureSection(IDetailCategoryBuilder& Sp
 			CreateTextureNameWidget(SourceTextureProperty, SourceTextureOverrideLabel)
 		]
 		.ValueContent()
+		.MaxDesiredWidth(TOptional<float>())
 		[
 			SourceTextureProperty->CreatePropertyValueWidget()
 		];
@@ -267,12 +268,15 @@ void FSpriteDetailsCustomization::GenerateAdditionalTextureWidget(TSharedRef<IPr
 
 	TextureRow.DisplayName(SlotDesc);
 
+	TextureRow.ShowPropertyButtons(false);
+
 	TextureRow.CustomWidget(false)
 		.NameContent()
 		[
 			CreateTextureNameWidget(PropertyHandle, ExtraText)
 		]
 		.ValueContent()
+		.MaxDesiredWidth(TOptional<float>())
 		[
 			PropertyHandle->CreatePropertyValueWidget()
 		];
