Sprite Editor: Minor sprite socket editing improvements
- Prevent showing sprite sockets in 'Source Region Edit Mode' (they were not useful there and were being shown in the wrong coordinate space)
- Clear selection set when switching modes to avoid phantom gizmos that can still be interacted with
- Allow deselecting sprite sockets in 'View mode'
[UE-11581] [UE-10930]

diff --git a/Engine/Plugins/2D/Paper2D/Source/Paper2DEditor/Private/SpriteEditor/SpriteEditor.cpp b/Engine/Plugins/2D/Paper2D/Source/Paper2DEditor/Private/SpriteEditor/SpriteEditor.cpp
--- a/Engine/Plugins/2D/Paper2D/Source/Paper2DEditor/Private/SpriteEditor/SpriteEditor.cpp
+++ b/Engine/Plugins/2D/Paper2D/Source/Paper2DEditor/Private/SpriteEditor/SpriteEditor.cpp
@@ -501,16 +501,6 @@ void FSpriteEditor::ExtendToolbar()
 				ToolbarBuilder.AddToolBarButton(FSpriteEditorCommands::Get().SnapAllVertices);
 			}
 			ToolbarBuilder.EndSection();
-
-// 			ToolbarBuilder.BeginSection("Mode");
-// 			{
-// 				ToolbarBuilder.AddToolBarButton(FSpriteEditorCommands::Get().EnterViewMode);
-// 				ToolbarBuilder.AddToolBarButton(FSpriteEditorCommands::Get().EnterSourceRegionEditMode);
-// 				ToolbarBuilder.AddToolBarButton(FSpriteEditorCommands::Get().EnterCollisionEditMode);
-// 				ToolbarBuilder.AddToolBarButton(FSpriteEditorCommands::Get().EnterRenderingEditMode);
-// 				//@TODO: PAPER2D: Re-enable once it does something: ToolbarBuilder.AddToolBarButton(FSpriteEditorCommands::Get().EnterAddSpriteMode);
-// 			}
-// 			ToolbarBuilder.EndSection();
 		}
 	};
 
diff --git a/Engine/Plugins/2D/Paper2D/Source/Paper2DEditor/Private/SpriteEditor/SpriteEditorViewportClient.cpp b/Engine/Plugins/2D/Paper2D/Source/Paper2DEditor/Private/SpriteEditor/SpriteEditorViewportClient.cpp
--- a/Engine/Plugins/2D/Paper2D/Source/Paper2DEditor/Private/SpriteEditor/SpriteEditorViewportClient.cpp
+++ b/Engine/Plugins/2D/Paper2D/Source/Paper2DEditor/Private/SpriteEditor/SpriteEditorViewportClient.cpp
@@ -370,7 +370,7 @@ void FSpriteEditorViewportClient::DrawGeometry(FViewport& InViewport, FSceneView
 			const FVector2D ScreenPos = TextureSpaceToScreenSpace(View, Polygon.Vertices[VertexIndex]);
 			const FVector2D NextScreenPos = TextureSpaceToScreenSpace(View, Polygon.Vertices[NextVertexIndex]);
 
-			bool IsEdgeSelected = IsPolygonVertexSelected(PolygonIndex, VertexIndex) && IsPolygonVertexSelected(PolygonIndex, NextVertexIndex);
+			const bool IsEdgeSelected = IsPolygonVertexSelected(PolygonIndex, VertexIndex) && IsPolygonVertexSelected(PolygonIndex, NextVertexIndex);
 
 
 			// Draw the normal tick
@@ -723,7 +723,7 @@ void FSpriteEditorViewportClient::DrawCanvas(FViewport& Viewport, FSceneView& Vi
 		DrawMarquee(Viewport, View, Canvas, MarqueeColor);
 	}
 
-	if (bShowSockets)
+	if (bShowSockets && !IsInSourceRegionEditMode())
 	{
 		FSocketEditingHelper::DrawSocketNames(RenderSpriteComponent, Viewport, View, Canvas);
 	}
@@ -740,7 +740,7 @@ void FSpriteEditorViewportClient::Draw(const FSceneView* View, FPrimitiveDrawInt
 		FUnrealEdUtils::DrawWidget(View, PDI, RenderSpriteComponent->ComponentToWorld.ToMatrixWithScale(), 0, 0, EAxisList::XZ, EWidgetMovementMode::WMM_Translate);
 	}
 
-	if (bShowSockets)
+	if (bShowSockets && !IsInSourceRegionEditMode())
 	{
 		FSocketEditingHelper::DrawSockets(RenderSpriteComponent, View, PDI);
 	}
@@ -890,7 +890,7 @@ void FSpriteEditorViewportClient::ProcessClick(FSceneView& View, HHitProxy* HitP
 	const bool bInsertVertexModifier = bIsShiftKeyDown;
 	HSpriteSelectableObjectHitProxy* SelectedItemProxy = HitProxyCast<HSpriteSelectableObjectHitProxy>(HitProxy);
 
-	if (bAllowSelectVertex && SelectedItemProxy)
+	if (bAllowSelectVertex && (SelectedItemProxy != nullptr))
 	{
 		if (!bClearSelectionModifier)
 		{
@@ -1034,6 +1034,11 @@ void FSpriteEditorViewportClient::ProcessClick(FSceneView& View, HHitProxy* HitP
 				}
 			}
 		}
+		else if (!IsEditingGeometry())
+		{
+			// Clicked on the background (missed any proxies), deselect the socket or whatever was selected
+			ClearSelectionSet();
+		}
 
 		if (!bHandled)
 		{
@@ -1525,6 +1530,14 @@ void FSpriteEditorViewportClient::SnapAllVerticesToPixelGrid()
 	EndTransaction();
 }
 
+void FSpriteEditorViewportClient::InternalActivateNewMode(ESpriteEditorMode::Type NewMode)
+{
+	CurrentMode = NewMode;
+	ClearSelectionSet();
+	ResetMarqueeTracking();
+	ResetAddPolyonMode();
+}
+
 FSpritePolygonCollection* FSpriteEditorViewportClient::GetGeometryBeingEdited() const
 {
 	UPaperSprite* Sprite = GetSpriteBeingEdited();
diff --git a/Engine/Plugins/2D/Paper2D/Source/Paper2DEditor/Private/SpriteEditor/SpriteEditorViewportClient.h b/Engine/Plugins/2D/Paper2D/Source/Paper2DEditor/Private/SpriteEditor/SpriteEditorViewportClient.h
--- a/Engine/Plugins/2D/Paper2D/Source/Paper2DEditor/Private/SpriteEditor/SpriteEditorViewportClient.h
+++ b/Engine/Plugins/2D/Paper2D/Source/Paper2DEditor/Private/SpriteEditor/SpriteEditorViewportClient.h
@@ -85,11 +85,11 @@ public:
 	void ToggleShowMeshEdges();
 	bool IsShowMeshEdgesChecked() const;
 
-	void EnterViewMode() { CurrentMode = ESpriteEditorMode::ViewMode; ResetMarqueeTracking(); ResetAddPolyonMode(); }
-	void EnterSourceRegionEditMode() { CurrentMode = ESpriteEditorMode::EditSourceRegionMode; ResetMarqueeTracking(); UpdateRelatedSpritesList(); ResetAddPolyonMode(); }
-	void EnterCollisionEditMode() { CurrentMode = ESpriteEditorMode::EditCollisionMode; ResetMarqueeTracking(); ResetAddPolyonMode(); }
-	void EnterRenderingEditMode() { CurrentMode = ESpriteEditorMode::EditRenderingGeomMode; ResetMarqueeTracking(); ResetAddPolyonMode(); }
-	void EnterAddSpriteMode() { CurrentMode = ESpriteEditorMode::AddSpriteMode; ResetMarqueeTracking(); ResetAddPolyonMode(); }
+	void EnterViewMode() { InternalActivateNewMode(ESpriteEditorMode::ViewMode); }
+	void EnterSourceRegionEditMode() { InternalActivateNewMode(ESpriteEditorMode::EditSourceRegionMode); UpdateRelatedSpritesList(); }
+	void EnterCollisionEditMode() { InternalActivateNewMode(ESpriteEditorMode::EditCollisionMode); }
+	void EnterRenderingEditMode() { InternalActivateNewMode(ESpriteEditorMode::EditRenderingGeomMode); }
+	void EnterAddSpriteMode() { InternalActivateNewMode(ESpriteEditorMode::AddSpriteMode); }
 
 	bool IsInViewMode() const { return CurrentMode == ESpriteEditorMode::ViewMode; }
 	bool IsInSourceRegionEditMode() const { return CurrentMode == ESpriteEditorMode::EditSourceRegionMode; }
@@ -240,7 +240,10 @@ private:
 	void ResetMarqueeTracking();
 	bool ConvertMarqueeToSourceTextureSpace(/*out*/FVector2D& OutStartPos, /*out*/FVector2D& OutDimension);
 	void SelectVerticesInMarquee(bool bAddToSelection);
-	
+
+	// Activates a new mode, clearing selection set, etc...
+	void InternalActivateNewMode(ESpriteEditorMode::Type NewMode);
+
 	// Can return null
 	FSpritePolygonCollection* GetGeometryBeingEdited() const;
 };
