Paper2D: Add a detail customization for sprite and flipbook components to push the sprite category above the material category, matching their relative importance
[UE-10088]

diff --git a/Engine/Plugins/2D/Paper2D/Source/Paper2DEditor/Private/FlipbookEditor/FlipbookComponentDetailsCustomization.cpp b/Engine/Plugins/2D/Paper2D/Source/Paper2DEditor/Private/FlipbookEditor/FlipbookComponentDetailsCustomization.cpp
new file mode 100644
--- /dev/null
+++ b/Engine/Plugins/2D/Paper2D/Source/Paper2DEditor/Private/FlipbookEditor/FlipbookComponentDetailsCustomization.cpp
@@ -0,0 +1,24 @@
+// Copyright 1998-2015 Epic Games, Inc. All Rights Reserved.
+
+#include "Paper2DEditorPrivatePCH.h"
+#include "FlipbookComponentDetailsCustomization.h"
+
+#define LOCTEXT_NAMESPACE "FlipbookEditor"
+
+//////////////////////////////////////////////////////////////////////////
+// FFlipbookComponentDetailsCustomization
+
+TSharedRef<IDetailCustomization> FFlipbookComponentDetailsCustomization::MakeInstance()
+{
+	return MakeShareable(new FFlipbookComponentDetailsCustomization);
+}
+
+void FFlipbookComponentDetailsCustomization::CustomizeDetails(IDetailLayoutBuilder& DetailBuilder)
+{
+	// Create a category so this is displayed early in the properties
+	DetailBuilder.EditCategory("Sprite", FText::GetEmpty(), ECategoryPriority::Important);
+}
+
+//////////////////////////////////////////////////////////////////////////
+
+#undef LOCTEXT_NAMESPACE
diff --git a/Engine/Plugins/2D/Paper2D/Source/Paper2DEditor/Private/FlipbookEditor/FlipbookComponentDetailsCustomization.h b/Engine/Plugins/2D/Paper2D/Source/Paper2DEditor/Private/FlipbookEditor/FlipbookComponentDetailsCustomization.h
new file mode 100644
--- /dev/null
+++ b/Engine/Plugins/2D/Paper2D/Source/Paper2DEditor/Private/FlipbookEditor/FlipbookComponentDetailsCustomization.h
@@ -0,0 +1,19 @@
+// Copyright 1998-2015 Epic Games, Inc. All Rights Reserved.
+
+#pragma once
+
+#include "PropertyEditing.h"
+
+//////////////////////////////////////////////////////////////////////////
+// FFlipbookComponentDetailsCustomization
+
+class FFlipbookComponentDetailsCustomization : public IDetailCustomization
+{
+public:
+	/** Makes a new instance of this detail layout class for a specific detail view requesting it */
+	static TSharedRef<IDetailCustomization> MakeInstance();
+
+	// IDetailCustomization interface
+	virtual void CustomizeDetails(IDetailLayoutBuilder& DetailLayout) override;
+	// End of IDetailCustomization interface
+};
diff --git a/Engine/Plugins/2D/Paper2D/Source/Paper2DEditor/Private/Paper2DEditorModule.cpp b/Engine/Plugins/2D/Paper2D/Source/Paper2DEditor/Private/Paper2DEditorModule.cpp
--- a/Engine/Plugins/2D/Paper2D/Source/Paper2DEditor/Private/Paper2DEditorModule.cpp
+++ b/Engine/Plugins/2D/Paper2D/Source/Paper2DEditor/Private/Paper2DEditorModule.cpp
@@ -16,11 +16,13 @@
 #include "SpriteAssetTypeActions.h"
 #include "PaperSpriteAssetBroker.h"
 #include "SpriteEditor/SpriteDetailsCustomization.h"
+#include "SpriteEditor/SpriteComponentDetailsCustomization.h"
 #include "SpriteEditor/SpritePolygonCollectionCustomization.h"
 
 // Flipbook support
 #include "FlipbookAssetTypeActions.h"
 #include "PaperFlipbookAssetBroker.h"
+#include "FlipbookEditor/FlipbookComponentDetailsCustomization.h"
 
 // Tile set support
 #include "TileSetAssetTypeActions.h"
@@ -119,6 +121,8 @@ public:
 			PropertyModule.RegisterCustomClassLayout(UPaperTileMapComponent::StaticClass()->GetFName(), FOnGetDetailCustomizationInstance::CreateStatic(&FPaperTileMapDetailsCustomization::MakeInstance));
 			PropertyModule.RegisterCustomClassLayout(UPaperTileMap::StaticClass()->GetFName(), FOnGetDetailCustomizationInstance::CreateStatic(&FPaperTileMapDetailsCustomization::MakeInstance));
 			PropertyModule.RegisterCustomClassLayout(UPaperSprite::StaticClass()->GetFName(), FOnGetDetailCustomizationInstance::CreateStatic(&FSpriteDetailsCustomization::MakeInstance));
+			PropertyModule.RegisterCustomClassLayout(UPaperSpriteComponent::StaticClass()->GetFName(), FOnGetDetailCustomizationInstance::CreateStatic(&FSpriteComponentDetailsCustomization::MakeInstance));
+			PropertyModule.RegisterCustomClassLayout(UPaperFlipbookComponent::StaticClass()->GetFName(), FOnGetDetailCustomizationInstance::CreateStatic(&FFlipbookComponentDetailsCustomization::MakeInstance));
 
 			//@TODO: Struct registration should happen using ::StaticStruct, not by string!!!
 			//PropertyModule.RegisterCustomPropertyTypeLayout( "SpritePolygonCollection", FOnGetPropertyTypeCustomizationInstance::CreateStatic( &FSpritePolygonCollectionCustomization::MakeInstance ) );
diff --git a/Engine/Plugins/2D/Paper2D/Source/Paper2DEditor/Private/SpriteEditor/SpriteComponentDetailsCustomization.cpp b/Engine/Plugins/2D/Paper2D/Source/Paper2DEditor/Private/SpriteEditor/SpriteComponentDetailsCustomization.cpp
new file mode 100644
--- /dev/null
+++ b/Engine/Plugins/2D/Paper2D/Source/Paper2DEditor/Private/SpriteEditor/SpriteComponentDetailsCustomization.cpp
@@ -0,0 +1,24 @@
+// Copyright 1998-2015 Epic Games, Inc. All Rights Reserved.
+
+#include "Paper2DEditorPrivatePCH.h"
+#include "SpriteComponentDetailsCustomization.h"
+
+#define LOCTEXT_NAMESPACE "SpriteEditor"
+
+//////////////////////////////////////////////////////////////////////////
+// FSpriteComponentDetailsCustomization
+
+TSharedRef<IDetailCustomization> FSpriteComponentDetailsCustomization::MakeInstance()
+{
+	return MakeShareable(new FSpriteComponentDetailsCustomization);
+}
+
+void FSpriteComponentDetailsCustomization::CustomizeDetails(IDetailLayoutBuilder& DetailBuilder)
+{
+	// Create a category so this is displayed early in the properties
+	DetailBuilder.EditCategory("Sprite", FText::GetEmpty(), ECategoryPriority::Important);
+}
+
+//////////////////////////////////////////////////////////////////////////
+
+#undef LOCTEXT_NAMESPACE
diff --git a/Engine/Plugins/2D/Paper2D/Source/Paper2DEditor/Private/SpriteEditor/SpriteComponentDetailsCustomization.h b/Engine/Plugins/2D/Paper2D/Source/Paper2DEditor/Private/SpriteEditor/SpriteComponentDetailsCustomization.h
new file mode 100644
--- /dev/null
+++ b/Engine/Plugins/2D/Paper2D/Source/Paper2DEditor/Private/SpriteEditor/SpriteComponentDetailsCustomization.h
@@ -0,0 +1,19 @@
+// Copyright 1998-2015 Epic Games, Inc. All Rights Reserved.
+
+#pragma once
+
+#include "PropertyEditing.h"
+
+//////////////////////////////////////////////////////////////////////////
+// FSpriteComponentDetailsCustomization
+
+class FSpriteComponentDetailsCustomization : public IDetailCustomization
+{
+public:
+	/** Makes a new instance of this detail layout class for a specific detail view requesting it */
+	static TSharedRef<IDetailCustomization> MakeInstance();
+
+	// IDetailCustomization interface
+	virtual void CustomizeDetails(IDetailLayoutBuilder& DetailLayout) override;
+	// End of IDetailCustomization interface
+};
