Paper2D: Remove the unused TickComponent method from UPaperTileMapComponent

diff --git a/Engine/Plugins/2D/Paper2D/Source/Paper2D/Classes/PaperTileMapComponent.h b/Engine/Plugins/2D/Paper2D/Source/Paper2D/Classes/PaperTileMapComponent.h
--- a/Engine/Plugins/2D/Paper2D/Source/Paper2D/Classes/PaperTileMapComponent.h
+++ b/Engine/Plugins/2D/Paper2D/Source/Paper2D/Classes/PaperTileMapComponent.h
@@ -64,7 +64,6 @@ public:
 	// End of UObject interface
 
 	// UActorComponent interface
-	virtual void TickComponent(float DeltaTime, enum ELevelTick TickType, FActorComponentTickFunction* ThisTickFunction) override;
 	virtual const UObject* AdditionalStatObject() const override;
 	// End of UActorComponent interface
 
diff --git a/Engine/Plugins/2D/Paper2D/Source/Paper2D/Private/PaperTileMapComponent.cpp b/Engine/Plugins/2D/Paper2D/Source/Paper2D/Private/PaperTileMapComponent.cpp
--- a/Engine/Plugins/2D/Paper2D/Source/Paper2D/Private/PaperTileMapComponent.cpp
+++ b/Engine/Plugins/2D/Paper2D/Source/Paper2D/Private/PaperTileMapComponent.cpp
@@ -71,12 +71,6 @@ FBoxSphereBounds UPaperTileMapComponent::CalcBounds(const FTransform& LocalToWor
 	}
 }
 
-void UPaperTileMapComponent::TickComponent(float DeltaTime, enum ELevelTick TickType, FActorComponentTickFunction* ThisTickFunction)
-{
-	// Indicate we need to send new dynamic data.
-	MarkRenderDynamicDataDirty();
-}
-
 #if WITH_EDITORONLY_DATA
 void UPaperTileMapComponent::Serialize(FArchive& Ar)
 {
