Paper2D: Add support for PixelsPerUnrealUnit scaling factor to UPaperTileMap [UE-8010]

diff --git a/Engine/Plugins/2D/Paper2D/Source/Paper2D/Classes/PaperTileMap.h b/Engine/Plugins/2D/Paper2D/Source/Paper2D/Classes/PaperTileMap.h
--- a/Engine/Plugins/2D/Paper2D/Source/Paper2D/Classes/PaperTileMap.h
+++ b/Engine/Plugins/2D/Paper2D/Source/Paper2D/Classes/PaperTileMap.h
@@ -49,9 +49,9 @@ class PAPER2D_API UPaperTileMap : public UObject
 	UPROPERTY(Category=Setup, EditAnywhere, BlueprintReadOnly, meta=(UIMin=1, ClampMin=1))
 	int32 TileHeight;
 
-	// Pixels per Unreal Unit (pixels per cm) - Note: Currently unused!
-	UPROPERTY(Category=Setup, EditAnywhere)
-	float PixelsPerUnit;
+	// The scaling factor between pixels and Unreal units (cm) (e.g., 0.64 would make a 64 pixel wide tile take up 100 cm)
+	UPROPERTY(Category = Setup, EditAnywhere)
+	float PixelsPerUnrealUnit;
 
 	// The Z-separation incurred as you travel in X (not strictly applied, batched tiles will be put at the same Z level) 
 	UPROPERTY(Category=Setup, EditAnywhere, AdvancedDisplay)
@@ -144,6 +144,17 @@ public:
 
 	// Resize the tile map and all layers
 	void ResizeMap(int32 NewWidth, int32 NewHeight, bool bForceResize = true);
+
+	// Return the scaling factor between pixels and Unreal units (cm)
+	float GetPixelsPerUnrealUnit() const { return PixelsPerUnrealUnit; }
+
+	// Return the scaling factor between Unreal units (cm) and pixels
+	float GetUnrealUnitsPerPixel() const { return 1.0f / PixelsPerUnrealUnit; }
+
+	// Called when a fresh tile map has been created (by factory or otherwise)
+	// Adds a default layer and pulls the PixelsPerUnrealUnit from the project settings
+	void InitializeNewEmptyTileMap();
+
 protected:
 	virtual void UpdateBodySetup();
 };
diff --git a/Engine/Plugins/2D/Paper2D/Source/Paper2D/Private/PaperTileMap.cpp b/Engine/Plugins/2D/Paper2D/Source/Paper2D/Private/PaperTileMap.cpp
--- a/Engine/Plugins/2D/Paper2D/Source/Paper2D/Private/PaperTileMap.cpp
+++ b/Engine/Plugins/2D/Paper2D/Source/Paper2D/Private/PaperTileMap.cpp
@@ -16,7 +16,7 @@ UPaperTileMap::UPaperTileMap(const FObjectInitializer& ObjectInitializer)
 	MapHeight = 4;
 	TileWidth = 32;
 	TileHeight = 32;
-	PixelsPerUnit = 1.0f;
+	PixelsPerUnrealUnit = 1.0f;
 	SeparationPerTileX = 0.0f;
 	SeparationPerTileY = 0.0f;
 	SeparationPerLayer = 4.0f;
@@ -84,6 +84,11 @@ void UPaperTileMap::PostEditChangeProperty(FPropertyChangedEvent& PropertyChange
 
 	ValidateSelectedLayerIndex();
 
+	if (PixelsPerUnrealUnit <= 0.0f)
+	{
+		PixelsPerUnrealUnit = 1.0f;
+	}
+
 	if ((PropertyName == GET_MEMBER_NAME_CHECKED(UPaperTileMap, MapWidth)) || (PropertyName == GET_MEMBER_NAME_CHECKED(UPaperTileMap, MapHeight)))
 	{
 		ResizeMap(MapWidth, MapHeight, /*bForceResize=*/ true);
@@ -181,54 +186,62 @@ void UPaperTileMap::UpdateBodySetup()
 
 void UPaperTileMap::GetTileToLocalParameters(FVector& OutCornerPosition, FVector& OutStepX, FVector& OutStepY, FVector& OutOffsetYFactor) const
 {
+	const float UnrealUnitsPerPixel = GetUnrealUnitsPerPixel();
+	const float TileWidthInUU = TileWidth * UnrealUnitsPerPixel;
+	const float TileHeightInUU = TileHeight * UnrealUnitsPerPixel;
+
 	switch (ProjectionMode)
 	{
 	case ETileMapProjectionMode::Orthogonal:
 	default:
-		OutCornerPosition = -(TileWidth * PaperAxisX * 0.5f) + (TileHeight * PaperAxisY * 0.5f);
+		OutCornerPosition = -(TileWidthInUU * PaperAxisX * 0.5f) + (TileHeightInUU * PaperAxisY * 0.5f);
 		OutOffsetYFactor = FVector::ZeroVector;
-		OutStepX = PaperAxisX * TileWidth;
-		OutStepY = -PaperAxisY * TileHeight;
+		OutStepX = PaperAxisX * TileWidthInUU;
+		OutStepY = -PaperAxisY * TileHeightInUU;
 		break;
 	case ETileMapProjectionMode::IsometricDiamond:
-		OutCornerPosition = (TileHeight * PaperAxisY * 0.5f);
+		OutCornerPosition = (TileHeightInUU * PaperAxisY * 0.5f);
 		OutOffsetYFactor = FVector::ZeroVector;
-		OutStepX = (TileWidth * PaperAxisX * 0.5f) - (TileHeight * PaperAxisY * 0.5f);
-		OutStepY = (TileWidth * PaperAxisX * -0.5f) - (TileHeight * PaperAxisY * 0.5f);
+		OutStepX = (TileWidthInUU * PaperAxisX * 0.5f) - (TileHeightInUU * PaperAxisY * 0.5f);
+		OutStepY = (TileWidthInUU * PaperAxisX * -0.5f) - (TileHeightInUU * PaperAxisY * 0.5f);
 		break;
 	case ETileMapProjectionMode::HexagonalStaggered:
 	case ETileMapProjectionMode::IsometricStaggered:
-		OutCornerPosition = -(TileWidth * PaperAxisX * 0.5f) + (TileHeight * PaperAxisY * 1.0f);
-		OutOffsetYFactor = 0.5f * TileWidth * PaperAxisX;
-		OutStepX = PaperAxisX * TileWidth;
-		OutStepY = 0.5f * -PaperAxisY * TileHeight;
+		OutCornerPosition = -(TileWidthInUU * PaperAxisX * 0.5f) + (TileHeightInUU * PaperAxisY * 1.0f);
+		OutOffsetYFactor = 0.5f * TileWidthInUU * PaperAxisX;
+		OutStepX = PaperAxisX * TileWidthInUU;
+		OutStepY = 0.5f * -PaperAxisY * TileHeightInUU;
 		break;
 	}
 }
 
 void UPaperTileMap::GetLocalToTileParameters(FVector& OutCornerPosition, FVector& OutStepX, FVector& OutStepY, FVector& OutOffsetYFactor) const
 {
+	const float UnrealUnitsPerPixel = GetUnrealUnitsPerPixel();
+	const float TileWidthInUU = TileWidth * UnrealUnitsPerPixel;
+	const float TileHeightInUU = TileHeight * UnrealUnitsPerPixel;
+
 	switch (ProjectionMode)
 	{
 	case ETileMapProjectionMode::Orthogonal:
 	default:
-		OutCornerPosition = -(TileWidth * PaperAxisX * 0.5f) + (TileHeight * PaperAxisY * 0.5f);
+		OutCornerPosition = -(TileWidthInUU * PaperAxisX * 0.5f) + (TileHeightInUU * PaperAxisY * 0.5f);
 		OutOffsetYFactor = FVector::ZeroVector;
-		OutStepX = PaperAxisX / TileWidth;
-		OutStepY = -PaperAxisY / TileHeight;
+		OutStepX = PaperAxisX / TileWidthInUU;
+		OutStepY = -PaperAxisY / TileHeightInUU;
 		break;
 	case ETileMapProjectionMode::IsometricDiamond:
-		OutCornerPosition = (TileHeight * PaperAxisY * 0.5f);
+		OutCornerPosition = (TileHeightInUU * PaperAxisY * 0.5f);
 		OutOffsetYFactor = FVector::ZeroVector;
-		OutStepX = (PaperAxisX / TileWidth) - (PaperAxisY / TileHeight);
-		OutStepY = (-PaperAxisX / TileWidth) - (PaperAxisY / TileHeight);
+		OutStepX = (PaperAxisX / TileWidthInUU) - (PaperAxisY / TileHeightInUU);
+		OutStepY = (-PaperAxisX / TileWidthInUU) - (PaperAxisY / TileHeightInUU);
 		break;
 	case ETileMapProjectionMode::HexagonalStaggered:
 	case ETileMapProjectionMode::IsometricStaggered:
-		OutCornerPosition = -(TileWidth * PaperAxisX * 0.5f) + (TileHeight * PaperAxisY * 1.0f);
-		OutOffsetYFactor = 0.5f * TileWidth * PaperAxisX;
-		OutStepX = PaperAxisX / TileWidth;
-		OutStepY = -PaperAxisY / TileHeight;
+		OutCornerPosition = -(TileWidthInUU * PaperAxisX * 0.5f) + (TileHeightInUU * PaperAxisY * 1.0f);
+		OutOffsetYFactor = 0.5f * TileWidthInUU * PaperAxisX;
+		OutStepX = PaperAxisX / TileWidthInUU;
+		OutStepY = -PaperAxisY / TileHeightInUU;
 		break;
 	}
 }
@@ -299,21 +312,25 @@ FBoxSphereBounds UPaperTileMap::GetRenderBounds() const
 	const float Depth = SeparationPerLayer * (TileLayers.Num() - 1);
 	const float HalfThickness = 2.0f;
 
+	const float UnrealUnitsPerPixel = GetUnrealUnitsPerPixel();
+	const float TileWidthInUU = TileWidth * UnrealUnitsPerPixel;
+	const float TileHeightInUU = TileHeight * UnrealUnitsPerPixel;
+
 	switch (ProjectionMode)
 	{
 		case ETileMapProjectionMode::Orthogonal:
 		default:
 		{
-			const FVector BottomLeft((-0.5f) * TileWidth, -HalfThickness - Depth, -(MapHeight - 0.5f) * TileHeight);
-			const FVector Dimensions(MapWidth*TileWidth, Depth + 2 * HalfThickness, MapHeight * TileHeight);
+			const FVector BottomLeft((-0.5f) * TileWidthInUU, -HalfThickness - Depth, -(MapHeight - 0.5f) * TileHeightInUU);
+			const FVector Dimensions(MapWidth * TileWidthInUU, Depth + 2 * HalfThickness, MapHeight * TileHeightInUU);
 
 			const FBox Box(BottomLeft, BottomLeft + Dimensions);
 			return FBoxSphereBounds(Box);
 		}
 		case ETileMapProjectionMode::IsometricDiamond:
 		{
-			 const FVector BottomLeft((-0.5f) * TileWidth * MapWidth, -HalfThickness - Depth, -MapHeight * TileHeight);
-			 const FVector Dimensions(MapWidth*TileWidth, Depth + 2 * HalfThickness, (MapHeight + 1) * TileHeight);
+			 const FVector BottomLeft((-0.5f) * TileWidthInUU * MapWidth, -HalfThickness - Depth, -MapHeight * TileHeightInUU);
+			 const FVector Dimensions(MapWidth * TileWidthInUU, Depth + 2 * HalfThickness, (MapHeight + 1) * TileHeightInUU);
 
 			 const FBox Box(BottomLeft, BottomLeft + Dimensions);
 			 return FBoxSphereBounds(Box);
@@ -387,6 +404,13 @@ void UPaperTileMap::ResizeMap(int32 NewWidth, int32 NewHeight, bool bForceResize
 	}
 }
 
+void UPaperTileMap::InitializeNewEmptyTileMap()
+{
+	AddNewLayer();
+
+	PixelsPerUnrealUnit = GetDefault<UPaperRuntimeSettings>()->DefaultPixelsPerUnrealUnit;
+}
+
 //////////////////////////////////////////////////////////////////////////
 
 #undef LOCTEXT_NAMESPACE
\ No newline at end of file
diff --git a/Engine/Plugins/2D/Paper2D/Source/Paper2D/Private/PaperTileMapComponent.cpp b/Engine/Plugins/2D/Paper2D/Source/Paper2D/Private/PaperTileMapComponent.cpp
--- a/Engine/Plugins/2D/Paper2D/Source/Paper2D/Private/PaperTileMapComponent.cpp
+++ b/Engine/Plugins/2D/Paper2D/Source/Paper2D/Private/PaperTileMapComponent.cpp
@@ -93,7 +93,7 @@ void UPaperTileMapComponent::PostLoad()
 		TileMap->MapHeight = MapHeight_DEPRECATED;
 		TileMap->TileWidth = TileWidth_DEPRECATED;
 		TileMap->TileHeight = TileHeight_DEPRECATED;
-		TileMap->PixelsPerUnit = 1.0f;
+		TileMap->PixelsPerUnrealUnit = 1.0f;
 		TileMap->SelectedTileSet = DefaultLayerTileSet_DEPRECATED;
 		TileMap->Material = Material_DEPRECATED;
 		TileMap->TileLayers = TileLayers_DEPRECATED;
@@ -211,6 +211,8 @@ void UPaperTileMapComponent::RebuildRenderData(FPaperTileMapRenderSceneProxy* Pr
 	FVector2D SourceDimensionsUV(1.0f, 1.0f);
 	FVector2D TileSizeXY(0.0f, 0.0f);
 
+	const float UnrealUnitsPerPixel = TileMap->GetUnrealUnitsPerPixel();
+
 	for (int32 Z = 0; Z < TileMap->TileLayers.Num(); ++Z)
 	{
 		UPaperTileLayer* Layer = TileMap->TileLayers[Z];
@@ -313,13 +315,13 @@ void UPaperTileMapComponent::RebuildRenderData(FPaperTileMapRenderSceneProxy* Pr
 						if (TileInfo.TileSet != nullptr)
 						{
 							SourceDimensionsUV = FVector2D(TileInfo.TileSet->TileWidth * InverseTextureSize.X, TileInfo.TileSet->TileHeight * InverseTextureSize.Y);
-							TileSizeXY = FVector2D(TileInfo.TileSet->TileWidth, TileInfo.TileSet->TileHeight);
+							TileSizeXY = FVector2D(UnrealUnitsPerPixel * TileInfo.TileSet->TileWidth, UnrealUnitsPerPixel * TileInfo.TileSet->TileHeight);
 							TileSetOffset = (TileInfo.TileSet->DrawingOffset.X * PaperAxisX) + (TileInfo.TileSet->DrawingOffset.Y * PaperAxisY);
 						}
 						else
 						{
 							SourceDimensionsUV = FVector2D(TileWidth * InverseTextureSize.X, TileHeight * InverseTextureSize.Y);
-							TileSizeXY = FVector2D(TileWidth, TileHeight);
+							TileSizeXY = FVector2D(UnrealUnitsPerPixel * TileWidth, UnrealUnitsPerPixel * TileHeight);
 							TileSetOffset = FVector::ZeroVector;
 						}
 						LastSourceTexture = SourceTexture;
@@ -371,7 +373,7 @@ void UPaperTileMapComponent::CreateNewOwnedTileMap()
 
 	UPaperTileMap* NewTileMap = NewObject<UPaperTileMap>(this);
 	NewTileMap->SetFlags(RF_Transactional);
-	NewTileMap->AddNewLayer();
+	NewTileMap->InitializeNewEmptyTileMap();
 
 	SetTileMap(NewTileMap);
 }
diff --git a/Engine/Plugins/2D/Paper2D/Source/Paper2DEditor/Private/PaperTileMapFactory.cpp b/Engine/Plugins/2D/Paper2D/Source/Paper2DEditor/Private/PaperTileMapFactory.cpp
--- a/Engine/Plugins/2D/Paper2D/Source/Paper2DEditor/Private/PaperTileMapFactory.cpp
+++ b/Engine/Plugins/2D/Paper2D/Source/Paper2DEditor/Private/PaperTileMapFactory.cpp
@@ -25,7 +25,7 @@ UObject* UPaperTileMapFactory::FactoryCreateNew(UClass* Class, UObject* InParent
 		NewTileMap->TileHeight = InitialTileSet->TileHeight;
 	}
 
-	NewTileMap->AddNewLayer();
+	NewTileMap->InitializeNewEmptyTileMap();
 
 	return NewTileMap;
 }
diff --git a/Engine/Plugins/2D/Paper2D/Source/Paper2DEditor/Private/TileMapEditing/EdModeTileMap.cpp b/Engine/Plugins/2D/Paper2D/Source/Paper2DEditor/Private/TileMapEditing/EdModeTileMap.cpp
--- a/Engine/Plugins/2D/Paper2D/Source/Paper2DEditor/Private/TileMapEditing/EdModeTileMap.cpp
+++ b/Engine/Plugins/2D/Paper2D/Source/Paper2DEditor/Private/TileMapEditing/EdModeTileMap.cpp
@@ -110,7 +110,7 @@ void FEdModeTileMap::Enter()
 	UWorld* World = GetWorld();
 
 	CursorPreviewComponent = NewObject<UPaperTileMapComponent>();
-	CursorPreviewComponent->TileMap->AddNewLayer();
+	CursorPreviewComponent->TileMap->InitializeNewEmptyTileMap();
 	CursorPreviewComponent->TranslucencySortPriority = 99999;
 	CursorPreviewComponent->UpdateBounds();
 	CursorPreviewComponent->AddToRoot();
@@ -997,7 +997,7 @@ void FEdModeTileMap::SynchronizePreviewWithTileMap(UPaperTileMap* NewTileMap)
 
 	UE_CHANGE_IF_DIFFERENT(TileWidth);
 	UE_CHANGE_IF_DIFFERENT(TileHeight);
-	UE_CHANGE_IF_DIFFERENT(PixelsPerUnit);
+	UE_CHANGE_IF_DIFFERENT(PixelsPerUnrealUnit);
 	UE_CHANGE_IF_DIFFERENT(SeparationPerTileX);
 	UE_CHANGE_IF_DIFFERENT(SeparationPerTileY);
 	UE_CHANGE_IF_DIFFERENT(SeparationPerLayer);
diff --git a/Engine/Plugins/2D/Paper2D/Source/Paper2DEditor/Private/TileMapEditing/TileMapActorFactory.cpp b/Engine/Plugins/2D/Paper2D/Source/Paper2DEditor/Private/TileMapEditing/TileMapActorFactory.cpp
--- a/Engine/Plugins/2D/Paper2D/Source/Paper2DEditor/Private/TileMapEditing/TileMapActorFactory.cpp
+++ b/Engine/Plugins/2D/Paper2D/Source/Paper2DEditor/Private/TileMapEditing/TileMapActorFactory.cpp
@@ -42,7 +42,7 @@ void UTileMapActorFactory::PostSpawnActor(UObject* Asset, AActor* NewActor)
 
 	if (RenderComponent->OwnsTileMap())
 	{
-		RenderComponent->TileMap->AddNewLayer();
+		RenderComponent->TileMap->InitializeNewEmptyTileMap();
 	}
 }
 
@@ -68,7 +68,7 @@ void UTileMapActorFactory::PostCreateBlueprint(UObject* Asset, AActor* CDO)
 
 		if (RenderComponent->OwnsTileMap())
 		{
-			RenderComponent->TileMap->AddNewLayer();
+			RenderComponent->TileMap->InitializeNewEmptyTileMap();
 		}
 	}
 }
diff --git a/Engine/Plugins/2D/Paper2D/Source/PaperTiledImporter/Private/PaperTiledImporterFactory.cpp b/Engine/Plugins/2D/Paper2D/Source/PaperTiledImporter/Private/PaperTiledImporterFactory.cpp
--- a/Engine/Plugins/2D/Paper2D/Source/PaperTiledImporter/Private/PaperTiledImporterFactory.cpp
+++ b/Engine/Plugins/2D/Paper2D/Source/PaperTiledImporter/Private/PaperTiledImporterFactory.cpp
@@ -195,6 +195,7 @@ UObject* UPaperTiledImporterFactory::FactoryCreateText(UClass* InClass, UObject*
 		Result->SeparationPerTileY = 0.0f;
 		Result->SeparationPerLayer = -1.0f;
 		Result->ProjectionMode = GlobalInfo.GetOrientationType();
+		Result->PixelsPerUnrealUnit = GetDefault<UPaperRuntimeSettings>()->DefaultPixelsPerUnrealUnit;
 
 		// Create the tile sets
 		for (const FTileSetFromTiled& TileSetData : GlobalInfo.TileSets)
