Paper2D: Choose a good default for the selected tile set in more cases
- when creating a tile map component from a tile set
- when creating a tile map asset from a tile set
- when creating a new 'instanced' tile map on a tile map component instance
[UE-12195]

diff --git a/Engine/Plugins/2D/Paper2D/Source/Paper2D/Private/PaperTileMapComponent.cpp b/Engine/Plugins/2D/Paper2D/Source/Paper2D/Private/PaperTileMapComponent.cpp
--- a/Engine/Plugins/2D/Paper2D/Source/Paper2D/Private/PaperTileMapComponent.cpp
+++ b/Engine/Plugins/2D/Paper2D/Source/Paper2D/Private/PaperTileMapComponent.cpp
@@ -375,6 +375,11 @@ void UPaperTileMapComponent::CreateNewOwnedTileMap()
 	NewTileMap->SetFlags(RF_Transactional);
 	NewTileMap->InitializeNewEmptyTileMap();
 
+	if (TileMap != nullptr)
+	{
+		NewTileMap->SelectedTileSet = TileMap->SelectedTileSet;
+	}
+
 	SetTileMap(NewTileMap);
 }
 
diff --git a/Engine/Plugins/2D/Paper2D/Source/Paper2DEditor/Private/PaperTileMapFactory.cpp b/Engine/Plugins/2D/Paper2D/Source/Paper2DEditor/Private/PaperTileMapFactory.cpp
--- a/Engine/Plugins/2D/Paper2D/Source/Paper2DEditor/Private/PaperTileMapFactory.cpp
+++ b/Engine/Plugins/2D/Paper2D/Source/Paper2DEditor/Private/PaperTileMapFactory.cpp
@@ -23,6 +23,7 @@ UObject* UPaperTileMapFactory::FactoryCreateNew(UClass* Class, UObject* InParent
 	{
 		NewTileMap->TileWidth = InitialTileSet->TileWidth;
 		NewTileMap->TileHeight = InitialTileSet->TileHeight;
+		NewTileMap->SelectedTileSet = InitialTileSet;
 	}
 
 	NewTileMap->InitializeNewEmptyTileMap();
diff --git a/Engine/Plugins/2D/Paper2D/Source/Paper2DEditor/Private/TileMapEditing/TileMapActorFactory.cpp b/Engine/Plugins/2D/Paper2D/Source/Paper2DEditor/Private/TileMapEditing/TileMapActorFactory.cpp
--- a/Engine/Plugins/2D/Paper2D/Source/Paper2DEditor/Private/TileMapEditing/TileMapActorFactory.cpp
+++ b/Engine/Plugins/2D/Paper2D/Source/Paper2DEditor/Private/TileMapEditing/TileMapActorFactory.cpp
@@ -36,6 +36,7 @@ void UTileMapActorFactory::PostSpawnActor(UObject* Asset, AActor* NewActor)
 			RenderComponent->UnregisterComponent();
 			RenderComponent->TileMap->TileWidth = TileSet->TileWidth;
 			RenderComponent->TileMap->TileHeight = TileSet->TileHeight;
+			RenderComponent->TileMap->SelectedTileSet = TileSet;
 			RenderComponent->RegisterComponent();
 		}
 	}
@@ -63,6 +64,7 @@ void UTileMapActorFactory::PostCreateBlueprint(UObject* Asset, AActor* CDO)
 			{
 				RenderComponent->TileMap->TileWidth = TileSet->TileWidth;
 				RenderComponent->TileMap->TileHeight = TileSet->TileHeight;
+				RenderComponent->TileMap->SelectedTileSet = TileSet;
 			}
 		}
 
