Paper2D: Fixed sprite thumbnails disappearing when the pivot is not center-center
[UE-9982]

diff --git a/Engine/Plugins/2D/Paper2D/Source/Paper2DEditor/Private/PaperSpriteThumbnailRenderer.cpp b/Engine/Plugins/2D/Paper2D/Source/Paper2DEditor/Private/PaperSpriteThumbnailRenderer.cpp
--- a/Engine/Plugins/2D/Paper2D/Source/Paper2DEditor/Private/PaperSpriteThumbnailRenderer.cpp
+++ b/Engine/Plugins/2D/Paper2D/Source/Paper2DEditor/Private/PaperSpriteThumbnailRenderer.cpp
@@ -62,14 +62,14 @@ void UPaperSpriteThumbnailRenderer::DrawFrame(class UPaperSprite* Sprite, int32
 		// Draw the sprite itself
 		// Use the baked render data, so we don't have to care about rotations and possibly
 		// other sprites overlapping in source, UV region, etc.
-		const TArray<FVector4> &BakedRenderData = Sprite->BakedRenderData;
+		const TArray<FVector4>& BakedRenderData = Sprite->BakedRenderData;
 		TArray<FVector2D> CanvasPositions;
 		TArray<FVector2D> CanvasUVs;
 
 		for (int Vertex = 0; Vertex < BakedRenderData.Num(); ++Vertex)
 		{
-			new(CanvasPositions)FVector2D(BakedRenderData[Vertex].X, -BakedRenderData[Vertex].Y);
-			new(CanvasUVs)FVector2D(BakedRenderData[Vertex].Z, BakedRenderData[Vertex].W);
+			new (CanvasPositions) FVector2D(BakedRenderData[Vertex].X, BakedRenderData[Vertex].Y);
+			new (CanvasUVs) FVector2D(BakedRenderData[Vertex].Z, BakedRenderData[Vertex].W);
 		}
 
 		// Determine the bounds to use
@@ -87,9 +87,9 @@ void UPaperSpriteThumbnailRenderer::DrawFrame(class UPaperSprite* Sprite, int32
 		const FVector2D MaxPoint(FVector::DotProduct(MaxPoint3D, PaperAxisX), FVector::DotProduct(MaxPoint3D, PaperAxisY));
 
 		float ScaleFactor = 1.0f;
-		float UnscaledWidth = MaxPoint.X - MinPoint.X;
-		float UnscaledHeight = MaxPoint.Y - MinPoint.Y;
-		FVector2D Origin(X + Width / 2.0f, Y + Height / 2.0f);
+		const float UnscaledWidth = MaxPoint.X - MinPoint.X;
+		const float UnscaledHeight = MaxPoint.Y - MinPoint.Y;
+		const FVector2D Origin(X + Width * 0.5f, Y + Height * 0.5f);
 		if ((UnscaledWidth > 0.0f) && (UnscaledHeight > 0.0f) && (UnscaledWidth > UnscaledHeight))
 		{ 
 			ScaleFactor = Width / UnscaledWidth;
@@ -104,10 +104,11 @@ void UPaperSpriteThumbnailRenderer::DrawFrame(class UPaperSprite* Sprite, int32
 		for (int Vertex = 0; Vertex < CanvasPositions.Num(); ++Vertex)
 		{
 			CanvasPositions[Vertex] = (CanvasPositions[Vertex] - CanvasPositionCenter) * ScaleFactor + Origin;
+			CanvasPositions[Vertex].Y = Height - CanvasPositions[Vertex].Y;
 		}
 
 		// Draw triangles
-		if (CanvasPositions.Num() > 0 && SourceTexture->Resource != nullptr)
+		if ((CanvasPositions.Num() > 0) && (SourceTexture->Resource != nullptr))
 		{
 			TArray<FCanvasUVTri> Triangles;
 			const FLinearColor SpriteColor(FLinearColor::White);
