Paper2D: Constified more of the sprite thumbnail rendering code

diff --git a/Engine/Plugins/2D/Paper2D/Source/Paper2DEditor/Private/PaperSpriteThumbnailRenderer.cpp b/Engine/Plugins/2D/Paper2D/Source/Paper2DEditor/Private/PaperSpriteThumbnailRenderer.cpp
--- a/Engine/Plugins/2D/Paper2D/Source/Paper2DEditor/Private/PaperSpriteThumbnailRenderer.cpp
+++ b/Engine/Plugins/2D/Paper2D/Source/Paper2DEditor/Private/PaperSpriteThumbnailRenderer.cpp
@@ -86,21 +86,14 @@ void UPaperSpriteThumbnailRenderer::DrawFrame(class UPaperSprite* Sprite, int32
 		const FVector2D MinPoint(FVector::DotProduct(MinPoint3D, PaperAxisX), FVector::DotProduct(MinPoint3D, PaperAxisY));
 		const FVector2D MaxPoint(FVector::DotProduct(MaxPoint3D, PaperAxisX), FVector::DotProduct(MaxPoint3D, PaperAxisY));
 
-		float ScaleFactor = 1.0f;
 		const float UnscaledWidth = MaxPoint.X - MinPoint.X;
 		const float UnscaledHeight = MaxPoint.Y - MinPoint.Y;
 		const FVector2D Origin(X + Width * 0.5f, Y + Height * 0.5f);
-		if ((UnscaledWidth > 0.0f) && (UnscaledHeight > 0.0f) && (UnscaledWidth > UnscaledHeight))
-		{ 
-			ScaleFactor = Width / UnscaledWidth;
-		}
-		else
-		{
-			ScaleFactor = Height / UnscaledHeight;
-		}
+		const bool bIsWider = (UnscaledWidth > 0.0f) && (UnscaledHeight > 0.0f) && (UnscaledWidth > UnscaledHeight);
+		const float ScaleFactor = bIsWider ? (Width / UnscaledWidth) : (Height / UnscaledHeight);
 
 		// Scale and recenter
-		FVector2D CanvasPositionCenter = (MaxPoint + MinPoint) * 0.5f;
+		const FVector2D CanvasPositionCenter = (MaxPoint + MinPoint) * 0.5f;
 		for (int Vertex = 0; Vertex < CanvasPositions.Num(); ++Vertex)
 		{
 			CanvasPositions[Vertex] = (CanvasPositions[Vertex] - CanvasPositionCenter) * ScaleFactor + Origin;
