Paper2D: Fix incorrect collision bounding box generation when sprites were rotated in the source texture
[UE-12196]

diff --git a/Engine/Plugins/2D/Paper2D/Source/Paper2D/Private/PaperCustomVersion.h b/Engine/Plugins/2D/Paper2D/Source/Paper2D/Private/PaperCustomVersion.h
--- a/Engine/Plugins/2D/Paper2D/Source/Paper2D/Private/PaperCustomVersion.h
+++ b/Engine/Plugins/2D/Paper2D/Source/Paper2D/Private/PaperCustomVersion.h
@@ -46,6 +46,9 @@ struct FPaperCustomVersion
 		// Add source texture size
 		AddSourceTextureSize = 12,
 
+		// Fix incorrect collision geometry generation when sprites were rotated in the source texture
+		FixIncorrectCollisionOnSourceRotatedSprites = 13,
+
 		// -----<new versions can be added above this line>-------------------------------------------------
 		VersionPlusOne,
 		LatestVersion = VersionPlusOne - 1
diff --git a/Engine/Plugins/2D/Paper2D/Source/Paper2D/Private/PaperSprite.cpp b/Engine/Plugins/2D/Paper2D/Source/Paper2D/Private/PaperSprite.cpp
--- a/Engine/Plugins/2D/Paper2D/Source/Paper2D/Private/PaperSprite.cpp
+++ b/Engine/Plugins/2D/Paper2D/Source/Paper2D/Private/PaperSprite.cpp
@@ -846,13 +846,14 @@ void UPaperSprite::BuildBoundingBoxCollisionData(bool bUseTightBounds)
 			BodySetup3D->AggGeom.EmptyElements();
 
 			// Determine the box size and center in pivot space
-			const FVector2D& BoxSize2DInPixels = CollisionGeometry.Polygons[0].BoxSize;
-			const FVector2D& BoxPos = CollisionGeometry.Polygons[0].BoxPosition;
-			const FVector2D CenterInTextureSpace = BoxPos + (BoxSize2DInPixels * 0.5f);
+			const FVector2D& BoxSizeInTextureSpace = CollisionGeometry.Polygons[0].BoxSize;
+			const FVector2D& BoxPosInTextureSpace = CollisionGeometry.Polygons[0].BoxPosition;
+			const FVector2D CenterInTextureSpace = BoxPosInTextureSpace + (BoxSizeInTextureSpace * 0.5f);
 			const FVector2D CenterInPivotSpace = ConvertTextureSpaceToPivotSpace(CenterInTextureSpace);
 
 			// Convert from pixels to uu
-			const FVector2D BoxSize2D = BoxSize2DInPixels * UnitsPerPixel;
+			const FVector2D BoxSizeInPivotSpace = bRotatedInSourceImage ? FVector2D(BoxSizeInTextureSpace.Y, BoxSizeInTextureSpace.X) : BoxSizeInTextureSpace;
+			const FVector2D BoxSize2D = BoxSizeInPivotSpace * UnitsPerPixel;
 			const FVector2D CenterInScaledSpace = CenterInPivotSpace * UnitsPerPixel;
 
 			// Create a new box primitive
@@ -871,13 +872,14 @@ void UPaperSprite::BuildBoundingBoxCollisionData(bool bUseTightBounds)
 			BodySetup2D->AggGeom2D.EmptyElements();
 
 			// Determine the box center in pivot space
-			const FVector2D& BoxSize2DInPixels = CollisionGeometry.Polygons[0].BoxSize;
-			const FVector2D& BoxPos = CollisionGeometry.Polygons[0].BoxPosition;
-			const FVector2D CenterInTextureSpace = BoxPos + (BoxSize2DInPixels * 0.5f);
+			const FVector2D& BoxSizeInTextureSpace = CollisionGeometry.Polygons[0].BoxSize;
+			const FVector2D& BoxPosInTextureSpace = CollisionGeometry.Polygons[0].BoxPosition;
+			const FVector2D CenterInTextureSpace = BoxPosInTextureSpace + (BoxSizeInTextureSpace * 0.5f);
 			const FVector2D CenterInPivotSpace = ConvertTextureSpaceToPivotSpace(CenterInTextureSpace);
 
 			// Convert from pixels to uu
-			const FVector2D BoxSize2D = BoxSize2DInPixels * UnitsPerPixel;
+			const FVector2D BoxSizeInPivotSpace = bRotatedInSourceImage ? FVector2D(BoxSizeInTextureSpace.Y, BoxSizeInTextureSpace.X) : BoxSizeInTextureSpace;
+			const FVector2D BoxSize2D = BoxSizeInPivotSpace * UnitsPerPixel;
 			const FVector2D CenterInScaledSpace = CenterInPivotSpace * UnitsPerPixel;
 
 			// Create a new box primitive
@@ -1899,7 +1901,7 @@ void UPaperSprite::PostLoad()
 		bRebuildCollision = true;
 		bRebuildRenderData = true;
 	}
-	else if (PaperVer < FPaperCustomVersion::FixTypoIn3DConvexHullCollisionGeneration)
+	else if (PaperVer < FPaperCustomVersion::FixIncorrectCollisionOnSourceRotatedSprites)
 	{
 		bRebuildCollision = true;
 	}
@@ -1909,7 +1911,7 @@ void UPaperSprite::PostLoad()
 		BodySetup->DefaultInstance.SetCollisionProfileName(UCollisionProfile::BlockAllDynamic_ProfileName);
 	}
 
-	if (PaperVer >= FPaperCustomVersion::AddSourceTextureSize && NeedRescaleSpriteData())
+	if ((PaperVer >= FPaperCustomVersion::AddSourceTextureSize) && NeedRescaleSpriteData())
 	{
 		RescaleSpriteData(GetSourceTexture());
 		bRebuildCollision = true;
