Tile Map Editor: Fixed selection outline not showing up in the standalone editor [UE-11507]

diff --git a/Engine/Plugins/2D/Paper2D/Source/Paper2DEditor/Private/TileMapEditing/TileMapEditorViewportClient.cpp b/Engine/Plugins/2D/Paper2D/Source/Paper2DEditor/Private/TileMapEditing/TileMapEditorViewportClient.cpp
--- a/Engine/Plugins/2D/Paper2D/Source/Paper2DEditor/Private/TileMapEditing/TileMapEditorViewportClient.cpp
+++ b/Engine/Plugins/2D/Paper2D/Source/Paper2DEditor/Private/TileMapEditing/TileMapEditorViewportClient.cpp
@@ -48,7 +48,7 @@ FTileMapEditorViewportClient::FTileMapEditorViewportClient(TWeakPtr<FTileMapEdit
 		RenderTileMapComponent = NewObject<UPaperTileMapComponent>();
 		UPaperTileMap* TileMap = GetTileMapBeingEdited();
 		RenderTileMapComponent->TileMap = TileMap;
-		GSelectedAnnotation.Set(RenderTileMapComponent);
+		RenderTileMapComponent->SelectionOverrideDelegate = UPrimitiveComponent::FSelectionOverride::CreateLambda([](const UPrimitiveComponent*) { return true; });
 
 		PreviewScene->AddComponent(RenderTileMapComponent, FTransform::Identity);
 	}
