Editor: Promote FAssetEditorModeManager into UnrealEd and make it the default for FEditorViewportClient::ModeTools when the derived class doesn't specify anything
- Prevents some types of selection highlighting issues where editors would use the global selection state incorrectly
[UE-11603]

diff --git a/Engine/Plugins/2D/Paper2D/Source/Paper2DEditor/Private/PaperEditorViewportClient.cpp b/Engine/Plugins/2D/Paper2D/Source/Paper2DEditor/Private/PaperEditorViewportClient.cpp
--- a/Engine/Plugins/2D/Paper2D/Source/Paper2DEditor/Private/PaperEditorViewportClient.cpp
+++ b/Engine/Plugins/2D/Paper2D/Source/Paper2DEditor/Private/PaperEditorViewportClient.cpp
@@ -7,54 +7,11 @@
 #include "SEditorViewport.h"
 #include "PreviewScene.h"
 
-//////////////////////////////////////////////////////////////////////////
-// FAssetEditorModeTools
-
-FAssetEditorModeTools::FAssetEditorModeTools()
-	: PreviewScene(nullptr)
-{
-	ActorSet = NewObject<USelection>();
-	ActorSet->SetFlags(RF_Transactional);
-	ActorSet->AddToRoot();
-
-	ObjectSet = NewObject<USelection>();
-	ObjectSet->SetFlags(RF_Transactional);
-	ObjectSet->AddToRoot();
-}
-
-FAssetEditorModeTools::~FAssetEditorModeTools()
-{
-	ActorSet->RemoveFromRoot();
-	ActorSet = nullptr;
-	ObjectSet->RemoveFromRoot();
-	ObjectSet = nullptr;
-}
-
-USelection* FAssetEditorModeTools::GetSelectedActors() const
-{
-	return ActorSet;
-}
-
-USelection* FAssetEditorModeTools::GetSelectedObjects() const
-{
-	return ObjectSet;
-}
-
-UWorld* FAssetEditorModeTools::GetWorld() const
-{
-	return (PreviewScene != nullptr) ? PreviewScene->GetWorld() : GEditor->GetEditorWorldContext().World();
-}
-
-void FAssetEditorModeTools::SetPreviewScene(class FPreviewScene* NewPreviewScene)
-{
-	PreviewScene = NewPreviewScene;
-}
-
 //////////////////////////////////////////////////////////////////////////
 // FPaperEditorViewportClient
 
 FPaperEditorViewportClient::FPaperEditorViewportClient(const TWeakPtr<SEditorViewport>& InEditorViewportWidget)
-	: FEditorViewportClient(new FAssetEditorModeTools(), nullptr, InEditorViewportWidget)
+	: FEditorViewportClient(new FAssetEditorModeManager(), nullptr, InEditorViewportWidget)
 	, CheckerboardTexture(nullptr)
 {
 	bOwnsModeTools = true;
diff --git a/Engine/Plugins/2D/Paper2D/Source/Paper2DEditor/Private/PaperEditorViewportClient.h b/Engine/Plugins/2D/Paper2D/Source/Paper2DEditor/Private/PaperEditorViewportClient.h
--- a/Engine/Plugins/2D/Paper2D/Source/Paper2DEditor/Private/PaperEditorViewportClient.h
+++ b/Engine/Plugins/2D/Paper2D/Source/Paper2DEditor/Private/PaperEditorViewportClient.h
@@ -3,6 +3,7 @@
 #pragma once
 
 #include "SceneViewport.h"
+#include "AssetEditorModeManager.h"
 
 struct FViewportSelectionRectangle
 {
@@ -11,28 +12,6 @@ struct FViewportSelectionRectangle
 	FLinearColor Color;
 };
 
-//////////////////////////////////////////////////////////////////////////
-// FAssetEditorModeTools
-
-class FAssetEditorModeTools : public FEditorModeTools
-{
-public:
-	FAssetEditorModeTools();
-	virtual ~FAssetEditorModeTools();
-
-	// FEditorModeTools interface
-	virtual class USelection* GetSelectedActors() const override;
-	virtual class USelection* GetSelectedObjects() const override;
-	virtual UWorld* GetWorld() const override;
-	// End of FEditorModeTools interface
-
-	void SetPreviewScene(class FPreviewScene* NewPreviewScene);
-protected:
-	class USelection* ActorSet;
-	class USelection* ObjectSet;
-	class FPreviewScene* PreviewScene;
-};
-
 //////////////////////////////////////////////////////////////////////////
 // FPaperEditorViewportClient
 
diff --git a/Engine/Plugins/2D/Paper2D/Source/Paper2DEditor/Private/TileMapEditing/TileMapEditorViewportClient.cpp b/Engine/Plugins/2D/Paper2D/Source/Paper2DEditor/Private/TileMapEditing/TileMapEditorViewportClient.cpp
--- a/Engine/Plugins/2D/Paper2D/Source/Paper2DEditor/Private/TileMapEditing/TileMapEditorViewportClient.cpp
+++ b/Engine/Plugins/2D/Paper2D/Source/Paper2DEditor/Private/TileMapEditing/TileMapEditorViewportClient.cpp
@@ -25,7 +25,7 @@ FTileMapEditorViewportClient::FTileMapEditorViewportClient(TWeakPtr<FTileMapEdit
 	check(TileMapEditorPtr.IsValid() && TileMapEditorViewportPtr.IsValid());
 
 	PreviewScene = &OwnedPreviewScene;
-	((FAssetEditorModeTools*)ModeTools)->SetPreviewScene(PreviewScene);
+	((FAssetEditorModeManager*)ModeTools)->SetPreviewScene(PreviewScene);
 
 	SetRealtime(true);
 
diff --git a/Engine/Source/Editor/UnrealEd/Private/AssetEditorModeManager.cpp b/Engine/Source/Editor/UnrealEd/Private/AssetEditorModeManager.cpp
new file mode 100644
--- /dev/null
+++ b/Engine/Source/Editor/UnrealEd/Private/AssetEditorModeManager.cpp
@@ -0,0 +1,49 @@
+// Copyright 1998-2015 Epic Games, Inc. All Rights Reserved.
+
+#include "UnrealEd.h"
+#include "AssetEditorModeManager.h"
+#include "PreviewScene.h"
+
+//////////////////////////////////////////////////////////////////////////
+// FAssetEditorModeManager
+
+FAssetEditorModeManager::FAssetEditorModeManager()
+	: PreviewScene(nullptr)
+{
+	ActorSet = NewObject<USelection>();
+	ActorSet->SetFlags(RF_Transactional);
+	ActorSet->AddToRoot();
+
+	ObjectSet = NewObject<USelection>();
+	ObjectSet->SetFlags(RF_Transactional);
+	ObjectSet->AddToRoot();
+}
+
+FAssetEditorModeManager::~FAssetEditorModeManager()
+{
+	ActorSet->RemoveFromRoot();
+	ActorSet = nullptr;
+	ObjectSet->RemoveFromRoot();
+	ObjectSet = nullptr;
+}
+
+USelection* FAssetEditorModeManager::GetSelectedActors() const
+{
+	return ActorSet;
+}
+
+USelection* FAssetEditorModeManager::GetSelectedObjects() const
+{
+	return ObjectSet;
+}
+
+UWorld* FAssetEditorModeManager::GetWorld() const
+{
+	return (PreviewScene != nullptr) ? PreviewScene->GetWorld() : GEditor->GetEditorWorldContext().World();
+}
+
+void FAssetEditorModeManager::SetPreviewScene(class FPreviewScene* NewPreviewScene)
+{
+	PreviewScene = NewPreviewScene;
+}
+
diff --git a/Engine/Source/Editor/UnrealEd/Private/EditorViewportClient.cpp b/Engine/Source/Editor/UnrealEd/Private/EditorViewportClient.cpp
--- a/Engine/Source/Editor/UnrealEd/Private/EditorViewportClient.cpp
+++ b/Engine/Source/Editor/UnrealEd/Private/EditorViewportClient.cpp
@@ -25,6 +25,7 @@
 #include "Components/BillboardComponent.h"
 #include "EngineUtils.h"
 #include "SEditorViewport.h"
+#include "AssetEditorModeManager.h"
 
 #define LOCTEXT_NAMESPACE "EditorViewportClient"
 
@@ -277,7 +278,7 @@ FEditorViewportClient::FEditorViewportClient(FEditorModeTools* InModeTools, FPre
 {
 	if (ModeTools == nullptr)
 	{
-		ModeTools = new FEditorModeTools();
+		ModeTools = new FAssetEditorModeManager();
 		bOwnsModeTools = true;
 	}
 
diff --git a/Engine/Source/Editor/UnrealEd/Public/AssetEditorModeManager.h b/Engine/Source/Editor/UnrealEd/Public/AssetEditorModeManager.h
new file mode 100644
--- /dev/null
+++ b/Engine/Source/Editor/UnrealEd/Public/AssetEditorModeManager.h
@@ -0,0 +1,27 @@
+// Copyright 1998-2015 Epic Games, Inc. All Rights Reserved.
+
+#pragma once
+
+#include "EditorModeManager.h"
+
+//////////////////////////////////////////////////////////////////////////
+// FAssetEditorModeManager
+
+class UNREALED_API FAssetEditorModeManager : public FEditorModeTools
+{
+public:
+	FAssetEditorModeManager();
+	virtual ~FAssetEditorModeManager();
+
+	// FEditorModeTools interface
+	virtual class USelection* GetSelectedActors() const override;
+	virtual class USelection* GetSelectedObjects() const override;
+	virtual UWorld* GetWorld() const override;
+	// End of FEditorModeTools interface
+
+	void SetPreviewScene(class FPreviewScene* NewPreviewScene);
+protected:
+	class USelection* ActorSet;
+	class USelection* ObjectSet;
+	class FPreviewScene* PreviewScene;
+};
