Paper2D: Change references to 'Sprite Flipbook' to be 'Paper Flipbook'

diff --git a/Engine/Plugins/2D/Paper2D/Source/Paper2DEditor/Private/FlipbookAssetTypeActions.cpp b/Engine/Plugins/2D/Paper2D/Source/Paper2DEditor/Private/FlipbookAssetTypeActions.cpp
--- a/Engine/Plugins/2D/Paper2D/Source/Paper2DEditor/Private/FlipbookAssetTypeActions.cpp
+++ b/Engine/Plugins/2D/Paper2D/Source/Paper2DEditor/Private/FlipbookAssetTypeActions.cpp
@@ -12,7 +12,7 @@
 
 FText FFlipbookAssetTypeActions::GetName() const
 {
-	return LOCTEXT("FFlipbookAssetTypeActionsName", "Sprite Flipbook");
+	return LOCTEXT("FFlipbookAssetTypeActionsName", "Paper Flipbook");
 }
 
 FColor FFlipbookAssetTypeActions::GetTypeColor() const
