Paper2D: Add search bar to the details panel in the sprite and flipbook editors [UE-9018]
[INTEGRATE] Change 2438811

diff --git a/Engine/Plugins/2D/Paper2D/Source/Paper2DEditor/Private/FlipbookEditor/FlipbookEditor.cpp b/Engine/Plugins/2D/Paper2D/Source/Paper2DEditor/Private/FlipbookEditor/FlipbookEditor.cpp
--- a/Engine/Plugins/2D/Paper2D/Source/Paper2DEditor/Private/FlipbookEditor/FlipbookEditor.cpp
+++ b/Engine/Plugins/2D/Paper2D/Source/Paper2DEditor/Private/FlipbookEditor/FlipbookEditor.cpp
@@ -186,7 +186,7 @@ public:
 	{
 		FlipbookEditorPtr = InFlipbookEditor;
 
-		SSingleObjectDetailsPanel::Construct(SSingleObjectDetailsPanel::FArguments());
+		SSingleObjectDetailsPanel::Construct(SSingleObjectDetailsPanel::FArguments(), /*bAutomaticallyObserveViaGetObjectToObserve=*/ true, /*bAllowSearch=*/ true);
 	}
 
 	// SSingleObjectDetailsPanel interface
diff --git a/Engine/Plugins/2D/Paper2D/Source/Paper2DEditor/Private/SpriteEditor/SpriteEditor.cpp b/Engine/Plugins/2D/Paper2D/Source/Paper2DEditor/Private/SpriteEditor/SpriteEditor.cpp
--- a/Engine/Plugins/2D/Paper2D/Source/Paper2DEditor/Private/SpriteEditor/SpriteEditor.cpp
+++ b/Engine/Plugins/2D/Paper2D/Source/Paper2DEditor/Private/SpriteEditor/SpriteEditor.cpp
@@ -252,7 +252,7 @@ public:
 	{
 		SpriteEditorPtr = InSpriteEditor;
 
-		SSingleObjectDetailsPanel::Construct(SSingleObjectDetailsPanel::FArguments());
+		SSingleObjectDetailsPanel::Construct(SSingleObjectDetailsPanel::FArguments(), /*bAutomaticallyObserveViaGetObjectToObserve=*/ true, /*bAllowSearch=*/ true);
 	}
 
 	// SSingleObjectDetailsPanel interface
