Tile Map Editor: Add a warning prompt visible when no tile set is selected; clicking on it opens the tile set picker [UE-11701]

diff --git a/Engine/Plugins/2D/Paper2D/Source/Paper2DEditor/Private/PaperStyle.cpp b/Engine/Plugins/2D/Paper2D/Source/Paper2DEditor/Private/PaperStyle.cpp
--- a/Engine/Plugins/2D/Paper2D/Source/Paper2DEditor/Private/PaperStyle.cpp
+++ b/Engine/Plugins/2D/Paper2D/Source/Paper2DEditor/Private/PaperStyle.cpp
@@ -7,6 +7,7 @@
 #define IMAGE_PLUGIN_BRUSH( RelativePath, ... ) FSlateImageBrush( FPaperStyle::InContent( RelativePath, ".png" ), __VA_ARGS__ )
 #define IMAGE_BRUSH(RelativePath, ...) FSlateImageBrush(StyleSet->RootToContentDir(RelativePath, TEXT(".png")), __VA_ARGS__)
 #define BOX_BRUSH(RelativePath, ...) FSlateBoxBrush(StyleSet->RootToContentDir(RelativePath, TEXT(".png")), __VA_ARGS__)
+#define TTF_FONT( RelativePath, ... ) FSlateFontInfo(StyleSet->RootToContentDir(RelativePath, TEXT(".ttf")), __VA_ARGS__)
 
 FString FPaperStyle::InContent(const FString& RelativePath, const ANSICHAR* Extension)
 {
@@ -81,6 +82,13 @@ void FPaperStyle::Initialize()
 			);
 
 		StyleSet->Set("TileMapEditor.LayerBrowser.SelectionColor", LayerSelectionColor);
+
+		StyleSet->Set("TileMapEditor.TileSetPalette.NothingSelectedText", FTextBlockStyle(FEditorStyle::Get().GetWidgetStyle<FTextBlockStyle>("NormalText"))
+			.SetFont(TTF_FONT("Fonts/Roboto-BoldCondensed", 18))
+			.SetColorAndOpacity(FLinearColor(0.8, 0.8f, 0.0f, 0.8f))
+			.SetShadowOffset(FVector2D(1.0f, 1.0f))
+			.SetShadowColorAndOpacity(FLinearColor(0.0f, 0.0f, 0.0f, 0.9f))
+			);
 	}
 
 	// Sprite editor
diff --git a/Engine/Plugins/2D/Paper2D/Source/Paper2DEditor/Private/TileMapEditing/TileMapEdModeToolkit.cpp b/Engine/Plugins/2D/Paper2D/Source/Paper2DEditor/Private/TileMapEditing/TileMapEdModeToolkit.cpp
--- a/Engine/Plugins/2D/Paper2D/Source/Paper2DEditor/Private/TileMapEditing/TileMapEdModeToolkit.cpp
+++ b/Engine/Plugins/2D/Paper2D/Source/Paper2DEditor/Private/TileMapEditing/TileMapEdModeToolkit.cpp
@@ -38,7 +38,7 @@ FText FTileMapEdModeToolkit::GetBaseToolkitName() const
 
 FText FTileMapEdModeToolkit::GetToolkitName() const
 {
-	if ( CurrentTileSetPtr.IsValid() )
+	if (CurrentTileSetPtr.IsValid())
 	{
 		const bool bDirtyState = CurrentTileSetPtr->GetOutermost()->IsDirty();
 
@@ -77,6 +77,39 @@ void FTileMapEdModeToolkit::Init(const TSharedPtr<IToolkitHost>& InitToolkitHost
 
 	TileSetPalette = SNew(STileSetSelectorViewport, CurrentTileSetPtr.Get(), TileMapEditor);
 
+	TSharedRef<SWidget> TileSetPaletteWidget = SNew(SOverlay)
+		// The palette widget
+		+SOverlay::Slot()
+		[
+			TileSetPalette.ToSharedRef()
+		]
+		// The no tile set selected warning text/button
+		+SOverlay::Slot()
+		.Padding(8.0f)
+		.VAlign(VAlign_Bottom)
+		.HAlign(HAlign_Right)
+		[
+			SNew(SButton)
+			.ButtonStyle(FEditorStyle::Get(), "NoBorder")
+			.Visibility(this, &FTileMapEdModeToolkit::GetTileSetPaletteCornerTextVisibility)
+			.OnClicked(this, &FTileMapEdModeToolkit::ClickedOnTileSetPaletteCornerText)
+			.Content()
+			[
+				SNew(STextBlock)
+				.TextStyle(FPaperStyle::Get(), "TileMapEditor.TileSetPalette.NothingSelectedText")
+				.Text(LOCTEXT("NoTileSetSelected", "Pick a tile set"))
+				.ToolTipText(LOCTEXT("NoTileSetSelectedTooltip", "A tile set must be selected before painting the tile map"))
+			]
+		];
+
+	TileSetAssetReferenceWidget = SNew(SContentReference)
+		.WidthOverride(ContentRefWidth)
+		.AssetReference(this, &FTileMapEdModeToolkit::GetCurrentTileSet)
+		.OnSetReference(this, &FTileMapEdModeToolkit::OnChangeTileSet)
+		.AllowedClass(UPaperTileSet::StaticClass())
+		.AllowSelectingNewAsset(true)
+		.AllowClearingReference(false);
+
 	// Create the contents of the editor mode toolkit
 	MyWidget = 
 		SNew(SBorder)
@@ -97,7 +130,6 @@ void FTileMapEdModeToolkit::Init(const TSharedPtr<IToolkitHost>& InitToolkitHost
 			.VAlign(VAlign_Fill)
 			[
 				SNew(SVerticalBox)
-				.Visibility(this, &FTileMapEdModeToolkit::GetTileSetSelectorVisibility)
 
 				+SVerticalBox::Slot()
 				.AutoHeight()
@@ -116,13 +148,7 @@ void FTileMapEdModeToolkit::Init(const TSharedPtr<IToolkitHost>& InitToolkitHost
 					.FillWidth(1.0f)
 					.HAlign(HAlign_Right)
 					[
-						SNew(SContentReference)
-						.WidthOverride(ContentRefWidth)
-						.AssetReference(this, &FTileMapEdModeToolkit::GetCurrentTileSet)
-						.OnSetReference(this, &FTileMapEdModeToolkit::OnChangeTileSet)
-						.AllowedClass(UPaperTileSet::StaticClass())
-						.AllowSelectingNewAsset(true)
-						.AllowClearingReference(false)
+						TileSetAssetReferenceWidget.ToSharedRef()
 					]
 				]
 
@@ -135,7 +161,7 @@ void FTileMapEdModeToolkit::Init(const TSharedPtr<IToolkitHost>& InitToolkitHost
 					+SHorizontalBox::Slot()
 					.HAlign(HAlign_Fill)
 					[
-						TileSetPalette.ToSharedRef()
+						TileSetPaletteWidget
 					]
 				]
 			]
@@ -223,11 +249,6 @@ bool FTileMapEdModeToolkit::IsToolSelected(ETileMapEditorTool::Type QueryTool) c
 	return (TileMapEditor->GetActiveTool() == QueryTool);
 }
 
-EVisibility FTileMapEdModeToolkit::GetTileSetSelectorVisibility() const
-{
-	return EVisibility::Visible;
-}
-
 TSharedRef<SWidget> FTileMapEdModeToolkit::BuildToolBar() const
 {
 	const FTileMapEditorCommands& Commands = FTileMapEditorCommands::Get();
@@ -281,6 +302,18 @@ TSharedRef<SWidget> FTileMapEdModeToolkit::BuildToolBar() const
 		];
 }
 
+EVisibility FTileMapEdModeToolkit::GetTileSetPaletteCornerTextVisibility() const
+{
+	return (GetCurrentTileSet() != nullptr) ? EVisibility::Collapsed : EVisibility::Visible;
+}
+
+FReply FTileMapEdModeToolkit::ClickedOnTileSetPaletteCornerText()
+{
+	TileSetAssetReferenceWidget->OpenAssetPickerMenu();
+
+	return FReply::Handled();
+}
+
 //////////////////////////////////////////////////////////////////////////
 
 #undef LOCTEXT_NAMESPACE
diff --git a/Engine/Plugins/2D/Paper2D/Source/Paper2DEditor/Private/TileMapEditing/TileMapEdModeToolkit.h b/Engine/Plugins/2D/Paper2D/Source/Paper2DEditor/Private/TileMapEditing/TileMapEdModeToolkit.h
--- a/Engine/Plugins/2D/Paper2D/Source/Paper2DEditor/Private/TileMapEditing/TileMapEdModeToolkit.h
+++ b/Engine/Plugins/2D/Paper2D/Source/Paper2DEditor/Private/TileMapEditing/TileMapEdModeToolkit.h
@@ -32,12 +32,14 @@ protected:
 
 	void BindCommands();
 
-	EVisibility GetTileSetSelectorVisibility() const;
-
 	void OnSelectTool(ETileMapEditorTool::Type NewTool);
 	bool IsToolSelected(ETileMapEditorTool::Type QueryTool) const;
 
 	TSharedRef<SWidget> BuildToolBar() const;
+
+	EVisibility GetTileSetPaletteCornerTextVisibility() const;
+	FReply ClickedOnTileSetPaletteCornerText();
+
 private:
 	class FEdModeTileMap* TileMapEditor;
 
@@ -48,4 +50,7 @@ private:
 
 	// The tile set selector palette
 	TSharedPtr<class STileSetSelectorViewport> TileSetPalette;
+
+	// The tile set asset reference widget
+	TSharedPtr<class SContentReference> TileSetAssetReferenceWidget;
 };
\ No newline at end of file
