Tile Map Editor: Improve the behavior of the tile preview cursor when nothing has been selected
- Now draws a red 1x1 square with a warning prompt when there is no ink for the paint or fill tools
- Added 'Paint' text back to the paintbrush tool and changed 'Eraser' to 'Erase' to match the other two
[UE-11702]

diff --git a/Engine/Plugins/2D/Paper2D/Source/Paper2DEditor/Private/TileMapEditing/EdModeTileMap.cpp b/Engine/Plugins/2D/Paper2D/Source/Paper2DEditor/Private/TileMapEditing/EdModeTileMap.cpp
--- a/Engine/Plugins/2D/Paper2D/Source/Paper2DEditor/Private/TileMapEditing/EdModeTileMap.cpp
+++ b/Engine/Plugins/2D/Paper2D/Source/Paper2DEditor/Private/TileMapEditing/EdModeTileMap.cpp
@@ -311,6 +311,9 @@ void FEdModeTileMap::Render(const FSceneView* View, FViewport* Viewport, FPrimit
 		return;
 	}
 
+	// Determine if the active tool is in a valid state
+	const bool bToolIsReadyToDraw = IsToolReadyToBeUsed();
+
 	// Draw the preview cursor
 	if (bIsLastCursorValid)
 	{
@@ -318,7 +321,7 @@ void FEdModeTileMap::Render(const FSceneView* View, FViewport* Viewport, FPrimit
 		{
 			// Slight depth bias so that the wireframe grid overlay doesn't z-fight with the tiles themselves
 			const float DepthBias = 0.0001f;
-			FLinearColor CursorWireColor = FLinearColor::White;
+			const FLinearColor CursorWireColor = bToolIsReadyToDraw ? FLinearColor::White : FLinearColor::Red;
 
 			const int32 CursorWidth = GetCursorWidth();
 			const int32 CursorHeight = GetCursorHeight();
@@ -339,19 +342,23 @@ void FEdModeTileMap::Render(const FSceneView* View, FViewport* Viewport, FPrimit
 void FEdModeTileMap::DrawHUD(FEditorViewportClient* ViewportClient, FViewport* Viewport, const FSceneView* View, FCanvas* Canvas)
 {
 	bool bDrawToolDescription = false;
-	FText ToolDescription = LOCTEXT("NoTool", "No tool selected");
+
+	static const FText UnknownTool = LOCTEXT("NoTool", "No tool selected");
+	static const FText NoTilesForTool = LOCTEXT("NoInkToolDesc", "No tile selected");
+
+	FText ToolDescription = UnknownTool;
 	switch (ActiveTool)
 	{
 	case ETileMapEditorTool::Eraser:
-		ToolDescription = LOCTEXT("EraserTool", "Eraser");
+		ToolDescription = LOCTEXT("EraserTool", "Erase");
 		bDrawToolDescription = true;
 		break;
 	case ETileMapEditorTool::Paintbrush:
-		ToolDescription = FText::GetEmpty();
+		ToolDescription = bHasValidInkSource ? LOCTEXT("BrushTool", "Paint") : NoTilesForTool;
 		bDrawToolDescription = true;
 		break;
 	case ETileMapEditorTool::PaintBucket:
-		ToolDescription = LOCTEXT("PaintBucketTool", "Fill");
+		ToolDescription = bHasValidInkSource ? LOCTEXT("PaintBucketTool", "Fill") : NoTilesForTool;
 		bDrawToolDescription = true;
 		break;
 	}
@@ -363,10 +370,13 @@ void FEdModeTileMap::DrawHUD(FEditorViewportClient* ViewportClient, FViewport* V
 		FVector2D ScreenSpacePreviewLocation;
 		if (View->WorldToPixel(DrawPreviewTopLeft, /*out*/ ScreenSpacePreviewLocation))
 		{
+			const bool bToolIsReadyToDraw = IsToolReadyToBeUsed();
+			const FLinearColor ToolPromptColor = bToolIsReadyToDraw ? FLinearColor::White : FLinearColor::Red;
+
 			int32 XL;
 			int32 YL;
 			StringSize(GEngine->GetLargeFont(), XL, YL, *ToolDescriptionString);
-			Canvas->DrawShadowedString(ScreenSpacePreviewLocation.X, ScreenSpacePreviewLocation.Y - YL, *ToolDescriptionString, GEngine->GetLargeFont(), FLinearColor::White);
+			Canvas->DrawShadowedString(ScreenSpacePreviewLocation.X, ScreenSpacePreviewLocation.Y - YL, *ToolDescriptionString, GEngine->GetLargeFont(), ToolPromptColor);
 		}
 	}
 
@@ -966,6 +976,27 @@ void FEdModeTileMap::RotateTilesInSelection(bool bIsClockwise)
 	}
 }
 
+bool FEdModeTileMap::IsToolReadyToBeUsed() const
+{
+	bool bToolIsReadyToDraw = false;
+	switch (ActiveTool)
+	{
+	case ETileMapEditorTool::Paintbrush:
+		bToolIsReadyToDraw = bHasValidInkSource;
+		break;
+	case ETileMapEditorTool::Eraser:
+		bToolIsReadyToDraw = true;
+		break;
+	case ETileMapEditorTool::PaintBucket:
+		bToolIsReadyToDraw = bHasValidInkSource;
+		break;
+	default:
+		check(false);
+		break;
+	}
+
+	return bToolIsReadyToDraw;
+}
 void FEdModeTileMap::RotateSelectionCW()
 {
 	RotateTilesInSelection(/*bIsClockwise=*/ true);
@@ -1131,20 +1162,20 @@ int32 FEdModeTileMap::GetBrushHeight() const
 
 int32 FEdModeTileMap::GetCursorWidth() const
 {
-	const int32 CursorWidth = (ActiveTool == ETileMapEditorTool::PaintBucket) ? 1 : GetBrushWidth();
+	const int32 CursorWidth = ((ActiveTool == ETileMapEditorTool::PaintBucket) || !bHasValidInkSource) ? 1 : GetBrushWidth();
 	return CursorWidth;
 }
 
 int32 FEdModeTileMap::GetCursorHeight() const
 {
-	const int32 CursorHeight = (ActiveTool == ETileMapEditorTool::PaintBucket) ? 1 : GetBrushHeight();
+	const int32 CursorHeight = ((ActiveTool == ETileMapEditorTool::PaintBucket) || !bHasValidInkSource) ? 1 : GetBrushHeight();
 	return CursorHeight;
 }
 
 void FEdModeTileMap::RefreshBrushSize()
 {
 	const bool bShowPreviewDesired = !DrawPreviewDimensionsLS.IsNearlyZero();
-	
+
 	switch (ActiveTool)
 	{
 	case ETileMapEditorTool::Paintbrush:
diff --git a/Engine/Plugins/2D/Paper2D/Source/Paper2DEditor/Private/TileMapEditing/EdModeTileMap.h b/Engine/Plugins/2D/Paper2D/Source/Paper2DEditor/Private/TileMapEditing/EdModeTileMap.h
--- a/Engine/Plugins/2D/Paper2D/Source/Paper2DEditor/Private/TileMapEditing/EdModeTileMap.h
+++ b/Engine/Plugins/2D/Paper2D/Source/Paper2DEditor/Private/TileMapEditing/EdModeTileMap.h
@@ -89,6 +89,8 @@ protected:
 	void SynchronizePreviewWithTileMap(UPaperTileMap* NewTileMap);
 
 	void RotateTilesInSelection(bool bIsClockwise);
+
+	bool IsToolReadyToBeUsed() const;
 public:
 	UPaperTileMapComponent* FindSelectedComponent() const;
 
