Tile Map Editor: Minor style tweaks
- Changed 'Active Set' text to 'Active Tile Set', right-aligned and vertically centered next to the picker widget
- Added a 'click here' line to the 'no tile set selected' prompt

diff --git a/Engine/Plugins/2D/Paper2D/Source/Paper2DEditor/Private/TileMapEditing/TileMapEdModeToolkit.cpp b/Engine/Plugins/2D/Paper2D/Source/Paper2DEditor/Private/TileMapEditing/TileMapEdModeToolkit.cpp
--- a/Engine/Plugins/2D/Paper2D/Source/Paper2DEditor/Private/TileMapEditing/TileMapEdModeToolkit.cpp
+++ b/Engine/Plugins/2D/Paper2D/Source/Paper2DEditor/Private/TileMapEditing/TileMapEdModeToolkit.cpp
@@ -98,7 +98,7 @@ void FTileMapEdModeToolkit::Init(const TSharedPtr<IToolkitHost>& InitToolkitHost
 				SNew(STextBlock)
 				.TextStyle(FPaperStyle::Get(), "TileMapEditor.TileSetPalette.NothingSelectedText")
 				.Text(LOCTEXT("NoTileSetSelected", "Pick a tile set"))
-				.ToolTipText(LOCTEXT("NoTileSetSelectedTooltip", "A tile set must be selected before painting the tile map"))
+				.ToolTipText(LOCTEXT("NoTileSetSelectedTooltip", "A tile set must be selected before painting the tile map.\nClick here to select one."))
 			]
 		];
 
@@ -138,15 +138,17 @@ void FTileMapEdModeToolkit::Init(const TSharedPtr<IToolkitHost>& InitToolkitHost
 				[
 					SNew(SHorizontalBox)
 					+SHorizontalBox::Slot()
-					.AutoWidth()
-					.HAlign(HAlign_Left)
+					.FillWidth(1.0f)
+					.HAlign(HAlign_Right)
+					.VAlign(VAlign_Center)
+					.Padding(FMargin(0.0f, 0.0f, 4.0f, 0.0f))
 					[
 						SNew(STextBlock)
-						.Text(LOCTEXT("CurrentTileSetAssetToPaintWith", "Active Set"))
+						.Text(LOCTEXT("CurrentTileSetAssetToPaintWith", "Active Tile Set"))
 					]
 					+SHorizontalBox::Slot()
-					.FillWidth(1.0f)
 					.HAlign(HAlign_Right)
+					.AutoWidth()
 					[
 						TileSetAssetReferenceWidget.ToSharedRef()
 					]
