Tile Map Editor: Nudge the tile map preview cursor out slightly in PaperAxisZ so it doesn't z-fight with the real tile map [UE-11703]

diff --git a/Engine/Plugins/2D/Paper2D/Source/Paper2DEditor/Private/TileMapEditing/EdModeTileMap.cpp b/Engine/Plugins/2D/Paper2D/Source/Paper2DEditor/Private/TileMapEditing/EdModeTileMap.cpp
--- a/Engine/Plugins/2D/Paper2D/Source/Paper2DEditor/Private/TileMapEditing/EdModeTileMap.cpp
+++ b/Engine/Plugins/2D/Paper2D/Source/Paper2DEditor/Private/TileMapEditing/EdModeTileMap.cpp
@@ -1080,7 +1080,12 @@ void FEdModeTileMap::UpdatePreviewCursor(const FViewportCursorLocation& Ray)
 
 			DrawPreviewDimensionsLS = 0.5f*((PaperAxisX * CursorWidth * TileMap->TileWidth) + (PaperAxisY * -CursorHeight * TileMap->TileHeight));
 
-			const FVector ComponentPreviewLocation = ComponentToWorld.TransformPosition(TileMap->GetTileCenterInLocalSpace(LocalTileX0, LocalTileY0, LayerIndex));
+			// Figure out how far to nudge out the tile map (we want a decent size (especially if the layer separation is small), but should never be a full layer out)
+			const float DepthBiasNudge = -FMath::Min(FMath::Max(1.0f, TileMap->SeparationPerLayer * 0.05f), TileMap->SeparationPerLayer * 0.5f);
+
+			const FVector ComponentPreviewLocationNoNudge = ComponentToWorld.TransformPosition(TileMap->GetTileCenterInLocalSpace(LocalTileX0, LocalTileY0, LayerIndex));
+			const FVector ComponentPreviewLocation = ComponentPreviewLocationNoNudge + (PaperAxisZ * DepthBiasNudge);
+
 			CursorPreviewComponent->SetWorldLocation(ComponentPreviewLocation);
 			CursorPreviewComponent->SetWorldRotation(FRotator(ComponentToWorld.GetRotation()));
 			CursorPreviewComponent->SetWorldScale3D(ComponentToWorld.GetScale3D());
