Tile Map Editing: Move 'Tile Map' and 'Setup' categories above Materials category

diff --git a/Engine/Plugins/2D/Paper2D/Source/Paper2DEditor/Private/TileMapEditing/PaperTileMapDetailsCustomization.cpp b/Engine/Plugins/2D/Paper2D/Source/Paper2DEditor/Private/TileMapEditing/PaperTileMapDetailsCustomization.cpp
--- a/Engine/Plugins/2D/Paper2D/Source/Paper2DEditor/Private/TileMapEditing/PaperTileMapDetailsCustomization.cpp
+++ b/Engine/Plugins/2D/Paper2D/Source/Paper2DEditor/Private/TileMapEditing/PaperTileMapDetailsCustomization.cpp
@@ -61,7 +61,12 @@ void FPaperTileMapDetailsCustomization::CustomizeDetails(IDetailLayoutBuilder& D
 	TileMapPtr = TileMap;
 	TileMapComponentPtr = TileComponent;
 
-	IDetailCategoryBuilder& TileMapCategory = DetailLayout.EditCategory("Tile Map");
+	// Make sure the Tile Map category is right below the Transform
+	IDetailCategoryBuilder& TileMapCategory = DetailLayout.EditCategory("Tile Map", FText::GetEmpty(), ECategoryPriority::Important);
+
+	// Make sure the setup category is near the top
+	DetailLayout.EditCategory("Setup", FText::GetEmpty(), ECategoryPriority::Important);
+
 
 	TAttribute<EVisibility> InternalInstanceVis = TAttribute<EVisibility>::Create(TAttribute<EVisibility>::FGetter::CreateSP(this, &FPaperTileMapDetailsCustomization::GetVisibilityForInstancedOnlyProperties));
 
@@ -179,9 +184,6 @@ void FPaperTileMapDetailsCustomization::CustomizeDetails(IDetailLayoutBuilder& D
 			}
 		}
 	}
-
-	// Make sure the setup category is near the top
-	DetailLayout.EditCategory("Setup");
 }
 
 FReply FPaperTileMapDetailsCustomization::EnterTileMapEditingMode()
