Sprite Editor: Fixed an issue with how the transform gizmo worked in Edit Source Region mode on rotated sprites
[UE-12293]

diff --git a/Engine/Plugins/2D/Paper2D/Source/Paper2D/Classes/PaperSprite.h b/Engine/Plugins/2D/Paper2D/Source/Paper2D/Classes/PaperSprite.h
--- a/Engine/Plugins/2D/Paper2D/Source/Paper2D/Classes/PaperSprite.h
+++ b/Engine/Plugins/2D/Paper2D/Source/Paper2D/Classes/PaperSprite.h
@@ -183,7 +183,7 @@ public:
 	FVector ConvertTextureSpaceToPivotSpace(FVector Input) const;
 	FVector ConvertPivotSpaceToTextureSpace(FVector Input) const;
 	FVector2D ConvertWorldSpaceToTextureSpace(const FVector& WorldPoint) const;
-	FVector2D ConvertWorldSpaceDeltaToTextureSpace(const FVector& WorldDelta) const;
+	FVector2D ConvertWorldSpaceDeltaToTextureSpace(const FVector& WorldDelta, bool bIgnoreRotation = false) const;
 
 	// World space WRT the sprite editor *only*
 	FVector ConvertTextureSpaceToWorldSpace(const FVector2D& SourcePoint) const;
diff --git a/Engine/Plugins/2D/Paper2D/Source/Paper2D/Private/PaperSprite.cpp b/Engine/Plugins/2D/Paper2D/Source/Paper2D/Private/PaperSprite.cpp
--- a/Engine/Plugins/2D/Paper2D/Source/Paper2D/Private/PaperSprite.cpp
+++ b/Engine/Plugins/2D/Paper2D/Source/Paper2D/Private/PaperSprite.cpp
@@ -1647,7 +1647,7 @@ FVector2D UPaperSprite::ConvertWorldSpaceToTextureSpace(const FVector& WorldPoin
 	return ConvertPivotSpaceToTextureSpace(FVector2D(XValue, YValue));
 }
 
-FVector2D UPaperSprite::ConvertWorldSpaceDeltaToTextureSpace(const FVector& WorldSpaceDelta) const
+FVector2D UPaperSprite::ConvertWorldSpaceDeltaToTextureSpace(const FVector& WorldSpaceDelta, bool bIgnoreRotation) const
 {
 	const FVector ProjectionX = WorldSpaceDelta.ProjectOnTo(PaperAxisX);
 	const FVector ProjectionY = WorldSpaceDelta.ProjectOnTo(PaperAxisY);
@@ -1656,7 +1656,7 @@ FVector2D UPaperSprite::ConvertWorldSpaceDeltaToTextureSpace(const FVector& Worl
 	float YValue = FMath::Sign(ProjectionY | PaperAxisY) * ProjectionY.Size() * PixelsPerUnrealUnit;
 
 	// Undo pivot space rotation, ignoring pivot position
-	if (bRotatedInSourceImage)
+	if (bRotatedInSourceImage && !bIgnoreRotation)
 	{
 		Swap(XValue, YValue);
 		XValue = -XValue;
diff --git a/Engine/Plugins/2D/Paper2D/Source/Paper2DEditor/Private/SpriteEditor/SpriteEditorSelections.h b/Engine/Plugins/2D/Paper2D/Source/Paper2DEditor/Private/SpriteEditor/SpriteEditorSelections.h
--- a/Engine/Plugins/2D/Paper2D/Source/Paper2DEditor/Private/SpriteEditor/SpriteEditorSelections.h
+++ b/Engine/Plugins/2D/Paper2D/Source/Paper2DEditor/Private/SpriteEditor/SpriteEditorSelections.h
@@ -60,8 +60,9 @@ public:
 		{
 			UTexture2D* SourceTexture = Sprite->GetSourceTexture();
 			const FVector2D SourceDims = (SourceTexture != nullptr) ? FVector2D(SourceTexture->GetSurfaceWidth(), SourceTexture->GetSurfaceHeight()) : FVector2D::ZeroVector;
-            int XAxis = 0; // -1 = min, 0 = none, 1 = max
-            int YAxis = 0; // ditto
+
+			int32 XAxis = 0; // -1 = min, 0 = none, 1 = max
+            int32 YAxis = 0; // ditto
             
 			switch (VertexIndex)
 			{
@@ -97,11 +98,11 @@ public:
 				break;
 			}
 
-			const FVector2D TextureSpaceDelta = Sprite->ConvertWorldSpaceDeltaToTextureSpace(PaperAxisX * WorldSpaceDelta.X + PaperAxisY * WorldSpaceDelta.Y);
+			const FVector2D TextureSpaceDelta = Sprite->ConvertWorldSpaceDeltaToTextureSpace(PaperAxisX * WorldSpaceDelta.X + PaperAxisY * WorldSpaceDelta.Y, /*bIgnoreRotation=*/ true);
 
             if (XAxis == -1)
             {
-                float AllowedDelta = FMath::Clamp(TextureSpaceDelta.X, -Sprite->SourceUV.X, Sprite->SourceDimension.X - 1);
+                const float AllowedDelta = FMath::Clamp(TextureSpaceDelta.X, -Sprite->SourceUV.X, Sprite->SourceDimension.X - 1);
                 Sprite->SourceUV.X += AllowedDelta;
                 Sprite->SourceDimension.X -= AllowedDelta;
             }
@@ -112,7 +113,7 @@ public:
             
             if (YAxis == -1)
             {
-				float AllowedDelta = FMath::Clamp(TextureSpaceDelta.Y, -Sprite->SourceUV.Y, Sprite->SourceDimension.Y - 1);
+				const float AllowedDelta = FMath::Clamp(TextureSpaceDelta.Y, -Sprite->SourceUV.Y, Sprite->SourceDimension.Y - 1);
 				Sprite->SourceUV.Y += AllowedDelta;
 				Sprite->SourceDimension.Y -= AllowedDelta;
             }
