Sprite Editor: Move the Alternate Material property to the Rendering category and hid it unless the mode is Diced [UE-10086]

diff --git a/Engine/Plugins/2D/Paper2D/Source/Paper2D/Classes/PaperSprite.h b/Engine/Plugins/2D/Paper2D/Source/Paper2D/Classes/PaperSprite.h
--- a/Engine/Plugins/2D/Paper2D/Source/Paper2D/Classes/PaperSprite.h
+++ b/Engine/Plugins/2D/Paper2D/Source/Paper2D/Classes/PaperSprite.h
@@ -99,7 +99,7 @@ protected:
 	UMaterialInterface* DefaultMaterial;
 
 	// The alternate material to use on a sprite instance if not overridden (this is only used for Diced render geometry, and will be the opaque material in that case, slot 1)
-	UPROPERTY(Category=Sprite, EditAnywhere, BlueprintReadOnly)
+	UPROPERTY(Category=Rendering, EditAnywhere, BlueprintReadOnly)
 	UMaterialInterface* AlternateMaterial;
 
 	// List of sockets on this sprite
diff --git a/Engine/Plugins/2D/Paper2D/Source/Paper2DEditor/Private/SpriteEditor/SpriteDetailsCustomization.cpp b/Engine/Plugins/2D/Paper2D/Source/Paper2DEditor/Private/SpriteEditor/SpriteDetailsCustomization.cpp
--- a/Engine/Plugins/2D/Paper2D/Source/Paper2DEditor/Private/SpriteEditor/SpriteDetailsCustomization.cpp
+++ b/Engine/Plugins/2D/Paper2D/Source/Paper2DEditor/Private/SpriteEditor/SpriteDetailsCustomization.cpp
@@ -97,9 +97,16 @@ void FSpriteDetailsCustomization::BuildRenderingSection(IDetailCategoryBuilder&
 {
 	// Add the rendering geometry mode into the parent container (renamed)
 	const FString RenderGeometryTypePropertyPath = FString::Printf(TEXT("%s.%s"), GET_MEMBER_NAME_STRING_CHECKED(UPaperSprite, RenderGeometry), GET_MEMBER_NAME_STRING_CHECKED(FSpritePolygonCollection, GeometryType));
-	RenderingCategory.AddProperty(DetailLayout.GetProperty(*RenderGeometryTypePropertyPath))
+	TSharedPtr<IPropertyHandle> RenderGeometryTypeProperty = DetailLayout.GetProperty(*RenderGeometryTypePropertyPath);
+	RenderingCategory.AddProperty(RenderGeometryTypeProperty)
 		.DisplayName(LOCTEXT("RenderGeometryType", "Render Geometry Type"));
 
+	// Show the alternate material, but only when the mode is Diced
+	TAttribute<EVisibility> ShowWhenModeIsDiced = TAttribute<EVisibility>::Create(TAttribute<EVisibility>::FGetter::CreateSP(this, &FSpriteDetailsCustomization::PolygonModeMatches, RenderGeometryTypeProperty, ESpritePolygonMode::Diced));
+	TSharedPtr<IPropertyHandle> AlternateMaterialProperty = DetailLayout.GetProperty(GET_MEMBER_NAME_CHECKED(UPaperSprite, AlternateMaterial));
+	RenderingCategory.AddProperty(AlternateMaterialProperty)
+		.Visibility(ShowWhenModeIsDiced);
+
 	// Show the rendering geometry settings
 	TSharedRef<IPropertyHandle> RenderGeometry = DetailLayout.GetProperty(GET_MEMBER_NAME_CHECKED(UPaperSprite, RenderGeometry));
 	IDetailPropertyRow& RenderGeometryProperty = RenderingCategory.AddProperty(RenderGeometry);
@@ -368,4 +375,21 @@ EVisibility FSpriteDetailsCustomization::Get2DPhysicsNotEnabledWarningVisibility
 	return GetDefault<UPhysicsSettings>()->bEnable2DPhysics ? EVisibility::Collapsed : PhysicsModeMatches(Property, ESpriteCollisionMode::Use2DPhysics);
 }
 
+EVisibility FSpriteDetailsCustomization::PolygonModeMatches(TSharedPtr<IPropertyHandle> Property, ESpritePolygonMode::Type DesiredMode) const
+{
+	if (Property.IsValid())
+	{
+		uint8 ValueAsByte;
+		FPropertyAccess::Result Result = Property->GetValue(/*out*/ ValueAsByte);
+
+		if (Result == FPropertyAccess::Success)
+		{
+			return (((ESpritePolygonMode::Type)ValueAsByte) == DesiredMode) ? EVisibility::Visible : EVisibility::Collapsed;
+		}
+	}
+
+	// If there are multiple values, show all properties
+	return EVisibility::Visible;
+}
+
 #undef LOCTEXT_NAMESPACE
diff --git a/Engine/Plugins/2D/Paper2D/Source/Paper2DEditor/Private/SpriteEditor/SpriteDetailsCustomization.h b/Engine/Plugins/2D/Paper2D/Source/Paper2DEditor/Private/SpriteEditor/SpriteDetailsCustomization.h
--- a/Engine/Plugins/2D/Paper2D/Source/Paper2DEditor/Private/SpriteEditor/SpriteDetailsCustomization.h
+++ b/Engine/Plugins/2D/Paper2D/Source/Paper2DEditor/Private/SpriteEditor/SpriteDetailsCustomization.h
@@ -23,6 +23,8 @@ protected:
 	EVisibility GetCustomPivotVisibility(TSharedPtr<IPropertyHandle> Property) const;
 	EVisibility Get2DPhysicsNotEnabledWarningVisibility(TSharedPtr<IPropertyHandle> Property) const;
 
+	EVisibility PolygonModeMatches(TSharedPtr<IPropertyHandle> Property, ESpritePolygonMode::Type DesiredMode) const;
+
 	static EVisibility GetAtlasGroupVisibility();
 
 	static FDetailWidgetRow& GenerateWarningRow(IDetailCategoryBuilder& WarningCategory, bool bExperimental, const FText& WarningText, const FText& Tooltip, const FString& ExcerptLink, const FString& ExcerptName);
