Tile Map Editor: Prevent cursor preview from being invisible when SeparationPerLayer is negative

diff --git a/Engine/Plugins/2D/Paper2D/Source/Paper2DEditor/Private/TileMapEditing/EdModeTileMap.cpp b/Engine/Plugins/2D/Paper2D/Source/Paper2DEditor/Private/TileMapEditing/EdModeTileMap.cpp
--- a/Engine/Plugins/2D/Paper2D/Source/Paper2DEditor/Private/TileMapEditing/EdModeTileMap.cpp
+++ b/Engine/Plugins/2D/Paper2D/Source/Paper2DEditor/Private/TileMapEditing/EdModeTileMap.cpp
@@ -1081,7 +1081,8 @@ void FEdModeTileMap::UpdatePreviewCursor(const FViewportCursorLocation& Ray)
 			DrawPreviewDimensionsLS = 0.5f*((PaperAxisX * CursorWidth * TileMap->TileWidth) + (PaperAxisY * -CursorHeight * TileMap->TileHeight));
 
 			// Figure out how far to nudge out the tile map (we want a decent size (especially if the layer separation is small), but should never be a full layer out)
-			const float DepthBiasNudge = -FMath::Min(FMath::Max(1.0f, TileMap->SeparationPerLayer * 0.05f), TileMap->SeparationPerLayer * 0.5f);
+			const float AbsoluteSeparation = FMath::Abs(TileMap->SeparationPerLayer);
+			const float DepthBiasNudge = -FMath::Min(FMath::Max(1.0f, AbsoluteSeparation * 0.05f), AbsoluteSeparation * 0.5f);
 
 			const FVector ComponentPreviewLocationNoNudge = ComponentToWorld.TransformPosition(TileMap->GetTileCenterInLocalSpace(LocalTileX0, LocalTileY0, LayerIndex));
 			const FVector ComponentPreviewLocation = ComponentPreviewLocationNoNudge + (PaperAxisZ * DepthBiasNudge);
