Paper2D: Fixed an issue where reimporting a JSON sprite sheet wouldn't reimport any existing texture pages
[UE-12041]

diff --git a/Engine/Plugins/2D/Paper2D/Source/PaperJsonImporter/Private/PaperJsonSpriteSheetImporter.cpp b/Engine/Plugins/2D/Paper2D/Source/PaperJsonImporter/Private/PaperJsonSpriteSheetImporter.cpp
--- a/Engine/Plugins/2D/Paper2D/Source/PaperJsonImporter/Private/PaperJsonSpriteSheetImporter.cpp
+++ b/Engine/Plugins/2D/Paper2D/Source/PaperJsonImporter/Private/PaperJsonSpriteSheetImporter.cpp
@@ -407,17 +407,30 @@ bool FPaperJsonSpriteSheetImporter::ImportTextures(const FString& LongPackagePat
 	TArray<FString> SheetFileNames;
 	SheetFileNames.Add(SourceSheetImageFilename);
 
-	if (bIsReimporting && ExistingTextureName == Image && ExistingTexture != nullptr)
+	if (bIsReimporting && (ExistingTextureName == Image) && (ExistingTexture != nullptr))
 	{
 		ImageTexture = ExistingTexture;
 		bFoundExistingTexture = true;
 	}
 
 	// Import the texture if we need to
-	if (!bIsReimporting || (bIsReimporting && !bFoundExistingTexture))
+	bool bNeedsImport = !bIsReimporting || (bIsReimporting && !bFoundExistingTexture);
+
+	if (bIsReimporting && bFoundExistingTexture)
+	{
+		if (!FReimportManager::Instance()->Reimport(ImageTexture, /*bAskForNewFileIfMissing=*/ true))
+		{
+			bNeedsImport = true;
+		}
+	}
+
+	if (bNeedsImport)
 	{
 		FAssetToolsModule& AssetToolsModule = FModuleManager::GetModuleChecked<FAssetToolsModule>("AssetTools");
-		TArray<UObject*> ImportedSheets = AssetToolsModule.Get().ImportAssets(SheetFileNames, TargetSubPath);
+
+		UReimportTextureFactory* ReimportFactory = nullptr;
+
+		TArray<UObject*> ImportedSheets = AssetToolsModule.Get().ImportAssets(SheetFileNames, TargetSubPath, ReimportFactory);
 
 		UTexture2D* ImportedTexture = (ImportedSheets.Num() > 0) ? Cast<UTexture2D>(ImportedSheets[0]) : nullptr;
 
@@ -429,6 +442,7 @@ bool FPaperJsonSpriteSheetImporter::ImportTextures(const FString& LongPackagePat
 		else
 		{
 			// Change the compression settings
+			//@TODO: Should we always be doing this (particularly the TF_Nearest seems potentially undesirable; maybe make it a property of the sprite sheet asset?)
 			ImportedTexture->Modify();
 			ImportedTexture->LODGroup = TEXTUREGROUP_UI;
 			ImportedTexture->CompressionSettings = TC_EditorIcon;
@@ -477,17 +491,6 @@ bool FPaperJsonSpriteSheetImporter::PerformImport(const FString& LongPackagePath
 
 		UObject* OuterForFrame = nullptr; // @TODO: Use this if we don't want them to be individual assets - Flipbook;
 
-		// Create a unique package name and asset name for the frame
-		const FString SanitizedFrameName = ObjectTools::SanitizeObjectName(Frame.FrameName.ToString());
-		const FString TentativePackagePath = PackageTools::SanitizePackageName(TargetSubPath + TEXT("/") + SanitizedFrameName);
-		FString DefaultSuffix;
-		FString AssetName;
-		FString PackageName;
-		AssetToolsModule.Get().CreateUniqueAssetName(TentativePackagePath, /*out*/ DefaultSuffix, /*out*/ PackageName, /*out*/ AssetName);
-
-		// Create a package for the frame
-		OuterForFrame = CreatePackage(nullptr, *PackageName);
-
 		// Create a frame in the package
 		UPaperSprite* TargetSprite = nullptr;
 		
@@ -498,6 +501,21 @@ bool FPaperJsonSpriteSheetImporter::PerformImport(const FString& LongPackagePath
 		
 		if (TargetSprite == nullptr)
 		{
+			const FString SanitizedFrameName = ObjectTools::SanitizeObjectName(Frame.FrameName.ToString());
+			const FString TentativePackagePath = PackageTools::SanitizePackageName(TargetSubPath + TEXT("/") + SanitizedFrameName);
+			FString DefaultSuffix;
+			FString AssetName;
+			FString PackageName;
+			AssetToolsModule.Get().CreateUniqueAssetName(TentativePackagePath, /*out*/ DefaultSuffix, /*out*/ PackageName, /*out*/ AssetName);
+
+			// Create a unique package name and asset name for the frame
+			if (OuterForFrame == nullptr)
+			{
+				// Create a package for the frame
+				OuterForFrame = CreatePackage(nullptr, *PackageName);
+			}
+
+			// Create the asset
 			TargetSprite = NewObject<UPaperSprite>(OuterForFrame, *AssetName, Flags);
 			FAssetRegistryModule::AssetCreated(TargetSprite);
 		}
