Paper2D: Fix an issue with the default value of PixelsPerUnrealUnit not being initialized in cooked builds, which would cause various issues for sprites with it set to 2.56
- Initialized PixelsPerUnrealUnits in cooked builds
- Added a couple of potentially needed ConditionalPostLoad calls to the sprite resizing logic
[UE-7918]

diff --git a/Engine/Plugins/2D/Paper2D/Source/Paper2D/Classes/PaperSprite.h b/Engine/Plugins/2D/Paper2D/Source/Paper2D/Classes/PaperSprite.h
--- a/Engine/Plugins/2D/Paper2D/Source/Paper2D/Classes/PaperSprite.h
+++ b/Engine/Plugins/2D/Paper2D/Source/Paper2D/Classes/PaperSprite.h
@@ -196,7 +196,7 @@ public:
 	FVector2D GetPivotPosition() const;
 
 	// Rescale properties to handle source texture size change
-	void RescaleSpriteData(const UTexture2D* Texture);
+	void RescaleSpriteData(UTexture2D* Texture);
 	bool NeedRescaleSpriteData();
 
 	void RebuildCollisionData();
diff --git a/Engine/Plugins/2D/Paper2D/Source/Paper2D/Private/PaperSprite.cpp b/Engine/Plugins/2D/Paper2D/Source/Paper2D/Private/PaperSprite.cpp
--- a/Engine/Plugins/2D/Paper2D/Source/Paper2D/Private/PaperSprite.cpp
+++ b/Engine/Plugins/2D/Paper2D/Source/Paper2D/Private/PaperSprite.cpp
@@ -332,21 +332,21 @@ UPaperSprite::UPaperSprite(const FObjectInitializer& ObjectInitializer)
 	
 	AlternateMaterialSplitIndex = INDEX_NONE;
 
-#if WITH_EDITOR
+#if WITH_EDITORONLY_DATA
 	PivotMode = ESpritePivotMode::Center_Center;
 	bSnapPivotToPixelGrid = true;
 
 	CollisionGeometry.GeometryType = ESpritePolygonMode::TightBoundingBox;
 	CollisionThickness = 10.0f;
 
-	PixelsPerUnrealUnit = 2.56f;
-
 	bTrimmedInSourceImage = false;
 	bRotatedInSourceImage = false;
 
 	SourceTextureDimension.Set(0, 0);
 #endif
 
+	PixelsPerUnrealUnit = 2.56f;
+
 	static ConstructorHelpers::FObjectFinder<UMaterialInterface> MaskedMaterialRef(TEXT("/Paper2D/MaskedUnlitSpriteMaterial"));
 	DefaultMaterial = MaskedMaterialRef.Object;
 
@@ -586,8 +586,9 @@ void UPaperSprite::PostEditChangeProperty(FPropertyChangedEvent& PropertyChanged
 	Super::PostEditChangeProperty(PropertyChangedEvent);
 }
 
-void UPaperSprite::RescaleSpriteData(const UTexture2D* Texture)
+void UPaperSprite::RescaleSpriteData(UTexture2D* Texture)
 {
+	Texture->ConditionalPostLoad();
 	FVector2D PreviousTextureDimension = SourceTextureDimension;
 	FVector2D NewTextureDimension(Texture->GetImportedSize().X, Texture->GetImportedSize().Y);
 
@@ -678,6 +679,7 @@ bool UPaperSprite::NeedRescaleSpriteData()
 {
 	if (UTexture2D* Texture = GetSourceTexture())
 	{
+		Texture->ConditionalPostLoad();
 		FIntPoint TextureSize = Texture->GetImportedSize();
 		bool bTextureSizeIsZero = TextureSize.X == 0 || TextureSize.Y == 0;
 		return !SourceTextureDimension.IsZero() && !bTextureSizeIsZero && (TextureSize.X != SourceTextureDimension.X || TextureSize.Y != SourceTextureDimension.Y);
