Paper2D: Added a minor optimization to UPaperFlipbook::FindSocket to prevent an unnecessary search when another component is attached but not to a socket

diff --git a/Engine/Plugins/2D/Paper2D/Source/Paper2D/Private/PaperFlipbook.cpp b/Engine/Plugins/2D/Paper2D/Source/Paper2D/Private/PaperFlipbook.cpp
--- a/Engine/Plugins/2D/Paper2D/Source/Paper2D/Private/PaperFlipbook.cpp
+++ b/Engine/Plugins/2D/Paper2D/Source/Paper2D/Private/PaperFlipbook.cpp
@@ -148,7 +148,7 @@ FBoxSphereBounds UPaperFlipbook::GetRenderBounds() const
 
 bool UPaperFlipbook::FindSocket(FName SocketName, int32 KeyFrameIndex, FTransform& OutLocalTransform)
 {
-	if (KeyFrames.IsValidIndex(KeyFrameIndex))
+	if (KeyFrames.IsValidIndex(KeyFrameIndex) && (SocketName != NAME_None))
 	{
 		if (UPaperSprite* SpriteFrame = KeyFrames[KeyFrameIndex].Sprite)
 		{
