Prevented a few more uses of Pop() from shrinking

diff --git a/Engine/Plugins/2D/Paper2D/Source/Paper2DEditor/Private/ExtractSprites/SPaperExtractSpritesDialog.cpp b/Engine/Plugins/2D/Paper2D/Source/Paper2DEditor/Private/ExtractSprites/SPaperExtractSpritesDialog.cpp
--- a/Engine/Plugins/2D/Paper2D/Source/Paper2DEditor/Private/ExtractSprites/SPaperExtractSpritesDialog.cpp
+++ b/Engine/Plugins/2D/Paper2D/Source/Paper2DEditor/Private/ExtractSprites/SPaperExtractSpritesDialog.cpp
@@ -244,7 +244,7 @@ void SPaperExtractSpritesDialog::PreviewExtractedSprites()
 				DominatedSprites.Empty(SpritesLeft.Num());
 				while (SpritesLeft.Num())
 				{
-					FIntRect DominatingSprite = SpritesLeft.Pop();
+					FIntRect DominatingSprite = SpritesLeft.Pop(/*bAllowShrinking=*/ false);
 					DominatedSprites.Add(DominatingSprite);
 
 					// Find the sprites that are dominated (intersect the infinite horizontal band described by the dominating sprite)
diff --git a/Engine/Source/Editor/BlueprintGraph/Private/EdGraphSchema_K2.cpp b/Engine/Source/Editor/BlueprintGraph/Private/EdGraphSchema_K2.cpp
--- a/Engine/Source/Editor/BlueprintGraph/Private/EdGraphSchema_K2.cpp
+++ b/Engine/Source/Editor/BlueprintGraph/Private/EdGraphSchema_K2.cpp
@@ -5512,7 +5512,7 @@ void UEdGraphSchema_K2::SplitPin(UEdGraphPin* Pin) const
 			SubPin->ParentPin = Pin;
 
 			// CreatePin puts the Pin in the array, but we are going to insert it later, so pop it back out
-			GraphNode->Pins.Pop();
+			GraphNode->Pins.Pop(/*bAllowShrinking=*/ false);
 
 			Pin->SubPins.Add(SubPin);
 		}
diff --git a/Engine/Source/Editor/BlueprintGraph/Private/K2Node.cpp b/Engine/Source/Editor/BlueprintGraph/Private/K2Node.cpp
--- a/Engine/Source/Editor/BlueprintGraph/Private/K2Node.cpp
+++ b/Engine/Source/Editor/BlueprintGraph/Private/K2Node.cpp
@@ -770,7 +770,7 @@ ERenamePinResult UK2Node::RenameUserDefinedPin(const FString& OldName, const FSt
 
 			while (PinsToUpdate.Num() > 0)
 			{
-				UEdGraphPin* PinToRename = PinsToUpdate.Pop();
+				UEdGraphPin* PinToRename = PinsToUpdate.Pop(/*bAllowShrinking=*/ false);
 				if (PinToRename->SubPins.Num() > 0)
 				{
 					PinsToUpdate.Append(PinToRename->SubPins);
