Paper2D: Coding standards fixes and file name sanitization for flipbook extraction

diff --git a/Engine/Plugins/2D/Paper2D/Source/Paper2DEditor/Private/PaperFlipbookHelpers.cpp b/Engine/Plugins/2D/Paper2D/Source/Paper2DEditor/Private/PaperFlipbookHelpers.cpp
--- a/Engine/Plugins/2D/Paper2D/Source/Paper2DEditor/Private/PaperFlipbookHelpers.cpp
+++ b/Engine/Plugins/2D/Paper2D/Source/Paper2DEditor/Private/PaperFlipbookHelpers.cpp
@@ -3,13 +3,14 @@
 #include "Paper2DEditorPrivatePCH.h"
 #include "PaperFlipbookHelpers.h"
 #include "Json.h"
+#include "ObjectTools.h"
 
 //////////////////////////////////////////////////////////////////////////
 // FPaperFlipbookHelpers
 
-static FString CleanName(const FString& Name)
+FString FPaperFlipbookHelpers::GetCleanerSpriteName(const FString& Name)
 {
-	int LastCharacter = Name.Len() - 1;
+	int32 LastCharacter = Name.Len() - 1;
 
 	// Strip off "_Sprite"
 	if (Name.Right(7) == TEXT("_Sprite"))
@@ -32,11 +33,11 @@ static FString CleanName(const FString& Name)
 	}
 }
 
-static bool ExtractSpriteNumber(const FString& String, FString& BareString, int& Number)
+bool FPaperFlipbookHelpers::ExtractSpriteNumber(const FString& String, FString& BareString, int32& Number)
 {
 	bool bExtracted = false;
 
-	int LastCharacter = String.Len() - 1;
+	int32 LastCharacter = String.Len() - 1;
 	if (LastCharacter >= 0)
 	{
 		// Find the last character that isn't a digit (Handle sprite names with numbers inside inverted commas / parentheses)
@@ -55,8 +56,8 @@ static bool ExtractSpriteNumber(const FString& String, FString& BareString, int&
 
 			if (LastCharacter >= 0)
 			{
-				int FirstDigit = LastCharacter;
-				int EndCharacter = FirstDigit;
+				const int32 FirstDigit = LastCharacter;
+				int32 EndCharacter = FirstDigit;
 				while (EndCharacter > 0 && !FChar::IsAlnum(String[EndCharacter - 1]))
 				{
 					EndCharacter--;
@@ -70,8 +71,8 @@ static bool ExtractSpriteNumber(const FString& String, FString& BareString, int&
 					EndCharacter = FirstDigit;
 				}
 
-				FString NumberString = String.Mid(FirstDigit);
-				BareString =  (EndCharacter > 0) ? CleanName(String.Left(EndCharacter)) : CleanName(String);
+				const FString NumberString = String.Mid(FirstDigit);
+				BareString = GetCleanerSpriteName((EndCharacter > 0) ? String.Left(EndCharacter) : String);
 				Number = FCString::Atoi(*NumberString);
 				bExtracted = true;
 			}
@@ -114,16 +115,12 @@ void FPaperFlipbookHelpers::ExtractFlipbooksFromSprites(TMap<FString, TArray<UPa
 
 		SpriteNameMap.Add(SpriteName, Sprite);
 
-		int SpriteNumber = 0;
+		int32 SpriteNumber = 0;
 		FString SpriteBareString;
-		if (ExtractSpriteNumber(SpriteName, /*out*/SpriteBareString, /*out*/SpriteNumber))
+		if (ExtractSpriteNumber(SpriteName, /*out*/ SpriteBareString, /*out*/ SpriteNumber))
 		{
-			if (!OutSpriteFlipbookMap.Contains(SpriteBareString))
-			{
-				OutSpriteFlipbookMap.Add(SpriteBareString, TArray<UPaperSprite*>());
-			}
-
-			OutSpriteFlipbookMap[SpriteBareString].Add(Sprite);
+			SpriteBareString = ObjectTools::SanitizeObjectName(SpriteBareString);
+			OutSpriteFlipbookMap.FindOrAdd(SpriteBareString).Add(Sprite);
 		}
 		else
 		{
@@ -140,14 +137,14 @@ void FPaperFlipbookHelpers::ExtractFlipbooksFromSprites(TMap<FString, TArray<UPa
 		bool operator()(UPaperSprite& LHS, UPaperSprite& RHS) const
 		{
 			FString LeftString;
-			int LeftNumber;
-			ExtractSpriteNumber(LHS.GetName(), /*out*/LeftString, /*out*/LeftNumber);
+			int32 LeftNumber;
+			ExtractSpriteNumber(LHS.GetName(), /*out*/ LeftString, /*out*/ LeftNumber);
 
 			FString RightString;
-			int RightNumber;
-			ExtractSpriteNumber(RHS.GetName(), /*out*/RightString, /*out*/RightNumber);
+			int32 RightNumber;
+			ExtractSpriteNumber(RHS.GetName(), /*out*/ RightString, /*out*/ RightNumber);
 
-			return (LeftString == RightString) ? LeftNumber < RightNumber : LeftString < RightString;
+			return (LeftString == RightString) ? (LeftNumber < RightNumber) : (LeftString < RightString);
 		}
 	};
 
@@ -160,10 +157,14 @@ void FPaperFlipbookHelpers::ExtractFlipbooksFromSprites(TMap<FString, TArray<UPa
 	}
 
 	// Create a flipbook from all remaining sprites
-	// Not sure if this is desirable behaviour, might want one flipbook per sprite
+	// Not sure if this is desirable behavior, might want one flipbook per sprite
 	if (RemainingSprites.Num() > 0)
 	{
 		RemainingSprites.Sort(FSpriteSortPredicate());
-		OutSpriteFlipbookMap.Add(CleanName(RemainingSprites[0]->GetName()) + TEXT("_Flipbook"), RemainingSprites);
+
+		const FString DesiredName = GetCleanerSpriteName(RemainingSprites[0]->GetName()) + TEXT("_Flipbook");
+		const FString SanitizedName = ObjectTools::SanitizeObjectName(DesiredName);
+
+		OutSpriteFlipbookMap.Add(SanitizedName, RemainingSprites);
 	}
 }
diff --git a/Engine/Plugins/2D/Paper2D/Source/Paper2DEditor/Public/PaperFlipbookHelpers.h b/Engine/Plugins/2D/Paper2D/Source/Paper2DEditor/Public/PaperFlipbookHelpers.h
--- a/Engine/Plugins/2D/Paper2D/Source/Paper2DEditor/Public/PaperFlipbookHelpers.h
+++ b/Engine/Plugins/2D/Paper2D/Source/Paper2DEditor/Public/PaperFlipbookHelpers.h
@@ -19,6 +19,10 @@ public:
 	static void ExtractFlipbooksFromSprites(TMap<FString, TArray<UPaperSprite*> >& OutSpriteFlipbookMap, const TArray<UPaperSprite*>& Sprites, const TArray<FString>& SpriteNames);
 
 private:
+
+	static FString GetCleanerSpriteName(const FString& Name);
+	static bool ExtractSpriteNumber(const FString& String, FString& BareString, int32& Number);
+
 	FPaperFlipbookHelpers() {}
 };
 
