Paper2D: Fixed some whitespace/comment issues

diff --git a/Engine/Plugins/2D/Paper2D/Source/PaperJsonImporter/Private/PaperJsonSpriteSheetImporter.cpp b/Engine/Plugins/2D/Paper2D/Source/PaperJsonImporter/Private/PaperJsonSpriteSheetImporter.cpp
--- a/Engine/Plugins/2D/Paper2D/Source/PaperJsonImporter/Private/PaperJsonSpriteSheetImporter.cpp
+++ b/Engine/Plugins/2D/Paper2D/Source/PaperJsonImporter/Private/PaperJsonSpriteSheetImporter.cpp
@@ -506,7 +506,7 @@ bool FPaperJsonSpriteSheetImporter::PerformImport(const FString& LongPackagePath
 			FString DefaultSuffix;
 			FString AssetName;
 			FString PackageName;
-			AssetToolsModule.Get().CreateUniqueAssetName(TentativePackagePath, /*out*/ DefaultSuffix, /*out*/ PackageName, /*out*/ AssetName);
+			AssetToolsModule.Get().CreateUniqueAssetName(TentativePackagePath, DefaultSuffix, /*out*/ PackageName, /*out*/ AssetName);
 
 			// Create a unique package name and asset name for the frame
 			if (OuterForFrame == nullptr)
diff --git a/Engine/Plugins/2D/Paper2D/Source/PaperJsonImporter/Private/PaperSpriteSheetAssetTypeActions.cpp b/Engine/Plugins/2D/Paper2D/Source/PaperJsonImporter/Private/PaperSpriteSheetAssetTypeActions.cpp
--- a/Engine/Plugins/2D/Paper2D/Source/PaperJsonImporter/Private/PaperSpriteSheetAssetTypeActions.cpp
+++ b/Engine/Plugins/2D/Paper2D/Source/PaperJsonImporter/Private/PaperSpriteSheetAssetTypeActions.cpp
@@ -109,7 +109,7 @@ void FPaperSpriteSheetAssetTypeActions::ExecuteCreateFlipbooks(TArray<TWeakObjec
 					}
 				}
 
-				FPaperFlipbookHelpers::ExtractFlipbooksFromSprites(/*out*/SpriteFlipbookMap, Sprites, SpriteNames);
+				FPaperFlipbookHelpers::ExtractFlipbooksFromSprites(/*out*/ SpriteFlipbookMap, Sprites, SpriteNames);
 			}
 
 			// Create one flipbook for every grouped flipbook name
@@ -133,7 +133,7 @@ void FPaperSpriteSheetAssetTypeActions::ExecuteCreateFlipbooks(TArray<TWeakObjec
 					FString DefaultSuffix;
 					FString AssetName;
 					FString PackageName;
-					AssetToolsModule.Get().CreateUniqueAssetName(TentativePackagePath, /*out*/ DefaultSuffix, /*out*/ PackageName, /*out*/ AssetName);
+					AssetToolsModule.Get().CreateUniqueAssetName(TentativePackagePath, DefaultSuffix, /*out*/ PackageName, /*out*/ AssetName);
 
 					FlipbookFactory->KeyFrames.Empty();
 					for (int32 SpriteIndex = 0; SpriteIndex < Sprites.Num(); ++SpriteIndex)
diff --git a/Engine/Plugins/2D/Paper2D/Source/PaperTiledImporter/Private/PaperTiledImporterFactory.cpp b/Engine/Plugins/2D/Paper2D/Source/PaperTiledImporter/Private/PaperTiledImporterFactory.cpp
--- a/Engine/Plugins/2D/Paper2D/Source/PaperTiledImporter/Private/PaperTiledImporterFactory.cpp
+++ b/Engine/Plugins/2D/Paper2D/Source/PaperTiledImporter/Private/PaperTiledImporterFactory.cpp
@@ -384,7 +384,7 @@ UObject* UPaperTiledImporterFactory::CreateNewAsset(UClass* AssetClass, const FS
 	FString DefaultSuffix;
 	FString AssetName;
 	FString PackageName;
-	AssetToolsModule.Get().CreateUniqueAssetName(TentativePackagePath, /*out*/ DefaultSuffix, /*out*/ PackageName, /*out*/ AssetName);
+	AssetToolsModule.Get().CreateUniqueAssetName(TentativePackagePath, DefaultSuffix, /*out*/ PackageName, /*out*/ AssetName);
 
 	// Create a package for the asset
 	UObject* OuterForAsset = CreatePackage(nullptr, *PackageName);
