Paper2D: Significantly improved the compatibility of the Tiled JSON importer
- Now parses all defined data as of Tiled 0.11.0, although many aspects are still lost during conversion to internal types

diff --git a/Engine/Plugins/2D/Paper2D/Source/PaperTiledImporter/Classes/PaperTiledImporterFactory.h b/Engine/Plugins/2D/Paper2D/Source/PaperTiledImporter/Classes/PaperTiledImporterFactory.h
--- a/Engine/Plugins/2D/Paper2D/Source/PaperTiledImporter/Classes/PaperTiledImporterFactory.h
+++ b/Engine/Plugins/2D/Paper2D/Source/PaperTiledImporter/Classes/PaperTiledImporterFactory.h
@@ -6,17 +6,56 @@
 #include "PaperTileLayer.h"
 #include "PaperTiledImporterFactory.generated.h"
 
-namespace ETiledOrientation
+enum class ETiledLayerType : uint8
 {
-	enum Type
-	{
-		Unknown,
-		Orthogonal,
-		Isometric,
-		Staggered,
-		Hexagonal
-	};
-}
+	TileLayer,
+	ObjectGroup,
+	ImageLayer
+};
+
+enum class ETiledOrientation : uint8
+{
+	Unknown,
+	Orthogonal,
+	Isometric,
+	Staggered,
+	Hexagonal
+};
+
+enum class ETiledObjectLayerDrawOrder : uint8
+{
+	TopDown,
+	Index
+};
+
+enum class ETiledStaggerAxis : uint8
+{
+	X,
+	Y
+};
+
+enum class ETiledStaggerIndex : uint8
+{
+	Odd,
+	Even
+};
+
+enum class ETiledRenderOrder : uint8
+{
+	RightDown,
+	RightUp,
+	LeftDown,
+	LeftUp
+};
+
+enum class ETiledObjectType : uint8
+{
+	Box,
+	Ellipse,
+	Polygon,
+	Polyline,
+	PlacedTile
+};
 
 // Imports a tile map (and associated textures & tile sets) exported from Tiled (http://www.mapeditor.org/)
 UCLASS()
@@ -46,89 +85,292 @@ protected:
 	static UTexture2D* ImportTexture(const FString& SourceFilename, const FString& TargetSubPath);
 };
 
+//////////////////////////////////////////////////////////////////////////
+// FTiledStringPair
+
+struct FTiledStringPair
+{
+	FString Key;
+	FString Value;
+
+public:
+	FTiledStringPair()
+	{
+	}
+
+	FTiledStringPair(const FString& InKey, const FString& InValue)
+		: Key(InKey)
+		, Value(InValue)
+	{
+	}
+
+	static void ParsePropertyBag(TArray<FTiledStringPair>& OutProperties, TSharedPtr<FJsonObject> Tree, const FString& NameForErrors, bool bSilent);
+};
+
+//////////////////////////////////////////////////////////////////////////
+// FTiledObject - A single object/shape placed inside a layer or tile
+
+// See [REDACTED_PROFILE_URL]/tiled/wiki/TMX-Map-Format#object for more information
+
+struct FTiledObject
+{
+	// The type of the object
+	ETiledObjectType TiledObjectType;
+
+	// ID of the shape (always 0 for per-tile collision shapes)
+	int32 ID;
+
+	// Arbitrary user-specified name
+	FString Name;
+
+	// Arbitrary user-specified type
+	FString UserType;
+
+	// Is the shape currently visible?
+	bool bVisible;
+
+	// Position of the shape:
+	//   Ellipse or Box: The center
+	//   PlacedTile: Bottom left
+	//   Polygon or Polyline: Position of first vertex (not their center, the Tiled editor always bakes down rotations for them!)
+	double X;
+	double Y;
+
+	// Local-space width/height of the shape (Ellipse, Box)
+	// Note: Set to (0,0) for Polygon, Polyline, PlacedTile
+	double Width;
+	double Height;
+
+	// Rotation (in degrees)
+	// Note: Always 0 for Polyline and Polygon
+	// Note: Includes winding (can be < 0 or > 360)
+	double RotationDegrees;
+
+	// Arbitrary user specified Key-Value pairs
+	TArray<FTiledStringPair> Properties;
+
+	// Points for polygon/polyline shapes, relative to (X,Y), which seems to always be the first point in the polygon
+	// Note: Only used when TiledObjectType is Polygon or Polyline
+	TArray<FVector2D> Points;
+
+	// The tile GID for placed tiles
+	// Note: Only used when TiledObjectType is PlacedTile
+	uint32 TileGID;
+
+public:
+	FTiledObject();
+
+	bool ParseFromJSON(TSharedPtr<FJsonObject> Tree, const FString& NameForErrors, bool bSilent = false);
+
+	static bool ParsePointArray(TArray<FVector2D>& OutPoints, const TArray<TSharedPtr<FJsonValue>>& InArray, const FString& NameForErrors, bool bSilent);
+};
+
+//////////////////////////////////////////////////////////////////////////
+// FTileLayerFromTiled - A layer, containing either tiles or objects
+
+// See [REDACTED_PROFILE_URL]/tiled/wiki/TMX-Map-Format#layer for more information
 
 struct FTileLayerFromTiled
 {
 	// Name of the layer
 	FString Name;
 
-	// Array of tiles
+	// Array of tiles (only used when LayerType is TileLayer)
 	TArray<uint32> TileIndices;
 
 	// Width and height in tiles
 	int32 Width;
 	int32 Height;
 
-	// Saved layer opacity
+	// Color of the layer (only set for object layers, to help distinguish them)
+	FColor Color;
+
+	// Object draw order (only used for object layers)
+	ETiledObjectLayerDrawOrder ObjectDrawOrder;
+
+	// Saved layer opacity (only RGB are used, A is ignored)
 	float Opacity;
 
-	// Is the layer currently visible
+	// Is the layer currently visible?
 	bool bVisible;
 
-	// 'tilelayer' ??? - not sure what this is
-	FString Type;
+	// Type of the layer
+	ETiledLayerType LayerType;
 
 	// Offset
 	int32 OffsetX;
 	int32 OffsetY;
 
+	// Placed objects (only used when LayerType is ObjectGroup)
+	TArray<FTiledObject> Objects;
+
+	// Overlay image (only used when LayerType is ImageLayer)
+	FString OverlayImagePath;
 
+	// Arbitrary user specified Key-Value pairs
+	TArray<FTiledStringPair> Properties;
+
+public:
 	FTileLayerFromTiled();
 
-	void ParseFromJSON(TSharedPtr<FJsonObject> Tree, const FString& NameForErrors, bool bSilent = false);
+	bool ParseFromJSON(TSharedPtr<FJsonObject> Tree, const FString& NameForErrors, bool bSilent = false);
 
 	bool IsValid() const;
 };
 
+//////////////////////////////////////////////////////////////////////////
+// FTiledTerrain
+
+// See [REDACTED_PROFILE_URL]/tiled/wiki/TMX-Map-Format#terrain for more information
+
+struct FTiledTerrain
+{
+	// The name of this terrain type
+	FString TerrainName;
+
+	// The index of the solid tile for this terrain (local  index to the tile set, not a GID!)
+	uint32 SolidTileLocalIndex;
+
+public:
+	FTiledTerrain();
+
+	bool ParseFromJSON(TSharedPtr<FJsonObject> Tree, const FString& NameForErrors, bool bSilent = false);
+};
+
+//////////////////////////////////////////////////////////////////////////
+// FTiledTileInfo
+
+// See [REDACTED_PROFILE_URL]/tiled/wiki/TMX-Map-Format#tile for more information
+
+struct FTiledTileInfo
+{
+	// The terrain indices (into the TerrainTypes array of the containing tile set, INDEX_NONE if invalid)
+	// Order is top-left, top-right, bottom-left, bottom-right
+	int32 TerrainIndices[4];
+
+	// The probability of placement of this tile (0..1)
+	// Note: Tiled doesn't allow editing of this value right now (at least as of 0.11.0), so it's of limited value/trustworthiness
+	float Probability;
+
+	//@TODO: image?
+
+	// Collision shapes
+	TArray<FTiledObject> Objects;
+
+	// Arbitrary user specified Key-Value pairs
+	TArray<FTiledStringPair> Properties;
+
+public:
+	FTiledTileInfo();
+
+	bool ParseTileInfoFromJSON(int32 TileIndex, TSharedPtr<FJsonObject> Tree, const FString& NameForErrors, bool bSilent = false);
+};
+
+//////////////////////////////////////////////////////////////////////////
+// FTileSetFromTiled
+
+// See [REDACTED_PROFILE_URL]/tiled/wiki/TMX-Map-Format#tileset for more information
+
 struct FTileSetFromTiled
 {
 	int32 FirstGID;
 	FString Name;
 
+	// Source image path
 	FString ImagePath;
+
+	// Source image dimensions
 	int32 ImageWidth;
 	int32 ImageHeight;
 
-	//Properties
-	//TileProperties
+	// Should we chroma-key remove certain pixels, making them transparent?
+	bool bRemoveTransparentColor;
+
+	// The color to remove if bRemoveTransparentColor is true
+	FColor ImageTransparentColor;
+
+	// Arbitrary user specified Key-Value pairs
+	TArray<FTiledStringPair> Properties;
+
+	// Terrain types
+	TArray<FTiledTerrain> TerrainTypes;
+
+	// Per-tile info (terrain membership, collision objects, properties, etc...)
+	TMap<int32, FTiledTileInfo> PerTileData;
+
+	// Offset used when drawing tiles from this tile set
 	int32 TileOffsetX;
 	int32 TileOffsetY;
 
+	// The spacing to ignore around the outer edge of the source image (in pixels)
 	int32 Margin;
+
+	// The spacing between each tile in the source image (in pixels)
 	int32 Spacing;
 
 	int32 TileWidth;
 	int32 TileHeight;
 
+public:
 	FTileSetFromTiled();
 
-	void ParseFromJSON(TSharedPtr<FJsonObject> Tree, const FString& NameForErrors, bool bSilent = false);
+	void ParseTileSetFromJSON(TSharedPtr<FJsonObject> Tree, const FString& NameForErrors, bool bSilent = false);
 
 	bool IsValid() const;
 };
 
+//////////////////////////////////////////////////////////////////////////
+// FTileMapFromTiled
+
+// See [REDACTED_PROFILE_URL]/tiled/wiki/TMX-Map-Format#map for more information
 
 struct FTileMapFromTiled
 {
+	// JSON export file version as defined by Tiled (1 is the only known version)
 	int32 FileVersion;
+
+	// Dimensions of the tile map (in tiles)
 	int32 Width;
 	int32 Height;
+
+	// Dimensions of a tile (in pixels)
 	int32 TileWidth;
 	int32 TileHeight;
 
-	ETiledOrientation::Type Orientation;
+	// Projection mode of the tile map
+	ETiledOrientation Orientation;
 
+	// Side length (only used in Hexagonal projection mode)
 	int32 HexSideLength;
 
+	// Stagger axis (only used in Staggered and Hexagonal modes)
+	ETiledStaggerAxis StaggerAxis;
+
+	// Stagger index (only used in Staggered and Hexagonal modes)
+	ETiledStaggerIndex StaggerIndex;
+
+	// Render order
+	ETiledRenderOrder RenderOrder;
+
+	// Background color
+	FColor BackgroundColor;
+
+	// Set of source tile sets imported from Tiled
 	TArray<FTileSetFromTiled> TileSets;
+
+	// Set of destination tile set assets created by this import
 	TArray<UPaperTileSet*> CreatedTileSetAssets;
+
+	// Layers
 	TArray<FTileLayerFromTiled> Layers;
-	//map of properties
 
+	// Arbitrary user specified Key-Value pairs
+	TArray<FTiledStringPair> Properties;
+
+public:
 	FPaperTileInfo ConvertTileGIDToPaper2D(uint32 GID) const;
 	ETileMapProjectionMode::Type GetOrientationType() const;
 
 	FTileMapFromTiled();
 	bool IsValid() const;
 };
-
diff --git a/Engine/Plugins/2D/Paper2D/Source/PaperTiledImporter/Private/PaperTiledImporterFactory.cpp b/Engine/Plugins/2D/Paper2D/Source/PaperTiledImporter/Private/PaperTiledImporterFactory.cpp
--- a/Engine/Plugins/2D/Paper2D/Source/PaperTiledImporter/Private/PaperTiledImporterFactory.cpp
+++ b/Engine/Plugins/2D/Paper2D/Source/PaperTiledImporter/Private/PaperTiledImporterFactory.cpp
@@ -16,22 +16,23 @@
 	if (!bSilent) { UE_LOG(LogPaperTiledImporter, Warning, FormatString, __VA_ARGS__); }
 
 //////////////////////////////////////////////////////////////////////////
-// FRequiredIntField
+// FRequiredScalarField
 
-struct FRequiredIntField
+template<typename ScalarType>
+struct FRequiredScalarField
 {
-	int32& Value;
+	ScalarType& Value;
 	const FString Key;
-	int32 MinValue;
+	ScalarType MinValue;
 
-	FRequiredIntField(int32& InValue, const FString& InKey)
+	FRequiredScalarField(ScalarType& InValue, const FString& InKey)
 		: Value(InValue)
 		, Key(InKey)
-		, MinValue(1)
+		, MinValue(1.0f)
 	{
 	}
 
-	FRequiredIntField(int32& InValue, const FString& InKey, int32 InMinValue)
+	FRequiredScalarField(ScalarType& InValue, const FString& InKey, ScalarType InMinValue)
 		: Value(InValue)
 		, Key(InKey)
 		, MinValue(InMinValue)
@@ -39,13 +40,19 @@ struct FRequiredIntField
 	}
 };
 
-bool ParseIntegerFields(FRequiredIntField* IntFieldArray, int32 ArrayCount, TSharedPtr<FJsonObject> Tree, const FString& NameForErrors, bool bSilent)
+typedef FRequiredScalarField<int32> FRequiredIntField;
+typedef FRequiredScalarField<double> FRequiredDoubleField;
+
+//////////////////////////////////////////////////////////////////////////
+
+template <typename ScalarType>
+bool ParseScalarFields(FRequiredScalarField<ScalarType>* FieldArray, int32 ArrayCount, TSharedPtr<FJsonObject> Tree, const FString& NameForErrors, bool bSilent)
 {
 	bool bSuccessfullyParsed = true;
 
 	for (int32 ArrayIndex = 0; ArrayIndex < ArrayCount; ++ArrayIndex)
 	{
-		const FRequiredIntField& Field = IntFieldArray[ArrayIndex];
+		const FRequiredScalarField<ScalarType>& Field = FieldArray[ArrayIndex];
 
 		if (!Tree->TryGetNumberField(Field.Key, /*out*/ Field.Value))
 		{
@@ -67,6 +74,10 @@ bool ParseIntegerFields(FRequiredIntField* IntFieldArray, int32 ArrayCount, TSha
 	return bSuccessfullyParsed;
 }
 
+bool ParseIntegerFields(FRequiredIntField* IntFieldArray, int32 ArrayCount, TSharedPtr<FJsonObject> Tree, const FString& NameForErrors, bool bSilent)
+{
+	return ParseScalarFields(IntFieldArray, ArrayCount, Tree, NameForErrors, bSilent);
+}
 
 //////////////////////////////////////////////////////////////////////////
 // UPaperTiledImporterFactory
@@ -144,6 +155,13 @@ UObject* UPaperTiledImporterFactory::FactoryCreateText(UClass* InClass, UObject*
 			UE_LOG(LogPaperTiledImporter, Warning, TEXT("JSON exported from Tiled in file '%s' has an unknown version %d (expected version 1).  Parsing will continue but some things may not import correctly"), *NameForErrors, GlobalInfo.FileVersion);
 		}
 
+		// Parse the global properties
+		const TSharedPtr<FJsonObject>* PropertiesSubobject;
+		if (DescriptorObject->TryGetObjectField(TEXT("properties"), /*out*/ PropertiesSubobject))
+		{
+			FTiledStringPair::ParsePropertyBag(/*out*/ GlobalInfo.Properties, *PropertiesSubobject, NameForErrors, /*bSilent=*/ false);
+		}
+
 		// Load the tile sets
 		const TArray<TSharedPtr<FJsonValue>>* TileSetDescriptors;
 		if (DescriptorObject->TryGetArrayField(TEXT("tilesets"), /*out*/ TileSetDescriptors))
@@ -151,7 +169,7 @@ UObject* UPaperTiledImporterFactory::FactoryCreateText(UClass* InClass, UObject*
 			for (TSharedPtr<FJsonValue> TileSetDescriptor : *TileSetDescriptors)
 			{
 				FTileSetFromTiled& TileSet = *new (GlobalInfo.TileSets) FTileSetFromTiled();
-				TileSet.ParseFromJSON(TileSetDescriptor->AsObject(), NameForErrors);
+				TileSet.ParseTileSetFromJSON(TileSetDescriptor->AsObject(), NameForErrors);
 				bLoadedSuccessfully = bLoadedSuccessfully && TileSet.IsValid();
 			}
 		}
@@ -178,11 +196,6 @@ UObject* UPaperTiledImporterFactory::FactoryCreateText(UClass* InClass, UObject*
 			bLoadedSuccessfully = false;
 		}
 
-		// Load the properties
-		//@TODO: What to do with Tiled properties?
-
-
-
 		// Create the new tile map asset and import basic/global data
 		Result = NewObject<UPaperTileMap>(InParent, InName, Flags);
 
@@ -196,6 +209,7 @@ UObject* UPaperTiledImporterFactory::FactoryCreateText(UClass* InClass, UObject*
 		Result->SeparationPerLayer = -1.0f;
 		Result->ProjectionMode = GlobalInfo.GetOrientationType();
 		Result->PixelsPerUnrealUnit = GetDefault<UPaperRuntimeSettings>()->DefaultPixelsPerUnrealUnit;
+		Result->BackgroundColor = GlobalInfo.BackgroundColor;
 
 		// Create the tile sets
 		for (const FTileSetFromTiled& TileSetData : GlobalInfo.TileSets)
@@ -214,8 +228,6 @@ UObject* UPaperTiledImporterFactory::FactoryCreateText(UClass* InClass, UObject*
 				TileSetAsset->Spacing = TileSetData.Spacing;
 				TileSetAsset->DrawingOffset = FIntPoint(TileSetData.TileOffsetX, TileSetData.TileOffsetY);
 
-				//@TODO: Ignoring ImageWidth and ImageHeight, and Properties/TileProperties aren't even loaded yet
-
 				// Import the texture
 				const FString SourceImageFilename = FPaths::Combine(*CurrentSourcePath, *TileSetData.ImagePath);
 				TileSetAsset->TileSheet = ImportTexture(SourceImageFilename, TargetTexturePath);
@@ -252,7 +264,8 @@ UObject* UPaperTiledImporterFactory::FactoryCreateText(UClass* InClass, UObject*
 				NewLayer->LayerName = FText::FromString(LayerData.Name);
 				NewLayer->bHiddenInEditor = !LayerData.bVisible;
 				NewLayer->LayerOpacity = FMath::Clamp<float>(LayerData.Opacity, 0.0f, 1.0f);
-				//@TODO: No support for Type, OffsetX, or OffsetY
+
+				//@TODO: No support for Objects (and thus Color, ObjectDrawOrder), Properties, OffsetX, or OffsetY
 
 				NewLayer->LayerWidth = LayerData.Width;
 				NewLayer->LayerHeight = LayerData.Height;
@@ -440,8 +453,68 @@ void UPaperTiledImporterFactory::ParseGlobalInfoFromJSON(TSharedPtr<FJsonObject>
 	};
 	ParseIntegerFields(OptionalIntFields, ARRAY_COUNT(OptionalIntFields), Tree, NameForErrors, /*bSilent=*/ true);
 
-	//@TODO:
-	// Parse staggeraxis and staggerindex (currently not supported by the renderer anyways)
+	// Parse StaggerAxis if present
+	const FString StaggerAxisStr = FPaperJSONHelpers::ReadString(Tree, TEXT("staggeraxis"), TEXT(""));
+	if (StaggerAxisStr == TEXT("x"))
+	{
+		OutParsedInfo.StaggerAxis = ETiledStaggerAxis::X;
+	}
+	else if (StaggerAxisStr == TEXT("y"))
+	{
+		OutParsedInfo.StaggerAxis = ETiledStaggerAxis::Y;
+	}
+	else if (!StaggerAxisStr.IsEmpty())
+	{
+		TILED_IMPORT_ERROR(TEXT("Failed to parse '%s'.  Invalid value for '%s' ('%s' but expected 'x' or 'y')"), *NameForErrors, TEXT("staggeraxis"), *StaggerAxisStr);
+		bSuccessfullyParsed = false;
+	}
+
+	// Parse StaggerIndex if present
+	const FString StaggerIndexStr = FPaperJSONHelpers::ReadString(Tree, TEXT("staggerindex"), TEXT(""));
+	if (StaggerIndexStr == TEXT("even"))
+	{
+		OutParsedInfo.StaggerIndex = ETiledStaggerIndex::Even;
+	}
+	else if (StaggerIndexStr == TEXT("odd"))
+	{
+		OutParsedInfo.StaggerIndex = ETiledStaggerIndex::Odd;
+	}
+	else if (!StaggerIndexStr.IsEmpty())
+	{
+		TILED_IMPORT_ERROR(TEXT("Failed to parse '%s'.  Invalid value for '%s' ('%s' but expected 'x' or 'y')"), *NameForErrors, TEXT("staggerindex"), *StaggerIndexStr);
+		bSuccessfullyParsed = false;
+	}
+
+	// Parse RenderOrder if present
+	const FString RenderOrderStr = FPaperJSONHelpers::ReadString(Tree, TEXT("staggerindex"), TEXT(""));
+	if (RenderOrderStr == TEXT("right-down"))
+	{
+		OutParsedInfo.RenderOrder = ETiledRenderOrder::RightDown;
+	}
+	else if (RenderOrderStr == TEXT("right-up"))
+	{
+		OutParsedInfo.RenderOrder = ETiledRenderOrder::RightUp;
+	}
+	else if (RenderOrderStr == TEXT("left-down"))
+	{
+		OutParsedInfo.RenderOrder = ETiledRenderOrder::LeftDown;
+	}
+	else if (RenderOrderStr == TEXT("left-up"))
+	{
+		OutParsedInfo.RenderOrder = ETiledRenderOrder::LeftUp;
+	}
+	else if (!RenderOrderStr.IsEmpty())
+	{
+		TILED_IMPORT_ERROR(TEXT("Failed to parse '%s'.  Invalid value for '%s' ('%s' but expected 'right-down', 'right-up', 'left-down', or 'left-up')"), *NameForErrors, TEXT("renderorder"), *RenderOrderStr);
+		bSuccessfullyParsed = false;
+	}
+
+	// Parse BackgroundColor if present
+	const FString ColorHexStr = FPaperJSONHelpers::ReadString(Tree, TEXT("backgroundcolor"), TEXT(""));
+	if (!ColorHexStr.IsEmpty())
+	{
+		OutParsedInfo.BackgroundColor = FColor::FromHex(ColorHexStr);
+	}
 
 	// Parse the orientation
 	const FString OrientationModeStr = FPaperJSONHelpers::ReadString(Tree, TEXT("orientation"), TEXT(""));
@@ -463,7 +536,7 @@ void UPaperTiledImporterFactory::ParseGlobalInfoFromJSON(TSharedPtr<FJsonObject>
 	}
 	else
 	{
-		TILED_IMPORT_ERROR(TEXT("Failed to parse '%s'.  Invalid value for '%s' (%s but expected 'orthogonal', 'isometric', 'staggered', or 'hexagonal')"), *NameForErrors, TEXT("orientation"), *OrientationModeStr);
+		TILED_IMPORT_ERROR(TEXT("Failed to parse '%s'.  Invalid value for '%s' ('%s' but expected 'orthogonal', 'isometric', 'staggered', or 'hexagonal')"), *NameForErrors, TEXT("orientation"), *OrientationModeStr);
 		bSuccessfullyParsed = false;
 		OutParsedInfo.Orientation = ETiledOrientation::Unknown;
 	}
@@ -480,6 +553,10 @@ FTileMapFromTiled::FTileMapFromTiled()
 	, TileHeight(0)
 	, Orientation(ETiledOrientation::Unknown)
 	, HexSideLength(0)
+	, StaggerAxis(ETiledStaggerAxis::Y)
+	, StaggerIndex(ETiledStaggerIndex::Odd)
+	, RenderOrder(ETiledRenderOrder::RightDown)
+	, BackgroundColor(55, 55, 55)
 {
 }
 
@@ -540,6 +617,8 @@ FTileSetFromTiled::FTileSetFromTiled()
 	: FirstGID(INDEX_NONE)
 	, ImageWidth(0)
 	, ImageHeight(0)
+	, bRemoveTransparentColor(false)
+	, ImageTransparentColor(FColor::Magenta)
 	, TileOffsetX(0)
 	, TileOffsetY(0)
 	, Margin(0)
@@ -549,7 +628,7 @@ FTileSetFromTiled::FTileSetFromTiled()
 {
 }
 
-void FTileSetFromTiled::ParseFromJSON(TSharedPtr<FJsonObject> Tree, const FString& NameForErrors, bool bSilent)
+void FTileSetFromTiled::ParseTileSetFromJSON(TSharedPtr<FJsonObject> Tree, const FString& NameForErrors, bool bSilent)
 {
 	bool bSuccessfullyParsed = true;
 
@@ -582,9 +661,6 @@ void FTileSetFromTiled::ParseFromJSON(TSharedPtr<FJsonObject> Tree, const FStrin
 		}
 	}
 
-	//@TODO: Parse the properties
-	//@TODO: Parse the tile properties
-
 	// Parse the tile set name
 	Name = FPaperJSONHelpers::ReadString(Tree, TEXT("name"), TEXT(""));
 	if (Name.IsEmpty())
@@ -597,9 +673,70 @@ void FTileSetFromTiled::ParseFromJSON(TSharedPtr<FJsonObject> Tree, const FStrin
 	ImagePath = FPaperJSONHelpers::ReadString(Tree, TEXT("image"), TEXT(""));
 	if (ImagePath.IsEmpty())
 	{
-		TILED_IMPORT_ERROR(TEXT("Failed to parse '%s'.  Invalid value for '%s' (%s but expected a path to an image)"), *NameForErrors, TEXT("image"), *ImagePath);
+		TILED_IMPORT_ERROR(TEXT("Failed to parse '%s'.  Invalid value for '%s' ('%s' but expected a path to an image)"), *NameForErrors, TEXT("image"), *ImagePath);
 		bSuccessfullyParsed = false;
 	}
+
+	// Parse the transparent color if present
+	const FString TransparentColorStr = FPaperJSONHelpers::ReadString(Tree, TEXT("transparentcolor"), TEXT(""));
+	if (!TransparentColorStr.IsEmpty())
+	{
+		bRemoveTransparentColor = true;
+		ImageTransparentColor = FColor::FromHex(TransparentColorStr);
+	}
+
+	// Parse the properties (if present)
+	const TSharedPtr<FJsonObject>* PropertiesSubobject;
+	if (Tree->TryGetObjectField(TEXT("properties"), /*out*/ PropertiesSubobject))
+	{
+		FTiledStringPair::ParsePropertyBag(/*out*/ Properties, *PropertiesSubobject, NameForErrors, bSilent);
+	}
+
+	// Parse the terrain types (if present)
+	const TArray<TSharedPtr<FJsonValue>>* TerrainTypesArray;
+	if (Tree->TryGetArrayField(TEXT("terrains"), /*out*/ TerrainTypesArray))
+	{
+		TerrainTypes.Reserve(TerrainTypesArray->Num());
+		for (TSharedPtr<FJsonValue> TerrainTypeSrc : *TerrainTypesArray)
+		{
+			FTiledTerrain NewTerrainType;
+			bSuccessfullyParsed = bSuccessfullyParsed && NewTerrainType.ParseFromJSON(TerrainTypeSrc->AsObject(), NameForErrors, bSilent);
+			TerrainTypes.Add(NewTerrainType);
+		}
+	}
+
+	// Parse the per-tile metadata if present (collision objects, terrain membership, etc...)
+	const TSharedPtr<FJsonObject>* PerTileInfoSubobject;
+	if (Tree->TryGetObjectField(TEXT("tiles"), /*out*/ PerTileInfoSubobject))
+	{
+		for (const auto& TileInfoSrcPair : (*PerTileInfoSubobject)->Values)
+		{
+			const int32 TileIndex = FCString::Atoi(*TileInfoSrcPair.Key);
+			const TSharedPtr<FJsonObject> TileInfoSrc = TileInfoSrcPair.Value->AsObject();
+
+			FTiledTileInfo& TileInfo = PerTileData.FindOrAdd(TileIndex);
+
+			TileInfo.ParseTileInfoFromJSON(TileIndex, TileInfoSrc, NameForErrors, bSilent);
+		}
+	}
+
+	// Parse the per-tile properties if present (stored separately to 'tiles' for reasons known only to the author of Tiled)
+	const TSharedPtr<FJsonObject>* PerTilePropertiesSubobject;
+	if (Tree->TryGetObjectField(TEXT("tiles"), /*out*/ PerTilePropertiesSubobject))
+	{
+		for (const auto& TilePropsSrcPair : (*PerTilePropertiesSubobject)->Values)
+		{
+			const int32 TileIndex = FCString::Atoi(*TilePropsSrcPair.Key);
+			const TSharedPtr<FJsonObject> TilePropsSrc = TilePropsSrcPair.Value->AsObject();
+
+			FTiledTileInfo& TileInfo = PerTileData.FindOrAdd(TileIndex);
+
+			FTiledStringPair::ParsePropertyBag(/*out*/ TileInfo.Properties, TilePropsSrc, NameForErrors, bSilent);
+		}
+	}
+
+	//@TODO: Should we parse ImageWidth and ImageHeight?
+	// Are these useful if the tile map gets resized to avoid invalidating everything?
 }
 
 bool FTileSetFromTiled::IsValid() const
@@ -613,14 +750,17 @@ bool FTileSetFromTiled::IsValid() const
 FTileLayerFromTiled::FTileLayerFromTiled()
 	: Width(0)
 	, Height(0)
+	, Color(FColor::White)
+	, ObjectDrawOrder(ETiledObjectLayerDrawOrder::TopDown)
 	, Opacity(1.0f)
 	, bVisible(true)
+	, LayerType(ETiledLayerType::TileLayer)
 	, OffsetX(0)
 	, OffsetY(0)
 {
 }
 
-void FTileLayerFromTiled::ParseFromJSON(TSharedPtr<FJsonObject> Tree, const FString& NameForErrors, bool bSilent)
+bool FTileLayerFromTiled::ParseFromJSON(TSharedPtr<FJsonObject> Tree, const FString& NameForErrors, bool bSilent)
 {
 	bool bSuccessfullyParsed = true;
 
@@ -650,29 +790,102 @@ void FTileLayerFromTiled::ParseFromJSON(TSharedPtr<FJsonObject> Tree, const FStr
 		bSuccessfullyParsed = false;
 	}
 
-	if (!Tree->TryGetStringField(TEXT("type"), /*out*/ Type))
+	// Parse the layer type
+	const FString LayerTypeStr = FPaperJSONHelpers::ReadString(Tree, TEXT("type"), TEXT(""));
+	if (LayerTypeStr == TEXT("tilelayer"))
+	{
+		LayerType = ETiledLayerType::TileLayer;
+	}
+	else if (LayerTypeStr == TEXT("objectgroup"))
+	{
+		LayerType = ETiledLayerType::ObjectGroup;
+	}
+	else if (LayerTypeStr == TEXT("imagelayer"))
+	{
+		LayerType = ETiledLayerType::ImageLayer;
+	}
+	else
+	{
+		TILED_IMPORT_ERROR(TEXT("Failed to parse '%s'.  Invalid value for '%s' ('%s' but expected 'tilelayer' or 'objectgroup')"), *NameForErrors, TEXT("type"), *LayerTypeStr);
+		bSuccessfullyParsed = false;
+	}
+
+	// Parse the object draw order (if present)
+	const FString ObjectDrawOrderStr = FPaperJSONHelpers::ReadString(Tree, TEXT("draworder"), TEXT(""));
+	if (ObjectDrawOrderStr == TEXT("index"))
 	{
-		//@TODO: Figure out what to do with the layer type!
-		TILED_IMPORT_ERROR(TEXT("Failed to parse '%s'.  Invalid value for '%s' (%s but expected a ??? )"), *NameForErrors, TEXT("type"), *Type);
+		ObjectDrawOrder = ETiledObjectLayerDrawOrder::Index;
+	}
+	else if (ObjectDrawOrderStr == TEXT("topdown"))
+	{
+		ObjectDrawOrder = ETiledObjectLayerDrawOrder::TopDown;
+	}
+	else if (!ObjectDrawOrderStr.IsEmpty())
+	{
+		TILED_IMPORT_ERROR(TEXT("Failed to parse '%s'.  Invalid value for '%s' ('%s' but expected 'index' or 'topdown')"), *NameForErrors, TEXT("draworder"), *ObjectDrawOrderStr);
 		bSuccessfullyParsed = false;
 	}
 
-	const TArray<TSharedPtr<FJsonValue>>* DataArray;
-	if (Tree->TryGetArrayField(TEXT("data"), /*out*/ DataArray))
+	// Parse the property bag if present
+	const TSharedPtr<FJsonObject>* PropertiesSubobject;
+	if (Tree->TryGetObjectField(TEXT("properties"), /*out*/ PropertiesSubobject))
+	{
+		FTiledStringPair::ParsePropertyBag(/*out*/ Properties, *PropertiesSubobject, NameForErrors, bSilent);
+	}
+
+	// Parse the data specific to this layer type
+	if (LayerType == ETiledLayerType::TileLayer)
 	{
-		TileIndices.Reserve(DataArray->Num());
-		for (TSharedPtr<FJsonValue> TileEntry : *DataArray)
+		const TArray<TSharedPtr<FJsonValue>>* DataArray;
+		if (Tree->TryGetArrayField(TEXT("data"), /*out*/ DataArray))
 		{
- 			const double TileIndexAsDouble = TileEntry->AsNumber();
-			uint32 TileID = (uint32)TileIndexAsDouble;
-			TileIndices.Add(TileID);
+			TileIndices.Reserve(DataArray->Num());
+			for (TSharedPtr<FJsonValue> TileEntry : *DataArray)
+			{
+				const double TileIndexAsDouble = TileEntry->AsNumber();
+				uint32 TileID = (uint32)TileIndexAsDouble;
+				TileIndices.Add(TileID);
+			}
+		}
+		else
+		{
+			TILED_IMPORT_ERROR(TEXT("Failed to parse '%s'.  Missing tile data for layer '%s'"), *NameForErrors, *Name);
+			bSuccessfullyParsed = false;
+		}
+	}
+	else if (LayerType == ETiledLayerType::ObjectGroup)
+	{
+		const TArray<TSharedPtr<FJsonValue>>* ObjectArray;
+		if (Tree->TryGetArrayField(TEXT("objects"), /*out*/ ObjectArray))
+		{
+			Objects.Reserve(ObjectArray->Num());
+			for (TSharedPtr<FJsonValue> ObjectEntry : *ObjectArray)
+			{
+				FTiledObject NewObject;
+				bSuccessfullyParsed = bSuccessfullyParsed && NewObject.ParseFromJSON(ObjectEntry->AsObject(), NameForErrors, bSilent);
+				Objects.Add(NewObject);
+			}
+		}
+		else
+		{
+			TILED_IMPORT_ERROR(TEXT("Failed to parse '%s'.  Missing object data for layer '%s'"), *NameForErrors, *Name);
+			bSuccessfullyParsed = false;
 		}
 	}
+	else if (LayerType == ETiledLayerType::ImageLayer)
+	{
+		OverlayImagePath = FPaperJSONHelpers::ReadString(Tree, TEXT("image"), TEXT(""));
+	}
 	else
 	{
-		TILED_IMPORT_ERROR(TEXT("Failed to parse '%s'.  Missing layer data for layer '%s'"), *NameForErrors, *Name);
-		bSuccessfullyParsed = false;
+		if (bSuccessfullyParsed)
+		{
+			TILED_IMPORT_ERROR(TEXT("Failed to parse '%s'.  Unknown layer type for layer '%s'"), *NameForErrors, *Name);
+			bSuccessfullyParsed = false;
+		}
 	}
+
+	return bSuccessfullyParsed;
 }
 
 bool FTileLayerFromTiled::IsValid() const
@@ -680,6 +893,231 @@ bool FTileLayerFromTiled::IsValid() const
 	return (Width > 0) && (Height > 0) && (TileIndices.Num() == (Width*Height));
 }
 
+//////////////////////////////////////////////////////////////////////////
+// FTiledTerrain
+
+FTiledTerrain::FTiledTerrain()
+	: SolidTileLocalIndex(0)
+{
+}
+
+bool FTiledTerrain::ParseFromJSON(TSharedPtr<FJsonObject> Tree, const FString& NameForErrors, bool bSilent)
+{
+	if (!Tree->TryGetStringField(TEXT("name"), /*out*/ TerrainName))
+	{
+		TILED_IMPORT_ERROR(TEXT("Failed to parse '%s'.  Terrain entry is missing the 'name' field"), *NameForErrors);
+		return false;
+	}
+
+	if (!Tree->TryGetNumberField(TEXT("tile"), /*out*/ SolidTileLocalIndex))
+	{
+		TILED_IMPORT_ERROR(TEXT("Failed to parse '%s'.  Terrain entry is missing the 'tile' field"), *NameForErrors);
+		return false;
+	}
+
+	return true;
+}
+
+//////////////////////////////////////////////////////////////////////////
+// FTiledTileInfo
+
+FTiledTileInfo::FTiledTileInfo()
+	: Probability(1.0f)
+{
+	for (int32 Index = 0; Index < 4; ++Index)
+	{
+		TerrainIndices[Index] = INDEX_NONE;
+	}
+}
+
+bool FTiledTileInfo::ParseTileInfoFromJSON(int32 TileIndex, TSharedPtr<FJsonObject> Tree, const FString& NameForErrors, bool bSilent)
+{
+	bool bSuccessfullyParsed = true;
+
+	// Try reading the terrain membership array if present
+	const TArray<TSharedPtr<FJsonValue>>* TerrainMembershipArray;
+	if (Tree->TryGetArrayField(TEXT("terrain"), /*out*/ TerrainMembershipArray))
+	{
+		if (TerrainMembershipArray->Num() == 4)
+		{
+			for (int32 Index = 0; Index < 4; ++Index)
+			{
+				TSharedPtr<FJsonValue> MembershipIndex = (*TerrainMembershipArray)[Index];
+				
+				if (!MembershipIndex->TryGetNumber(TerrainIndices[Index]))
+				{
+					TILED_IMPORT_ERROR(TEXT("Failed to parse '%s'.  The 'terrain' array for tile %d should contain 4 indices into the terrain array"), *NameForErrors, TileIndex);
+					bSuccessfullyParsed = false;
+				}
+			}
+		}
+		else
+		{
+			TILED_IMPORT_ERROR(TEXT("Failed to parse '%s'.  The 'terrain' array for tile %d should contain 4 entries but it contained %d entries"), *NameForErrors, TileIndex, TerrainMembershipArray->Num());
+			bSuccessfullyParsed = false;
+		}
+	}
+
+	// Try reading the probability if present
+	double DoubleProbability;
+	if (Tree->TryGetNumberField(TEXT("probability"), /*out*/ DoubleProbability))
+	{
+		Probability = FMath::Clamp<float>((float)DoubleProbability, 0.0f, 1.0f);
+	}
+
+	// Try reading the per-tile collision data if present
+	// Note: This is really an entire fake objectgroup layer, but only the objects array matters; Tiled doens't even provide a way to edit the rest of the data.
+	const TSharedPtr<FJsonObject>* ObjectGroupSubobject;
+	if (Tree->TryGetObjectField(TEXT("objectgroup"), /*out*/ ObjectGroupSubobject))
+	{
+		const TArray<TSharedPtr<FJsonValue>>* ObjectArray;
+		if ((*ObjectGroupSubobject)->TryGetArrayField(TEXT("objects"), /*out*/ ObjectArray))
+		{
+			Objects.Reserve(ObjectArray->Num());
+			for (TSharedPtr<FJsonValue> ObjectEntry : *ObjectArray)
+			{
+				FTiledObject NewObject;
+				bSuccessfullyParsed = bSuccessfullyParsed && NewObject.ParseFromJSON(ObjectEntry->AsObject(), NameForErrors, bSilent);
+				Objects.Add(NewObject);
+			}
+		}
+		else
+		{
+			TILED_IMPORT_ERROR(TEXT("Failed to parse '%s'.  Expected an 'objects' entry inside 'objectgroup' for tile %d"), *NameForErrors, TileIndex);
+			bSuccessfullyParsed = false;
+		}
+	}
+
+	return bSuccessfullyParsed;
+}
+
+//////////////////////////////////////////////////////////////////////////
+// FTiledStringPair
+
+void FTiledStringPair::ParsePropertyBag(TArray<FTiledStringPair>& OutProperties, TSharedPtr<FJsonObject> Tree, const FString& NameForErrors, bool bSilent)
+{
+	OutProperties.Reserve(Tree->Values.Num());
+	for (auto KV : Tree->Values)
+	{
+		OutProperties.Add(FTiledStringPair(KV.Key, KV.Value->AsString()));
+	}
+}
+
+//////////////////////////////////////////////////////////////////////////
+// FTiledObject
+
+FTiledObject::FTiledObject()
+	: TiledObjectType(ETiledObjectType::Box)
+	, ID(0)
+	, bVisible(true)
+	, X(0.0)
+	, Y(0.0)
+	, Width(0.0)
+	, Height(0.0)
+	, RotationDegrees(0.0)
+	, TileGID(0)
+{
+}
+
+bool FTiledObject::ParseFromJSON(TSharedPtr<FJsonObject> Tree, const FString& NameForErrors, bool bSilent)
+{
+	bool bSuccessfullyParsed = true;
+
+	// Parse all of the integer fields
+	FRequiredDoubleField FloatFields[] = {
+		FRequiredDoubleField(Width, TEXT("width"), 0.0),
+		FRequiredDoubleField(Height, TEXT("height"), 0.0),
+		FRequiredDoubleField(X, TEXT("x"), -MAX_FLT),
+		FRequiredDoubleField(Y, TEXT("y"), -MAX_FLT),
+		FRequiredDoubleField(RotationDegrees, TEXT("rotation"), -MAX_FLT)
+	};
+
+	bSuccessfullyParsed = bSuccessfullyParsed && ParseScalarFields(FloatFields, ARRAY_COUNT(FloatFields), Tree, NameForErrors, bSilent);
+
+	if (!Tree->TryGetBoolField(TEXT("visible"), /*out*/ bVisible))
+	{
+		bVisible = true;
+	}
+
+	if (!Tree->TryGetStringField(TEXT("name"), /*out*/ Name))
+	{
+		TILED_IMPORT_ERROR(TEXT("Failed to parse '%s'.  Expected an object name"), *NameForErrors);
+		bSuccessfullyParsed = false;
+	}
+
+	if (!Tree->TryGetStringField(TEXT("type"), /*out*/ UserType))
+	{
+		TILED_IMPORT_ERROR(TEXT("Failed to parse '%s'.  Expected an object type"), *NameForErrors);
+		bSuccessfullyParsed = false;
+	}
+
+	if (!Tree->TryGetNumberField(TEXT("id"), /*out*/ ID))
+	{
+		TILED_IMPORT_ERROR(TEXT("Failed to parse '%s'.  Expected an object ID"), *NameForErrors);
+		bSuccessfullyParsed = false;
+	}
+
+	const TSharedPtr<FJsonObject>* PropertiesSubobject;
+	if (Tree->TryGetObjectField(TEXT("properties"), /*out*/ PropertiesSubobject))
+	{
+		FTiledStringPair::ParsePropertyBag(/*out*/ Properties, *PropertiesSubobject, NameForErrors, bSilent);
+	}
+
+	// Determine the object type
+	if (Tree->TryGetNumberField(TEXT("gid"), /*out*/ TileGID))
+	{
+		TiledObjectType = ETiledObjectType::PlacedTile;
+	}
+	else if (Tree->HasField(TEXT("ellipse")))
+	{
+		TiledObjectType = ETiledObjectType::Ellipse;
+	}
+	else
+	{
+		const TArray<TSharedPtr<FJsonValue>>* PointsArray;
+		if (Tree->TryGetArrayField(TEXT("polygon"), /*out*/ PointsArray))
+		{
+			TiledObjectType = ETiledObjectType::Polygon;
+			bSuccessfullyParsed = bSuccessfullyParsed && ParsePointArray(/*out*/ Points, *PointsArray, NameForErrors, bSilent);
+		}
+		else if (Tree->TryGetArrayField(TEXT("polyline"), /*out*/ PointsArray))
+		{
+			TiledObjectType = ETiledObjectType::Polyline;
+			bSuccessfullyParsed = bSuccessfullyParsed && ParsePointArray(/*out*/ Points, *PointsArray, NameForErrors, bSilent);
+		}
+		else
+		{
+			TiledObjectType = ETiledObjectType::Box;
+		}
+	}
+
+
+	return bSuccessfullyParsed;
+}
+
+bool FTiledObject::ParsePointArray(TArray<FVector2D>& OutPoints, const TArray<TSharedPtr<FJsonValue>>& InArray, const FString& NameForErrors, bool bSilent)
+{
+	bool bSuccessfullyParsed = true;
+
+	OutPoints.Reserve(InArray.Num());
+	for (TSharedPtr<FJsonValue> ArrayElement : InArray)
+	{
+		double X = 0.0;
+		double Y = 0.0;
+
+		FRequiredDoubleField FloatFields[] = {
+			FRequiredDoubleField(X, TEXT("x"), -MAX_FLT),
+			FRequiredDoubleField(Y, TEXT("y"), -MAX_FLT)
+		};
+
+		bSuccessfullyParsed = bSuccessfullyParsed && ParseScalarFields(FloatFields, ARRAY_COUNT(FloatFields), ArrayElement->AsObject(), NameForErrors, bSilent);
+
+		OutPoints.Add(FVector2D(X, Y));
+	}
+
+	return bSuccessfullyParsed;
+}
+
 //////////////////////////////////////////////////////////////////////////
 
 #undef TILED_IMPORT_ERROR
