Fix for tooltips on some inline UFUNCTIONs, spotted while testing.

diff --git a/Engine/Plugins/2D/Paper2D/Source/Paper2D/Classes/PaperFlipbook.h b/Engine/Plugins/2D/Paper2D/Source/Paper2D/Classes/PaperFlipbook.h
--- a/Engine/Plugins/2D/Paper2D/Source/Paper2D/Classes/PaperFlipbook.h
+++ b/Engine/Plugins/2D/Paper2D/Source/Paper2D/Classes/PaperFlipbook.h
@@ -100,20 +100,16 @@ public:
 	// Returns the number of key frames
 	UFUNCTION(BlueprintCallable, Category="Sprite")
 	int32 GetNumKeyFrames() const
-#if CPP
 	{
 		return KeyFrames.Num();
 	}
-#endif
 
 	// Is the specified Index within the valid range of key frames?
 	UFUNCTION(BlueprintCallable, Category="Sprite")
 	bool IsValidKeyFrameIndex(int32 Index) const
-#if CPP
 	{
 		return KeyFrames.IsValidIndex(Index);
 	}
-#endif
 
 	// Returns the key frame at the specified index, make sure the index is valid before use
 	const FPaperFlipbookKeyFrame& GetKeyFrameChecked(int32 Index) const
