Paper2D: Changed default verbosity for LogPaper2DEditor from Verbose to Log

diff --git a/Engine/Plugins/2D/Paper2D/Source/Paper2DEditor/Private/Paper2DEditorPrivatePCH.h b/Engine/Plugins/2D/Paper2D/Source/Paper2DEditor/Private/Paper2DEditorPrivatePCH.h
--- a/Engine/Plugins/2D/Paper2D/Source/Paper2DEditor/Private/Paper2DEditorPrivatePCH.h
+++ b/Engine/Plugins/2D/Paper2D/Source/Paper2DEditor/Private/Paper2DEditorPrivatePCH.h
@@ -7,5 +7,5 @@
 #include "Paper2DModule.h"
 #include "Paper2DEditorClasses.h"
 
-DECLARE_LOG_CATEGORY_EXTERN(LogPaper2DEditor, Verbose, All);
+DECLARE_LOG_CATEGORY_EXTERN(LogPaper2DEditor, Log, All);
 
