Fixed a couple of typos in comments

diff --git a/Engine/Plugins/2D/Paper2D/Source/Paper2DEditor/Private/PaperEditorShared/SocketEditing.cpp b/Engine/Plugins/2D/Paper2D/Source/Paper2DEditor/Private/PaperEditorShared/SocketEditing.cpp
--- a/Engine/Plugins/2D/Paper2D/Source/Paper2DEditor/Private/PaperEditorShared/SocketEditing.cpp
+++ b/Engine/Plugins/2D/Paper2D/Source/Paper2DEditor/Private/PaperEditorShared/SocketEditing.cpp
@@ -9,7 +9,7 @@
 IMPLEMENT_HIT_PROXY(HSpriteSelectableObjectHitProxy, HHitProxy);
 
 //////////////////////////////////////////////////////////////////////////
-// FSocketEditingHelper
+// FSpriteSelectedSocket
 
 const FName FSpriteSelectedSocket::SocketTypeID(TEXT("Socket"));
 
diff --git a/Engine/Source/Runtime/Engine/Classes/PhysicsEngine/AggregateGeometry2D.h b/Engine/Source/Runtime/Engine/Classes/PhysicsEngine/AggregateGeometry2D.h
--- a/Engine/Source/Runtime/Engine/Classes/PhysicsEngine/AggregateGeometry2D.h
+++ b/Engine/Source/Runtime/Engine/Classes/PhysicsEngine/AggregateGeometry2D.h
@@ -82,7 +82,7 @@ struct ENGINE_API FBoxElement2D
 	UPROPERTY(Category=KBoxElem, VisibleAnywhere)
 	float Height;
 
-	// Rotation of the box (in degreens)
+	// Rotation of the box (in degrees)
 	UPROPERTY(Category = KBoxElem, VisibleAnywhere)
 	float Angle;
 
