Paper2D: Moved PaperGeomTools.h to be a public header, allowing use in the editor module

diff --git a/Engine/Plugins/2D/Paper2D/Source/Paper2D/Private/PaperGeomTools.h b/Engine/Plugins/2D/Paper2D/Source/Paper2D/Private/PaperGeomTools.h
deleted file mode 100644
diff --git a/Engine/Plugins/2D/Paper2D/Source/Paper2D/Public/PaperGeomTools.h b/Engine/Plugins/2D/Paper2D/Source/Paper2D/Public/PaperGeomTools.h
new file mode 100644
--- /dev/null
+++ b/Engine/Plugins/2D/Paper2D/Source/Paper2D/Public/PaperGeomTools.h
@@ -0,0 +1,37 @@
+// Copyright 1998-2015 Epic Games, Inc. All Rights Reserved.
+
+#pragma once
+
+class PAPER2D_API PaperGeomTools
+{
+public:
+	// Corrects the polygon winding to match bNegativeWinding
+	// Ie. If the vertex order doesn't match the winding, it is reversed
+	// Returns empty array of points if num points < 3
+	static void CorrectPolygonWinding(TArray<FVector2D>& OutVertices, const TArray<FVector2D>& Vertices, const bool bNegativeWinding);
+
+	// Returns true if the points forming a polygon have CCW winding
+	// Returns true if the polygon isn't valid
+	static bool IsPolygonWindingCCW(const TArray<FVector2D>& Points);
+	static bool IsPolygonWindingCCW(const TArray<FIntPoint>& Points);
+
+	// Checks that these polygons can be successfully triangulated	
+	static bool ArePolygonsValid(const TArray<TArray<FVector2D>>& Polygons);
+
+	// Merge additive and subtractive polygons, split them up into additive polygons
+	// Assumes all polygons and overlapping polygons are valid, and the windings match the setting on the polygon
+	static TArray<TArray<FVector2D>> ReducePolygons(const TArray<TArray<FVector2D>>& Polygons, const TArray<bool>& PolygonNegativeWinding);
+
+	// Triangulate a polygon. Check notes in implementation
+	static bool TriangulatePoly(TArray<FVector2D>& OutTris, const TArray<FVector2D>& PolygonVertices, bool bKeepColinearVertices = false);
+
+	// 2D version of RemoveRedundantTriangles from GeomTools.cpp
+	static void RemoveRedundantTriangles(TArray<FVector2D>& OutTriangles, const TArray<FVector2D>& InTriangles);
+
+	// Generate convex shapes
+	static void GenerateConvexPolygonsFromTriangles(TArray<TArray<FVector2D>>& OutPolygons, const TArray<FVector2D>& InTriangles);
+
+	// Generate convex hull from points
+	static void GenerateConvexHullFromPoints(TArray<FVector2D>& OutConvexHull, TArray<FVector2D>& Points);
+};
+
