Sprite Editor: Removed some dead code and obsolete TODOs

diff --git a/Engine/Plugins/2D/Paper2D/Source/Paper2D/Private/PaperSprite.cpp b/Engine/Plugins/2D/Paper2D/Source/Paper2D/Private/PaperSprite.cpp
--- a/Engine/Plugins/2D/Paper2D/Source/Paper2D/Private/PaperSprite.cpp
+++ b/Engine/Plugins/2D/Paper2D/Source/Paper2D/Private/PaperSprite.cpp
@@ -2142,7 +2142,6 @@ void FSpriteGeometryCollection::ConditionGeometry()
 					Shape.SetNewPivot(A + 0.5f * (AB + BC));
 					Shape.BoxSize = FVector2D(AB.Size(), DA.Size());
 					Shape.ShapeType = ESpriteShapeType::Box;
-					//@TODO: Still need to deal with this here!
 				}
 			}
 		}
diff --git a/Engine/Plugins/2D/Paper2D/Source/Paper2DEditor/Private/SpriteEditor/SpriteEditorViewportClient.cpp b/Engine/Plugins/2D/Paper2D/Source/Paper2DEditor/Private/SpriteEditor/SpriteEditorViewportClient.cpp
--- a/Engine/Plugins/2D/Paper2D/Source/Paper2DEditor/Private/SpriteEditor/SpriteEditorViewportClient.cpp
+++ b/Engine/Plugins/2D/Paper2D/Source/Paper2DEditor/Private/SpriteEditor/SpriteEditorViewportClient.cpp
@@ -370,59 +370,6 @@ void FSpriteEditorViewportClient::DrawGeometry_CanvasPass(FViewport& InViewport,
 
 		const FLinearColor LineColor = Shape.IsShapeValid() ? LineColorRaw : FMath::Lerp(LineColorRaw, FLinearColor::Red, 0.8f);
 
-		// This is a work-in-progress, it currently blocks widget selection even with a lower hit proxy priority, so it is disabled
-#define UE_DRAW_SPRITE_SHAPE_BACKGROUNDS 0
-#if UE_DRAW_SPRITE_SHAPE_BACKGROUNDS
-		// Draw the interior (allowing selection of the whole shape)
-		if (bIsHitTesting)
-		{
-			TSharedPtr<FSpriteSelectedShape> Data = MakeShareable(new FSpriteSelectedShape(*this, Geometry, ShapeIndex, /*bIsBackground=*/ true));
-			Canvas.SetHitProxy(new HSpriteSelectableObjectHitProxy(Data));
-		}
-
-		FLinearColor BackgroundColor = bIsShapeSelected ? SpriteEditingConstants::GeometrySelectedColor : LineColor;
-		BackgroundColor.A *= 0.5f;
-
-		if (Shape.ShapeType == ESpriteShapeType::Circle)
-		{
-			const FVector2D PixelSpaceRadius = Shape.BoxSize * 0.5f;
-
-			//@TODO: This is the wrong size when in Perspective mode
-			const FVector2D ScreenSpaceRadius = TextureSpaceToScreenSpace(View, PixelSpaceRadius) - TextureSpaceToScreenSpace(View, FVector2D::ZeroVector);
-
-			const FVector2D ScreenPos = TextureSpaceToScreenSpace(View, Shape.BoxPosition);
-
-			FCanvasNGonItem CircleBackground(ScreenPos, ScreenSpaceRadius, SpriteEditingConstants::CircleShapeNumSides, BackgroundColor);
-			//@TODO: CircleBackground.BlendMode = SE_BLEND_Translucent;
-			CircleBackground.Draw(&Canvas);
-		}
-		else if (Shape.ShapeType == ESpriteShapeType::Box)
-		{
-			const FVector2D HalfSize = Shape.BoxSize * 0.5f;
-			const FVector2D TopLeft = TextureSpaceToScreenSpace(View, Shape.BoxPosition - HalfSize);
-			const FVector2D BottomRight = TextureSpaceToScreenSpace(View, Shape.BoxPosition + HalfSize);
-			const FVector2D TopRight = TextureSpaceToScreenSpace(View, FVector2D(Shape.BoxPosition.X + HalfSize.X, Shape.BoxPosition.Y - HalfSize.Y));
-			const FVector2D BottomLeft = TextureSpaceToScreenSpace(View, FVector2D(Shape.BoxPosition.X - HalfSize.X, Shape.BoxPosition.Y + HalfSize.Y));
-
-			FCanvasTriangleItem Triangle(TopLeft, TopRight, BottomRight, GWhiteTexture);
-			//@TODO: Triangle.BlendMode = SE_BLEND_Translucent;
-			Triangle.SetColor(BackgroundColor);
-			Triangle.Draw(&Canvas);
-
-			Triangle.SetPoints(TopLeft, BottomRight, BottomLeft);
-			Triangle.Draw(&Canvas);
-		}
-		else
-		{
-			//@TODO: Draw the background for polygons
-		}
-
-		if (bIsHitTesting)
-		{
-			Canvas.SetHitProxy(nullptr);
-		}
-#endif
-
 		// Draw the circle shape if necessary
 		if (Shape.ShapeType == ESpriteShapeType::Circle)
 		{
@@ -432,7 +379,7 @@ void FSpriteEditorViewportClient::DrawGeometry_CanvasPass(FViewport& InViewport,
 				Canvas.SetHitProxy(new HSpriteSelectableObjectHitProxy(Data));
 			}
 
-			//@TODO: Draw the circle
+			// Draw the circle
 			const float RadiusX = Shape.BoxSize.X * 0.5f;
 			const float RadiusY = Shape.BoxSize.Y * 0.5f;
 
@@ -561,7 +508,6 @@ void FSpriteEditorViewportClient::DrawGeometryStats(FViewport& InViewport, FScen
 	TextItem.Draw(&Canvas);
 	TextItem.Position.Y += 18.0f;
 
-	//@TODO: Should show the agg numbers instead for collision, and the render tris numbers for rendering
 	// Draw the number of vertices
 	int32 NumVerts = 0;
 	for (int32 PolyIndex = 0; PolyIndex < Geometry.Shapes.Num(); ++PolyIndex)
