Sprite Editor: Added note to source region edit mode and documentation about double-clicking to create new sprites and selecting other sprites
[UE-12891]

diff --git a/Engine/Documentation/Source/Engine/Paper2D/Sprites/RenderGeometry/SpriteRenderGeometry.INT.udn b/Engine/Documentation/Source/Engine/Paper2D/Sprites/RenderGeometry/SpriteRenderGeometry.INT.udn
--- a/Engine/Documentation/Source/Engine/Paper2D/Sprites/RenderGeometry/SpriteRenderGeometry.INT.udn
+++ b/Engine/Documentation/Source/Engine/Paper2D/Sprites/RenderGeometry/SpriteRenderGeometry.INT.udn
@@ -27,6 +27,12 @@ The **Edit Source Region** mode displays the current bounding rectangle of the S
 at each corner and at the center of each side. Each handle can be selected and dragged using the transform 
 gizmo.
 
+New sprites can be created in this mode by Ctrl+Dragging around a region of the source texture, or by double-clicking
+on a region of the texture (which will find all connected pixels and create a new sprite from the bounds).
+
+Other sprites that share the same source texture will be drawn with a subdued border, and clicking on one of them will
+change the current sprite being edited.
+
 ![](edit_source_region.png)
 
 ## Modifying Render Geometry
diff --git a/Engine/Plugins/2D/Paper2D/Source/Paper2DEditor/Private/SpriteEditor/SpriteEditorViewportClient.cpp b/Engine/Plugins/2D/Paper2D/Source/Paper2DEditor/Private/SpriteEditor/SpriteEditorViewportClient.cpp
--- a/Engine/Plugins/2D/Paper2D/Source/Paper2DEditor/Private/SpriteEditor/SpriteEditorViewportClient.cpp
+++ b/Engine/Plugins/2D/Paper2D/Source/Paper2DEditor/Private/SpriteEditor/SpriteEditorViewportClient.cpp
@@ -739,7 +739,7 @@ void FSpriteEditorViewportClient::DrawCanvas(FViewport& Viewport, FSceneView& Vi
 	static const FText GeomAddPolygonHelpStr = LOCTEXT("GeomClickAddPolygon", "Click to start creating a polygon\nCtrl + Click to start creating a subtractive polygon\n");
 	static const FText GeomAddCollisionPolygonHelpStr = LOCTEXT("GeomClickAddCollisionPolygon", "Click to start creating a polygon");
 	static const FText GeomAddVerticesHelpStr = LOCTEXT("GeomClickAddVertices", "Click to add points to the polygon\nEnter to finish adding a polygon\nEscape to cancel this polygon\n");
-	static const FText SourceRegionHelpStr = LOCTEXT("SourceRegionHelp", "Drag handles to adjust source region\nDouble-click on an image region to select all connected pixels\nHold down Ctrl and drag a rectangle to create a new sprite at that position");
+	static const FText SourceRegionHelpStr = LOCTEXT("SourceRegionHelp", "Drag handles to adjust source region\nDouble-click on an image region to select all connected pixels\nHold down Ctrl and drag a rectangle to create a new sprite at that position\nDouble-click to create a new sprite under the cursor\nClick on other sprite rectangles to change the active sprite");
 
 	switch (CurrentMode)
 	{
