Sprite Editor: Fixed a bug where dragging the edge of a rectangle shape would cause it to slide away

diff --git a/Engine/Plugins/2D/Paper2D/Source/Paper2D/Classes/SpriteEditorOnlyTypes.h b/Engine/Plugins/2D/Paper2D/Source/Paper2D/Classes/SpriteEditorOnlyTypes.h
--- a/Engine/Plugins/2D/Paper2D/Source/Paper2D/Classes/SpriteEditorOnlyTypes.h
+++ b/Engine/Plugins/2D/Paper2D/Source/Paper2D/Classes/SpriteEditorOnlyTypes.h
@@ -76,6 +76,7 @@ public:
 		return ShapeSpacePoint.GetRotated(-Rotation) + BoxPosition;
 	}
 
+	// Returns the polygon centroid in texture space
 	FVector2D GetPolygonCentroid() const
 	{
 		FVector2D LocalSpaceResult = FVector2D::ZeroVector;
@@ -105,7 +106,7 @@ public:
 	// Sets a new pivot and adjusts all vertices to be relative to it
 	void SetNewPivot(const FVector2D& NewPosInTextureSpace)
 	{
-		const FVector2D DeltaTexturePos = NewPosInTextureSpace - BoxPosition;
+		const FVector2D DeltaTexturePos = BoxPosition - NewPosInTextureSpace;
 		for (FVector2D& Vertex : Vertices)
 		{
 			const FVector2D NewVertexPosTS = ConvertShapeSpaceToTextureSpace(Vertex) + DeltaTexturePos;
diff --git a/Engine/Plugins/2D/Paper2D/Source/Paper2D/Private/PaperSprite.cpp b/Engine/Plugins/2D/Paper2D/Source/Paper2D/Private/PaperSprite.cpp
--- a/Engine/Plugins/2D/Paper2D/Source/Paper2D/Private/PaperSprite.cpp
+++ b/Engine/Plugins/2D/Paper2D/Source/Paper2D/Private/PaperSprite.cpp
@@ -1179,9 +1179,6 @@ void UPaperSprite::BuildGeometryFromContours(FSpriteGeometryCollection& GeomOwne
 			// Recenter them
 			const FVector2D AverageCenterFloat = NewShape.GetPolygonCentroid();
 			const FVector2D AverageCenterSnapped(FMath::RoundToInt(AverageCenterFloat.X), FMath::RoundToInt(AverageCenterFloat.Y));
-
-			//@TODO: This looks a bit odd but it's due to the fact that we had the pivot in the top left and are changing it to the center
-			NewShape.BoxPosition = AverageCenterSnapped * 2.0f;
 			NewShape.SetNewPivot(AverageCenterSnapped);
 
 			// Get intended winding
@@ -2139,7 +2136,8 @@ void FSpriteGeometryCollection::ConditionGeometry()
 				const bool bMeetsRotationConstraint = FMath::IsNearlyEqual(AB.Y, 0.0f);
 				if (bMeetsRotationConstraint)
 				{
-					Shape.SetNewPivot(A + 0.5f * (AB + BC));
+					const FVector2D NewPivotTextureSpace = Shape.GetPolygonCentroid();
+					Shape.SetNewPivot(NewPivotTextureSpace);
 					Shape.BoxSize = FVector2D(AB.Size(), DA.Size());
 					Shape.ShapeType = ESpriteShapeType::Box;
 				}
