Sprite Editor: Fixed an issue where the transform widget wasn't immediately usable when selecting an entire shape by double-clicking

diff --git a/Engine/Plugins/2D/Paper2D/Source/Paper2DEditor/Private/SpriteEditor/SpriteEditorViewportClient.cpp b/Engine/Plugins/2D/Paper2D/Source/Paper2DEditor/Private/SpriteEditor/SpriteEditorViewportClient.cpp
--- a/Engine/Plugins/2D/Paper2D/Source/Paper2DEditor/Private/SpriteEditor/SpriteEditorViewportClient.cpp
+++ b/Engine/Plugins/2D/Paper2D/Source/Paper2DEditor/Private/SpriteEditor/SpriteEditorViewportClient.cpp
@@ -1081,7 +1081,6 @@ void FSpriteEditorViewportClient::ProcessClick(FSceneView& View, HHitProxy* HitP
 			{
 				if (SelectedVertex->IsValidInEditor(GetSpriteBeingEdited(), IsInRenderingEditMode()))
 				{
-					//findme
 					SelectShape(SelectedVertex->ShapeIndex);
 				}
 			}
@@ -1851,6 +1850,8 @@ FSpriteGeometryCollection* FSpriteEditorViewportClient::GetGeometryBeingEdited()
 
 void FSpriteEditorViewportClient::SelectShape(const int32 ShapeIndex)
 {
+	Viewport->InvalidateHitProxy();
+
 	FSpriteGeometryCollection* Geometry = GetGeometryBeingEdited();
 	if (ensure((Geometry != nullptr) && Geometry->Shapes.IsValidIndex(ShapeIndex)))
 	{
