Sprite Editor: Allow editing an edge of the source region by clicking anywhere on the edge, rather than only at the midpoint

diff --git a/Engine/Plugins/2D/Paper2D/Source/Paper2DEditor/Private/SpriteEditor/SpriteEditorViewportClient.cpp b/Engine/Plugins/2D/Paper2D/Source/Paper2DEditor/Private/SpriteEditor/SpriteEditorViewportClient.cpp
--- a/Engine/Plugins/2D/Paper2D/Source/Paper2DEditor/Private/SpriteEditor/SpriteEditorViewportClient.cpp
+++ b/Engine/Plugins/2D/Paper2D/Source/Paper2DEditor/Private/SpriteEditor/SpriteEditorViewportClient.cpp
@@ -290,36 +290,37 @@ void FSpriteEditorViewportClient::DrawSourceRegion(FViewport& InViewport, FScene
 	{
 		const int32 NextVertexIndex = (VertexIndex + 1) % 4;
 
+		// Draw the edge
+		if (bIsHitTesting)
+		{
+			TSharedPtr<FSpriteSelectedSourceRegion> Data = MakeShareable(new FSpriteSelectedSourceRegion());
+			Data->SpritePtr = Sprite;
+			Data->VertexIndex = 4 + VertexIndex;
+			Canvas.SetHitProxy(new HSpriteSelectableObjectHitProxy(Data));
+		}
+
 		FCanvasLineItem LineItem(BoundsVertices[VertexIndex], BoundsVertices[NextVertexIndex]);
 		LineItem.SetColor(GeometryVertexColor);
 		Canvas.DrawItem(LineItem);
 
-		const FVector2D MidPoint = (BoundsVertices[VertexIndex] + BoundsVertices[NextVertexIndex]) * 0.5f;
-        const FVector2D CornerPoint = BoundsVertices[VertexIndex];
-
         // Add edge hit proxy
         if (bDrawEdgeHitProxies)
         {
-            if (bIsHitTesting)
-            {
-                TSharedPtr<FSpriteSelectedSourceRegion> Data = MakeShareable(new FSpriteSelectedSourceRegion());
-                Data->SpritePtr = Sprite;
-                Data->VertexIndex = 4 + VertexIndex;
-                Canvas.SetHitProxy(new HSpriteSelectableObjectHitProxy(Data));
-            }
-
-            Canvas.DrawTile(MidPoint.X - EdgeCollisionVertexSize*0.5f, MidPoint.Y - EdgeCollisionVertexSize*0.5f, EdgeCollisionVertexSize, EdgeCollisionVertexSize, 0.f, 0.f, 1.f, 1.f, GeometryVertexColor, GWhiteTexture);
-
-            if (bIsHitTesting)
-            {
-                Canvas.SetHitProxy(nullptr);
-            }
+			const FVector2D MidPoint = (BoundsVertices[VertexIndex] + BoundsVertices[NextVertexIndex]) * 0.5f;
+			Canvas.DrawTile(MidPoint.X - EdgeCollisionVertexSize*0.5f, MidPoint.Y - EdgeCollisionVertexSize*0.5f, EdgeCollisionVertexSize, EdgeCollisionVertexSize, 0.f, 0.f, 1.f, 1.f, GeometryVertexColor, GWhiteTexture);
         }
 
+		if (bIsHitTesting)
+		{
+			Canvas.SetHitProxy(nullptr);
+		}
+
         // Add corner hit proxy
         if (bDrawCornerHitProxies)
         {
-            if (bIsHitTesting)
+			const FVector2D CornerPoint = BoundsVertices[VertexIndex];
+			
+			if (bIsHitTesting)
             {
                 TSharedPtr<FSpriteSelectedSourceRegion> Data = MakeShareable(new FSpriteSelectedSourceRegion());
                 Data->SpritePtr = Sprite;
