Sprite Editor: Removed dead method DrawTriangleList

diff --git a/Engine/Plugins/2D/Paper2D/Source/Paper2DEditor/Private/SpriteEditor/SpriteEditorViewportClient.cpp b/Engine/Plugins/2D/Paper2D/Source/Paper2DEditor/Private/SpriteEditor/SpriteEditorViewportClient.cpp
--- a/Engine/Plugins/2D/Paper2D/Source/Paper2DEditor/Private/SpriteEditor/SpriteEditorViewportClient.cpp
+++ b/Engine/Plugins/2D/Paper2D/Source/Paper2DEditor/Private/SpriteEditor/SpriteEditorViewportClient.cpp
@@ -216,35 +216,6 @@ FVector FSpriteEditorViewportClient::SourceTextureSpaceToWorldSpace(const FVecto
 	return Sprite->ConvertTextureSpaceToWorldSpace(SourcePoint);
 }
 
-void FSpriteEditorViewportClient::DrawTriangleList(FViewport& InViewport, FSceneView& View, FCanvas& Canvas, const TArray<FVector2D>& Triangles)
-{
-	// Draw the collision data (non-interactive)
-
-	// Draw the vertices
-	for (int32 VertexIndex = 0; VertexIndex < Triangles.Num(); ++VertexIndex)
-	{
-		const FVector2D VertexPos(TextureSpaceToScreenSpace(View, Triangles[VertexIndex]));
-		Canvas.DrawTile(VertexPos.X - SpriteEditingConstants::BakedCollisionVertexSize*0.5f, VertexPos.Y - SpriteEditingConstants::BakedCollisionVertexSize*0.5f, SpriteEditingConstants::BakedCollisionVertexSize, SpriteEditingConstants::BakedCollisionVertexSize, 0.f, 0.f, 1.f, 1.f, SpriteEditingConstants::BakedCollisionRenderColor, GWhiteTexture);
-	}
-
-	// Draw the lines
-	for (int32 TriangleIndex = 0; TriangleIndex < Triangles.Num() / 3; ++TriangleIndex)
-	{
-		for (int32 Offset = 0; Offset < 3; ++Offset)
-		{
-			const int32 VertexIndex = (TriangleIndex * 3) + Offset;
-			const int32 NextVertexIndex = (TriangleIndex * 3) + ((Offset + 1) % 3);
-
-			const FVector2D VertexPos(TextureSpaceToScreenSpace(View, Triangles[VertexIndex]));
-			const FVector2D NextVertexPos(TextureSpaceToScreenSpace(View, Triangles[NextVertexIndex]));
-
-			FCanvasLineItem LineItem(VertexPos, NextVertexPos);
-			LineItem.SetColor(SpriteEditingConstants::BakedCollisionLineRenderColor);
-			Canvas.DrawItem(LineItem);
-		}
-	}
-}
-
 void FSpriteEditorViewportClient::DrawRelatedSprites(FViewport& InViewport, FSceneView& View, FCanvas& Canvas, const FLinearColor& LineColor)
 {
 	for (int32 SpriteIndex = 0; SpriteIndex < RelatedSprites.Num(); ++SpriteIndex)
diff --git a/Engine/Plugins/2D/Paper2D/Source/Paper2DEditor/Private/SpriteEditor/SpriteEditorViewportClient.h b/Engine/Plugins/2D/Paper2D/Source/Paper2DEditor/Private/SpriteEditor/SpriteEditorViewportClient.h
--- a/Engine/Plugins/2D/Paper2D/Source/Paper2DEditor/Private/SpriteEditor/SpriteEditorViewportClient.h
+++ b/Engine/Plugins/2D/Paper2D/Source/Paper2DEditor/Private/SpriteEditor/SpriteEditorViewportClient.h
@@ -228,7 +228,6 @@ private:
 	FVector2D SourceTextureSpaceToScreenSpace(const FSceneView& View, const FVector2D& SourcePoint) const;
 	FVector SourceTextureSpaceToWorldSpace(const FVector2D& SourcePoint) const;
 
-	void DrawTriangleList(FViewport& InViewport, FSceneView& View, FCanvas& Canvas, const TArray<FVector2D>& Triangles);
 	void DrawBoundsAsText(FViewport& InViewport, FSceneView& View, FCanvas& Canvas, int32& YPos);
 	void DrawGeometry_CanvasPass(FViewport& InViewport, FSceneView& View, FCanvas& Canvas, FSpriteGeometryCollection& Geometry, const FLinearColor& GeometryVertexColor, const FLinearColor& NegativeGeometryVertexColor, bool bIsRenderGeometry);
 
diff --git a/Engine/Plugins/2D/Paper2D/Source/Paper2DEditor/Private/TileMapEditing/TileMapEditorViewportClient.h b/Engine/Plugins/2D/Paper2D/Source/Paper2DEditor/Private/TileMapEditing/TileMapEditorViewportClient.h
--- a/Engine/Plugins/2D/Paper2D/Source/Paper2DEditor/Private/TileMapEditing/TileMapEditorViewportClient.h
+++ b/Engine/Plugins/2D/Paper2D/Source/Paper2DEditor/Private/TileMapEditing/TileMapEditorViewportClient.h
@@ -80,7 +80,6 @@ private:
 		return TileMapEditorPtr.Pin()->GetTileMapBeingEdited();
 	}
 
-	void DrawTriangleList(FViewport& InViewport, FSceneView& View, FCanvas& Canvas, const TArray<FVector2D>& Triangles);
 	void DrawBoundsAsText(FViewport& InViewport, FSceneView& View, FCanvas& Canvas, int32& YPos);
 	
 	void BeginTransaction(const FText& SessionName);
