Paper2D: Fixed an issue with loading older sprite that could cause bad render or collision geometry to be created
- Geometry set to FullyCustom that started off with a bounding box would be incorrectly offset
- Geometry set to one of the boxes from a sprite with an offset source region would be incorrectly offset

diff --git a/Engine/Plugins/2D/Paper2D/Source/Paper2D/Private/PaperSprite.cpp b/Engine/Plugins/2D/Paper2D/Source/Paper2D/Private/PaperSprite.cpp
--- a/Engine/Plugins/2D/Paper2D/Source/Paper2D/Private/PaperSprite.cpp
+++ b/Engine/Plugins/2D/Paper2D/Source/Paper2D/Private/PaperSprite.cpp
@@ -1868,6 +1868,41 @@ void UPaperSprite::Serialize(FArchive& Ar)
 	Ar.UsingCustomVersion(FPaperCustomVersion::GUID);
 }
 
+static void UpdateGeometryToBeBoxPositionRelative(FSpriteGeometryCollection& Geometry)
+{
+	// Make sure the per-shape GeometryType fields are up to date (introduced in this version)
+	const bool bWasBoundingBox = (Geometry.GeometryType == ESpritePolygonMode::SourceBoundingBox) || (Geometry.GeometryType == ESpritePolygonMode::TightBoundingBox);
+
+	if (bWasBoundingBox)
+	{
+		for (FSpriteGeometryShape& Shape : Geometry.Shapes)
+		{
+			Shape.ShapeType = ESpriteShapeType::Box;
+
+			// Recenter the bounding box (BoxPosition is now defined as the center)
+			const FVector2D AmountToSubtract = Shape.BoxPosition + Shape.BoxSize * 0.5f;
+			Shape.BoxPosition += Shape.BoxSize * 0.5f;
+			for (FVector2D& Vertex : Shape.Vertices)
+			{
+				Vertex -= AmountToSubtract;
+			}
+		}
+	}
+	else
+	{
+		for (FSpriteGeometryShape& Shape : Geometry.Shapes)
+		{
+			Shape.ShapeType = ESpriteShapeType::Polygon;
+
+			// Make sure BoxPosition is zeroed since polygon points are relative to it now, but it was being ignored
+			//@TODO: Consider computing the center and recentering verts to keep the numbers small/relative
+			Shape.BoxPosition = FVector2D::ZeroVector;
+			Shape.BoxSize = FVector2D::ZeroVector;
+		}
+	}
+}
+
+
 void UPaperSprite::PostLoad()
 {
 	Super::PostLoad();
@@ -1885,45 +1920,8 @@ void UPaperSprite::PostLoad()
 
 	if (PaperVer < FPaperCustomVersion::RefactorPolygonStorageToSupportShapes)
 	{
-		// Make sure the per-shape GeometryType fields are up to date for collision (introduced in this version)
-		const bool bWasCollisionBoundingBox = (CollisionGeometry.GeometryType == ESpritePolygonMode::SourceBoundingBox) || (CollisionGeometry.GeometryType == ESpritePolygonMode::TightBoundingBox);
-		const ESpriteShapeType CollisionPerShapeType = bWasCollisionBoundingBox ? ESpriteShapeType::Box : ESpriteShapeType::Polygon;
-
-		for (FSpriteGeometryShape& Shape : CollisionGeometry.Shapes)
-		{
-			Shape.ShapeType = CollisionPerShapeType;
-
-			if (bWasCollisionBoundingBox)
-			{
-				// Recenter the bounding box (BoxPosition is now defined as the center)
-				Shape.BoxPosition += Shape.BoxSize * 0.5f;
-
-				for (FVector2D& Vertex : Shape.Vertices)
-				{
-					Vertex -= Shape.BoxSize * 0.5f;
-				}
-			}
-		}
-
-		// Make sure the per-shape GeometryType fields are up to date for rendering (introduced in this version)
-		const bool bWasRenderingBoundingBox = (RenderGeometry.GeometryType == ESpritePolygonMode::SourceBoundingBox) || (RenderGeometry.GeometryType == ESpritePolygonMode::TightBoundingBox);
-		const ESpriteShapeType RenderPerShapeType = bWasRenderingBoundingBox ? ESpriteShapeType::Box : ESpriteShapeType::Polygon;
-
-		for (FSpriteGeometryShape& Shape : RenderGeometry.Shapes)
-		{
-			Shape.ShapeType = RenderPerShapeType;
-
-			if (bWasRenderingBoundingBox)
-			{
-				// Recenter the bounding box (BoxPosition is now defined as the center)
-				Shape.BoxPosition += Shape.BoxSize * 0.5f;
-
-				for (FVector2D& Vertex : Shape.Vertices)
-				{
-					Vertex -= Shape.BoxSize * 0.5f;
-				}
-			}
-		}
+		UpdateGeometryToBeBoxPositionRelative(CollisionGeometry);
+		UpdateGeometryToBeBoxPositionRelative(RenderGeometry);
 	}
 
 	if (PaperVer < FPaperCustomVersion::AddPivotSnapToPixelGrid)
