Paper2D: Refactored sprite geometry editing, phase 2
- Pushed more code from FSpriteEditorViewportClient into FSpriteGeometryEditingHelper

diff --git a/Engine/Plugins/2D/Paper2D/Source/Paper2DEditor/Private/PaperEditorShared/AssetEditorSelectedItem.h b/Engine/Plugins/2D/Paper2D/Source/Paper2DEditor/Private/PaperEditorShared/AssetEditorSelectedItem.h
--- a/Engine/Plugins/2D/Paper2D/Source/Paper2DEditor/Private/PaperEditorShared/AssetEditorSelectedItem.h
+++ b/Engine/Plugins/2D/Paper2D/Source/Paper2DEditor/Private/PaperEditorShared/AssetEditorSelectedItem.h
@@ -41,11 +41,6 @@ public:
 	{
 	}
 
-	//@TODO: Doesn't belong here in base!
-	virtual void SplitEdge()
-	{
-	}
-
 	virtual FVector GetWorldPos() const
 	{
 		return FVector::ZeroVector;
diff --git a/Engine/Plugins/2D/Paper2D/Source/Paper2DEditor/Private/PaperEditorShared/SocketEditing.cpp b/Engine/Plugins/2D/Paper2D/Source/Paper2DEditor/Private/PaperEditorShared/SocketEditing.cpp
--- a/Engine/Plugins/2D/Paper2D/Source/Paper2DEditor/Private/PaperEditorShared/SocketEditing.cpp
+++ b/Engine/Plugins/2D/Paper2D/Source/Paper2DEditor/Private/PaperEditorShared/SocketEditing.cpp
@@ -93,11 +93,6 @@ FVector FSpriteSelectedSocket::GetWorldPos() const
 	return FVector::ZeroVector;
 }
 
-void FSpriteSelectedSocket::SplitEdge()
-{
-	// Nonsense operation on a socket, do nothing
-}
-
 //////////////////////////////////////////////////////////////////////////
 // FSocketEditingHelper
 
diff --git a/Engine/Plugins/2D/Paper2D/Source/Paper2DEditor/Private/PaperEditorShared/SocketEditing.h b/Engine/Plugins/2D/Paper2D/Source/Paper2DEditor/Private/PaperEditorShared/SocketEditing.h
--- a/Engine/Plugins/2D/Paper2D/Source/Paper2DEditor/Private/PaperEditorShared/SocketEditing.h
+++ b/Engine/Plugins/2D/Paper2D/Source/Paper2DEditor/Private/PaperEditorShared/SocketEditing.h
@@ -22,7 +22,6 @@ public:
 	virtual bool Equals(const FSelectedItem& OtherItem) const override;
 	virtual void ApplyDelta(const FVector2D& Delta, const FRotator& Rotation, const FVector& Scale3D, FWidget::EWidgetMode MoveMode) override;
 	FVector GetWorldPos() const override;
-	virtual void SplitEdge() override;
 	// End of FSelectedItem interface
 };
 
diff --git a/Engine/Plugins/2D/Paper2D/Source/Paper2DEditor/Private/PaperEditorShared/SpriteGeometryEditing.cpp b/Engine/Plugins/2D/Paper2D/Source/Paper2DEditor/Private/PaperEditorShared/SpriteGeometryEditing.cpp
--- a/Engine/Plugins/2D/Paper2D/Source/Paper2DEditor/Private/PaperEditorShared/SpriteGeometryEditing.cpp
+++ b/Engine/Plugins/2D/Paper2D/Source/Paper2DEditor/Private/PaperEditorShared/SpriteGeometryEditing.cpp
@@ -6,20 +6,41 @@
 #include "PaperGeomTools.h"
 #include "DynamicMeshBuilder.h"
 
-//@TODO: Hackery
+#define LOCTEXT_NAMESPACE "PaperGeometryEditing"
+
+//////////////////////////////////////////////////////////////////////////
+
 namespace SpriteEditingConstantsEX
 {
 	// The length and color to draw a normal tick in
-// 	const FLinearColor GeometryNormalColor(0.0f, 1.0f, 0.0f, 0.5f);
-// 	const float GeometryNormalLength = 15.0f;
-
+	const float GeometryNormalLength = 15.0f;
+	const FLinearColor GeometryNormalColor(0.0f, 1.0f, 0.0f, 0.5f);
+	
 	// Geometry edit rendering stuff
-// 	const float GeometryVertexSize = 8.0f;
-// 	const float GeometryBorderLineThickness = 2.0f;
 	const int32 CircleShapeNumSides = 64;
+	const float GeometryVertexSize = 8.0f;
+	const float GeometryBorderLineThickness = 2.0f;
 	const FLinearColor GeometrySelectedColor(FLinearColor::White);
 }
 
+//////////////////////////////////////////////////////////////////////////
+// FSpriteSelectionHelper
+
+void FSpriteSelectionHelper::ClearSelectionSet()
+{
+	SelectedItemSet.Empty();
+}
+
+void FSpriteSelectionHelper::DeleteSelectedItems(bool bShouldTransact)
+{
+	//@TODO: Allow deletion of sockets, etc... here
+}
+
+void FSpriteSelectionHelper::SelectItem(TSharedPtr<FSelectedItem> NewItem)
+{
+	SelectedItemSet.Add(NewItem);
+}
+
 //////////////////////////////////////////////////////////////////////////
 // FSpriteSelectedShape
 
@@ -32,15 +53,6 @@ FSpriteSelectedShape::FSpriteSelectedShape(ISpriteSelectionContext& InEditorCont
 {
 }
 
-// bool FSpriteSelectedShape::IsValidInEditor(UPaperSprite* Sprite, bool bInRenderData) const
-// {
-// 	if ((Sprite != nullptr) && (&Geometry == (bInRenderData ? &(Sprite->RenderGeometry) : &(Sprite->CollisionGeometry))))
-// 	{
-// 		return Geometry.Shapes.IsValidIndex(ShapeIndex);
-// 	}
-// 	return false;
-// }
-
 uint32 FSpriteSelectedShape::GetTypeHash() const
 {
 	return (ShapeIndex * 311);
@@ -165,27 +177,110 @@ FVector FSpriteSelectedShape::GetWorldPos() const
 	return Result;
 }
 
-void FSpriteSelectedShape::SplitEdge()
-{
-	// Nonsense operation on a whole shape, do nothing
-}
-
 //////////////////////////////////////////////////////////////////////////
 // FSpriteGeometryEditingHelper
 
 FSpriteGeometryEditingHelper::FSpriteGeometryEditingHelper(ISpriteSelectionContext& InEditorContext)
 	: EditorContext(InEditorContext)
+	, GeometryBeingEdited(nullptr)
 {
 	WidgetVertexColorMaterial = (UMaterial*)StaticLoadObject(UMaterial::StaticClass(), NULL, TEXT("/Engine/EditorMaterials/WidgetVertexColorMaterial.WidgetVertexColorMaterial"), NULL, LOAD_None, NULL);
 }
 
+void FSpriteGeometryEditingHelper::ClearSelectionSet()
+{
+	FSpriteSelectionHelper::ClearSelectionSet();
+	SelectedIDSet.Empty();
+
+	ResetAddPolygonMode();
+
+	EditorContext.InvalidateViewportAndHitProxies();
+}
+
+void FSpriteGeometryEditingHelper::DeleteSelectedItems(bool bShouldTransact)
+{
+	// Make a set of polygon-vertices "pairs" to delete
+	TSet<FShapeVertexPair> CompositeIndicesSet;
+	TSet<int32> ShapesToDeleteSet;
+
+	FSpriteGeometryCollection& Geometry = GetGeometryChecked();
+
+	for (TSharedPtr<FSelectedItem> SelectionIt : GetSelectionSet())
+	{
+		if (const FSpriteSelectedVertex* SelectedVertex = SelectionIt->CastTo<const FSpriteSelectedVertex>(FSelectionTypes::Vertex))
+		{
+			CompositeIndicesSet.Add(FShapeVertexPair(SelectedVertex->ShapeIndex, SelectedVertex->VertexIndex));
+
+			if (SelectedVertex->IsA(FSelectionTypes::Edge)) // add the "next" point for the edge
+			{
+				const int32 NextIndex = (SelectedVertex->VertexIndex + 1) % Geometry.Shapes[SelectedVertex->ShapeIndex].Vertices.Num();
+				CompositeIndicesSet.Add(FShapeVertexPair(SelectedVertex->ShapeIndex, NextIndex));
+			}
+		}
+		else if (const FSpriteSelectedShape* SelectedShape = SelectionIt->CastTo<const FSpriteSelectedShape>(FSelectionTypes::GeometryShape))
+		{
+			ShapesToDeleteSet.Add(SelectedShape->ShapeIndex);
+		}
+	}
+
+	if ((CompositeIndicesSet.Num() > 0) || (ShapesToDeleteSet.Num() > 0))
+	{
+		EditorContext.BeginTransaction(LOCTEXT("DeleteSelectionTransaction", "Delete Selection"));
+
+		// Delete the selected vertices first, as they may cause entire shapes to need to be deleted (sort so we delete from the back first)
+		TArray<FShapeVertexPair> CompositeIndices = CompositeIndicesSet.Array();
+		CompositeIndices.Sort([](const FShapeVertexPair& A, const FShapeVertexPair& B) { return (A.VertexIndex > B.VertexIndex); });
+		for (const FShapeVertexPair& Composite : CompositeIndices)
+		{
+			const int32 ShapeIndex = Composite.ShapeIndex;
+			const int32 VertexIndex = Composite.VertexIndex;
+			if (DeleteVertexInPolygonInternal(Geometry, ShapeIndex, VertexIndex))
+			{
+				ShapesToDeleteSet.Add(ShapeIndex);
+			}
+		}
+
+		// Delete the selected shapes (plus any shapes that became empty due to selected vertices)
+		if (ShapesToDeleteSet.Num() > 0)
+		{
+			// Sort so we delete from the back first
+			TArray<int32> ShapesToDeleteIndicies = ShapesToDeleteSet.Array();
+			ShapesToDeleteIndicies.Sort([](const int32& A, const int32& B) { return (A > B); });
+			for (const int32 ShapeToDeleteIndex : ShapesToDeleteIndicies)
+			{
+				Geometry.Shapes.RemoveAt(ShapeToDeleteIndex);
+			}
+		}
+
+		Geometry.GeometryType = ESpritePolygonMode::FullyCustom;
+		EditorContext.MarkTransactionAsDirty();
+
+		//@TODO: Should allow deleting other things when geometry is also selected!
+		//FSpriteGeometryEditingHelper::DeleteSelectedItems(/*bShouldTransact=*/ false);
+
+		ClearSelectionSet();
+		EditorContext.EndTransaction();
+		ResetAddPolygonMode();
+	}
+	else
+	{
+		FSpriteGeometryEditingHelper::DeleteSelectedItems(/*bShouldTransact=*/ true);
+	}
+}
+
 void FSpriteGeometryEditingHelper::AddReferencedObjects(FReferenceCollector& Collector)
 {
 	Collector.AddReferencedObject(WidgetVertexColorMaterial);
 }
 
-void FSpriteGeometryEditingHelper::DrawGeometry(const FSceneView& View, FPrimitiveDrawInterface& PDI, FSpriteGeometryCollection& Geometry, const FLinearColor& GeometryVertexColor, const FLinearColor& NegativeGeometryVertexColor, bool bIsRenderGeometry)
+void FSpriteGeometryEditingHelper::DrawGeometry(const FSceneView& View, FPrimitiveDrawInterface& PDI, const FLinearColor& GeometryVertexColor, const FLinearColor& NegativeGeometryVertexColor)
 {
+	if (GeometryBeingEdited == nullptr)
+	{
+		return;
+	}
+
+	FSpriteGeometryCollection& Geometry = GetGeometryChecked();
 	const bool bIsHitTesting = PDI.IsHitTesting();
 	const float UnitsPerPixel = EditorContext.SelectedItemGetUnitsPerPixel();
 
@@ -259,3 +354,471 @@ void FSpriteGeometryEditingHelper::DrawGeometry(const FSceneView& View, FPrimiti
 		}
 	}
 }
+
+void FSpriteGeometryEditingHelper::DrawGeometry_CanvasPass(FViewport& InViewport, FSceneView& View, FCanvas& Canvas, /*inout*/ int32& YPos, const FLinearColor& GeometryVertexColor, const FLinearColor& NegativeGeometryVertexColor)
+{
+	//@TODO: Display 'action required' text in a bright color for Add Polygon mode
+	//@TODO: Move all of the line drawing to the PDI pass
+	FSpriteGeometryCollection& Geometry = GetGeometryChecked();
+
+	// Display tool help
+	{
+		static const FText GeomHelpStr = LOCTEXT("GeomEditHelp", "Shift + click to insert a vertex.\nSelect one or more vertices and press Delete to remove them.\nDouble click a vertex to select a polygon\n");
+		static const FText GeomClickAddPolygon_NoSubtractive = LOCTEXT("GeomClickAddPolygon_NoSubtractive", "Click to start creating a polygon\n");
+		static const FText GeomClickAddPolygon_AllowSubtractive = LOCTEXT("GeomClickAddPolygon_AllowSubtractive", "Click to start creating a polygon\nCtrl + Click to start creating a subtractive polygon\n");
+		static const FText GeomAddVerticesHelpStr = LOCTEXT("GeomClickAddVertices", "Click to add points to the polygon\nEnter to finish adding a polygon\nEscape to cancel this polygon\n");
+
+		const FText* HelpStr;
+		if (IsAddingPolygon())
+		{
+			if (AddingPolygonIndex == INDEX_NONE)
+			{
+				HelpStr = bAllowSubtractivePolygons ? &GeomClickAddPolygon_AllowSubtractive : &GeomClickAddPolygon_NoSubtractive;
+			}
+			else
+			{
+				HelpStr = &GeomAddVerticesHelpStr;
+			}
+		}
+		else
+		{
+			HelpStr = &GeomHelpStr;
+		}
+
+		FCanvasTextItem TextItem(FVector2D(6, YPos), *HelpStr, GEngine->GetSmallFont(), FLinearColor::White);
+		TextItem.EnableShadow(FLinearColor::Black);
+		TextItem.Draw(&Canvas);
+		YPos += 54;
+	}
+
+	const bool bIsHitTesting = Canvas.IsHitTesting();
+
+	// Run thru the geometry shapes and draw hit proxies for them
+	for (int32 ShapeIndex = 0; ShapeIndex < Geometry.Shapes.Num(); ++ShapeIndex)
+	{
+		const FSpriteGeometryShape& Shape = Geometry.Shapes[ShapeIndex];
+
+		const bool bIsShapeSelected = IsGeometrySelected(FShapeVertexPair(ShapeIndex, INDEX_NONE));
+		const FLinearColor LineColorRaw = Shape.bNegativeWinding ? NegativeGeometryVertexColor : GeometryVertexColor;
+		const FLinearColor VertexColor = Shape.bNegativeWinding ? NegativeGeometryVertexColor : GeometryVertexColor;
+
+		const FLinearColor LineColor = Shape.IsShapeValid() ? LineColorRaw : FMath::Lerp(LineColorRaw, FLinearColor::Red, 0.8f);
+
+		// Draw the circle shape if necessary
+		if (Shape.ShapeType == ESpriteShapeType::Circle)
+		{
+			if (bIsHitTesting)
+			{
+				TSharedPtr<FSpriteSelectedShape> Data = MakeShareable(new FSpriteSelectedShape(EditorContext, Geometry, ShapeIndex, /*bIsBackground=*/ false));
+				Canvas.SetHitProxy(new HSpriteSelectableObjectHitProxy(Data));
+			}
+
+			// Draw the circle
+			const float RadiusX = Shape.BoxSize.X * 0.5f;
+			const float RadiusY = Shape.BoxSize.Y * 0.5f;
+
+			const float	AngleDelta = 2.0f * PI / SpriteEditingConstantsEX::CircleShapeNumSides;
+
+			const float LastX = Shape.BoxPosition.X + RadiusX;
+			const float LastY = Shape.BoxPosition.Y;
+			FVector2D LastVertexPos = TextureSpaceToScreenSpace(View, FVector2D(LastX, LastY));
+
+
+			for (int32 SideIndex = 0; SideIndex < SpriteEditingConstantsEX::CircleShapeNumSides; SideIndex++)
+			{
+				const float X = Shape.BoxPosition.X + RadiusX * FMath::Cos(AngleDelta * (SideIndex + 1));
+				const float Y = Shape.BoxPosition.Y + RadiusY * FMath::Sin(AngleDelta * (SideIndex + 1));
+				const FVector2D ScreenPos = TextureSpaceToScreenSpace(View, FVector2D(X, Y));
+
+				FCanvasLineItem LineItem(LastVertexPos, ScreenPos);
+				LineItem.SetColor(bIsShapeSelected ? SpriteEditingConstantsEX::GeometrySelectedColor : LineColor);
+				LineItem.LineThickness = SpriteEditingConstantsEX::GeometryBorderLineThickness;
+
+				Canvas.DrawItem(LineItem);
+
+				LastVertexPos = ScreenPos;
+			}
+
+
+			if (bIsHitTesting)
+			{
+				Canvas.SetHitProxy(nullptr);
+			}
+		}
+
+		// Draw lines connecting the vertices of the shape
+		for (int32 VertexIndex = 0; VertexIndex < Shape.Vertices.Num(); ++VertexIndex)
+		{
+			const int32 NextVertexIndex = (VertexIndex + 1) % Shape.Vertices.Num();
+
+			const FVector2D ScreenPos = TextureSpaceToScreenSpace(View, Shape.ConvertShapeSpaceToTextureSpace(Shape.Vertices[VertexIndex]));
+			const FVector2D NextScreenPos = TextureSpaceToScreenSpace(View, Shape.ConvertShapeSpaceToTextureSpace(Shape.Vertices[NextVertexIndex]));
+
+			const bool bIsThisVertexSelected = IsGeometrySelected(FShapeVertexPair(ShapeIndex, VertexIndex));
+			const bool bIsNextVertexSelected = IsGeometrySelected(FShapeVertexPair(ShapeIndex, NextVertexIndex));
+
+			const bool bIsEdgeSelected = bIsShapeSelected || (bIsThisVertexSelected && bIsNextVertexSelected);
+
+			// Draw the normal tick
+			if (bShowNormals)
+			{
+				const FVector2D Direction = (NextScreenPos - ScreenPos).GetSafeNormal();
+				const FVector2D Normal = FVector2D(-Direction.Y, Direction.X);
+
+				const FVector2D Midpoint = (ScreenPos + NextScreenPos) * 0.5f;
+				const FVector2D NormalPoint = Midpoint - Normal * SpriteEditingConstantsEX::GeometryNormalLength;
+				FCanvasLineItem LineItem(Midpoint, NormalPoint);
+				LineItem.SetColor(SpriteEditingConstantsEX::GeometryNormalColor);
+
+				Canvas.DrawItem(LineItem);
+			}
+
+			// Draw the edge
+			{
+				if (bIsHitTesting)
+				{
+					TSharedPtr<FSpriteSelectedEdge> Data = MakeShareable(new FSpriteSelectedEdge(EditorContext, Geometry, ShapeIndex, VertexIndex));
+					Canvas.SetHitProxy(new HSpriteSelectableObjectHitProxy(Data));
+				}
+
+				FCanvasLineItem LineItem(ScreenPos, NextScreenPos);
+				LineItem.SetColor(bIsEdgeSelected ? SpriteEditingConstantsEX::GeometrySelectedColor : LineColor);
+				LineItem.LineThickness = SpriteEditingConstantsEX::GeometryBorderLineThickness;
+				Canvas.DrawItem(LineItem);
+
+				if (bIsHitTesting)
+				{
+					Canvas.SetHitProxy(nullptr);
+				}
+			}
+		}
+
+		// Draw the vertices
+		for (int32 VertexIndex = 0; VertexIndex < Shape.Vertices.Num(); ++VertexIndex)
+		{
+			const FVector2D ScreenPos = TextureSpaceToScreenSpace(View, Shape.ConvertShapeSpaceToTextureSpace(Shape.Vertices[VertexIndex]));
+			const float X = ScreenPos.X;
+			const float Y = ScreenPos.Y;
+
+			const bool bIsVertexSelected = IsGeometrySelected(FShapeVertexPair(ShapeIndex, VertexIndex));
+			const bool bIsVertexLastAdded = IsAddingPolygon() && (AddingPolygonIndex == ShapeIndex) && (VertexIndex == Shape.Vertices.Num() - 1);
+			const bool bNeedHighlightVertex = bIsShapeSelected || bIsVertexSelected || bIsVertexLastAdded;
+
+			if (bIsHitTesting)
+			{
+				TSharedPtr<FSpriteSelectedVertex> Data = MakeShareable(new FSpriteSelectedVertex(EditorContext, Geometry, ShapeIndex, VertexIndex));
+				Canvas.SetHitProxy(new HSpriteSelectableObjectHitProxy(Data));
+			}
+
+			const float VertSize = SpriteEditingConstantsEX::GeometryVertexSize;
+			Canvas.DrawTile(ScreenPos.X - VertSize*0.5f, ScreenPos.Y - VertSize*0.5f, VertSize, VertSize, 0.f, 0.f, 1.f, 1.f, bNeedHighlightVertex ? SpriteEditingConstantsEX::GeometrySelectedColor : VertexColor, GWhiteTexture);
+
+			if (bIsHitTesting)
+			{
+				Canvas.SetHitProxy(nullptr);
+			}
+		}
+	}
+}
+
+void FSpriteGeometryEditingHelper::AddPointToGeometry(const FVector2D& TextureSpacePoint, const int32 SelectedPolygonIndex)
+{
+	FSpriteGeometryCollection& Geometry = GetGeometryChecked();
+
+	int32 ClosestShapeIndex = INDEX_NONE;
+	int32 ClosestVertexInsertIndex = INDEX_NONE;
+	float ClosestDistanceSquared = MAX_FLT;
+
+	int32 StartPolygonIndex = 0;
+	int32 EndPolygonIndex = Geometry.Shapes.Num();
+	if ((SelectedPolygonIndex >= 0) && (SelectedPolygonIndex < Geometry.Shapes.Num()))
+	{
+		StartPolygonIndex = SelectedPolygonIndex;
+		EndPolygonIndex = SelectedPolygonIndex + 1;
+	}
+
+	// Determine where we should insert the vertex
+	for (int32 PolygonIndex = StartPolygonIndex; PolygonIndex < EndPolygonIndex; ++PolygonIndex)
+	{
+		FSpriteGeometryShape& Polygon = Geometry.Shapes[PolygonIndex];
+		if (Polygon.Vertices.Num() >= 3)
+		{
+			for (int32 VertexIndex = 0; VertexIndex < Polygon.Vertices.Num(); ++VertexIndex)
+			{
+				FVector2D ClosestPoint;
+				const FVector2D LineStart = Polygon.ConvertShapeSpaceToTextureSpace(Polygon.Vertices[VertexIndex]);
+				int32 NextVertexIndex = (VertexIndex + 1) % Polygon.Vertices.Num();
+				const FVector2D LineEnd = Polygon.ConvertShapeSpaceToTextureSpace(Polygon.Vertices[NextVertexIndex]);
+				if (ClosestPointOnLine(TextureSpacePoint, LineStart, LineEnd, /*out*/ ClosestPoint))
+				{
+					const float CurrentDistanceSquared = FVector2D::DistSquared(ClosestPoint, TextureSpacePoint);
+					if (CurrentDistanceSquared < ClosestDistanceSquared)
+					{
+						ClosestShapeIndex = PolygonIndex;
+						ClosestDistanceSquared = CurrentDistanceSquared;
+						ClosestVertexInsertIndex = NextVertexIndex;
+					}
+				}
+			}
+		}
+		else
+		{
+			// Simply insert the vertex
+			for (int32 VertexIndex = 0; VertexIndex < Polygon.Vertices.Num(); ++VertexIndex)
+			{
+				const FVector2D CurrentVertexTS = Polygon.ConvertShapeSpaceToTextureSpace(Polygon.Vertices[VertexIndex]);
+				const float CurrentDistanceSquared = FVector2D::DistSquared(CurrentVertexTS, TextureSpacePoint);
+				if (CurrentDistanceSquared < ClosestDistanceSquared)
+				{
+					ClosestShapeIndex = PolygonIndex;
+					ClosestDistanceSquared = CurrentDistanceSquared;
+					ClosestVertexInsertIndex = VertexIndex + 1;
+				}
+			}
+		}
+	}
+
+	if ((ClosestVertexInsertIndex != INDEX_NONE) && (ClosestShapeIndex != INDEX_NONE))
+	{
+		FSpriteGeometryShape& Shape = Geometry.Shapes[ClosestShapeIndex];
+		if (Shape.ShapeType != ESpriteShapeType::Circle)
+		{
+			EditorContext.BeginTransaction(LOCTEXT("AddPolygonVertexTransaction", "Add Vertex to Polygon"));
+			Shape.Vertices.Insert(Shape.ConvertTextureSpaceToShapeSpace(TextureSpacePoint), ClosestVertexInsertIndex);
+			Shape.ShapeType = ESpriteShapeType::Polygon;
+			Geometry.GeometryType = ESpritePolygonMode::FullyCustom;
+			EditorContext.MarkTransactionAsDirty();
+			EditorContext.EndTransaction();
+
+			// Select this vertex
+			ClearSelectionSet();
+			AddPolygonVertexToSelection(ClosestShapeIndex, ClosestVertexInsertIndex);
+		}
+	}
+}
+
+void FSpriteGeometryEditingHelper::SelectGeometry(const FShapeVertexPair& GeometyItem)
+{
+	SelectedIDSet.Add(GeometyItem);
+	EditorContext.InvalidateViewportAndHitProxies();
+}
+
+bool FSpriteGeometryEditingHelper::IsGeometrySelected(const FShapeVertexPair& GeometyItem) const
+{
+	return SelectedIDSet.Contains(GeometyItem);
+}
+
+void FSpriteGeometryEditingHelper::AddShapeToSelection(const int32 ShapeIndex)
+{
+	FSpriteGeometryCollection& Geometry = GetGeometryChecked();
+	if (Geometry.Shapes.IsValidIndex(ShapeIndex))
+	{
+		if (!IsGeometrySelected(FShapeVertexPair(ShapeIndex, INDEX_NONE)))
+		{
+			FSpriteGeometryShape& Shape = Geometry.Shapes[ShapeIndex];
+
+			TSharedPtr<FSpriteSelectedShape> SelectedShape = MakeShareable(new FSpriteSelectedShape(EditorContext, Geometry, ShapeIndex));
+
+			SelectItem(SelectedShape);
+			SelectGeometry(FShapeVertexPair(ShapeIndex, INDEX_NONE));
+		}
+	}
+}
+
+void FSpriteGeometryEditingHelper::AddPolygonVertexToSelection(const int32 ShapeIndex, const int32 VertexIndex)
+{
+	FSpriteGeometryCollection& Geometry = GetGeometryChecked();
+	if (Geometry.Shapes.IsValidIndex(ShapeIndex))
+	{
+		FSpriteGeometryShape& Shape = Geometry.Shapes[ShapeIndex];
+		if (Shape.Vertices.IsValidIndex(VertexIndex))
+		{
+			if (!IsGeometrySelected(FShapeVertexPair(ShapeIndex, VertexIndex)))
+			{
+				TSharedPtr<FSpriteSelectedVertex> Vertex = MakeShareable(new FSpriteSelectedVertex(EditorContext, Geometry, ShapeIndex, VertexIndex));
+
+				SelectItem(Vertex);
+				SelectGeometry(FShapeVertexPair(ShapeIndex, VertexIndex));
+			}
+		}
+	}
+}
+
+void FSpriteGeometryEditingHelper::AddPolygonEdgeToSelection(const int32 ShapeIndex, const int32 FirstVertexIndex)
+{
+	FSpriteGeometryCollection& Geometry = GetGeometryChecked();
+	if (Geometry.Shapes.IsValidIndex(ShapeIndex))
+	{
+		FSpriteGeometryShape& Shape = Geometry.Shapes[ShapeIndex];
+		const int32 NextVertexIndex = (FirstVertexIndex + 1) % Shape.Vertices.Num();
+
+		AddPolygonVertexToSelection(ShapeIndex, FirstVertexIndex);
+		AddPolygonVertexToSelection(ShapeIndex, NextVertexIndex);
+	}
+}
+
+void FSpriteGeometryEditingHelper::SetGeometryBeingEdited(FSpriteGeometryCollection* NewGeometryBeingEdited, bool bInAllowCircles, bool bInAllowSubtractivePolygons)
+{
+	ClearSelectionSet();
+	GeometryBeingEdited = NewGeometryBeingEdited;
+	bAllowCircles = bInAllowCircles;
+	bAllowSubtractivePolygons = bInAllowSubtractivePolygons;
+}
+
+void FSpriteGeometryEditingHelper::AddNewCircleShape(const FVector2D& CircleLocation, float Radius)
+{
+	FSpriteGeometryCollection& Geometry = GetGeometryChecked();
+	EditorContext.BeginTransaction(LOCTEXT("AddCircleShapeTransaction", "Add Circle Shape"));
+
+	// Create the new shape
+	const FVector2D CircleSize2D(Radius*2.0f, Radius*2.0f);
+	Geometry.AddCircleShape(CircleLocation, CircleSize2D);
+	Geometry.GeometryType = ESpritePolygonMode::FullyCustom;
+
+	// Select the new shape
+	ClearSelectionSet();
+	AddShapeToSelection(Geometry.Shapes.Num() - 1);
+
+	EditorContext.MarkTransactionAsDirty();
+	EditorContext.EndTransaction();
+}
+
+void FSpriteGeometryEditingHelper::AddNewBoxShape(const FVector2D& BoxLocation, const FVector2D& BoxSize)
+{
+	FSpriteGeometryCollection& Geometry = GetGeometryChecked();
+	EditorContext.BeginTransaction(LOCTEXT("AddBoxShapeTransaction", "Add Box Shape"));
+
+	// Create the new shape
+	Geometry.AddRectangleShape(BoxLocation, BoxSize);
+	Geometry.GeometryType = ESpritePolygonMode::FullyCustom;
+
+	// Select the new shape
+	ClearSelectionSet();
+	AddShapeToSelection(Geometry.Shapes.Num() - 1);
+
+	EditorContext.MarkTransactionAsDirty();
+	EditorContext.EndTransaction();
+}
+
+void FSpriteGeometryEditingHelper::ToggleAddPolygonMode()
+{
+	if (IsAddingPolygon())
+	{
+		ResetAddPolygonMode();
+	}
+	else
+	{
+		ClearSelectionSet();
+
+		bIsAddingPolygon = true;
+		AddingPolygonIndex = INDEX_NONE;
+		EditorContext.InvalidateViewportAndHitProxies();
+	}
+}
+
+void FSpriteGeometryEditingHelper::AbandonAddPolygonMode()
+{
+	check(bIsAddingPolygon);
+
+	FSpriteGeometryCollection& Geometry = GetGeometryChecked();
+
+	if ((AddingPolygonIndex != INDEX_NONE) && Geometry.Shapes.IsValidIndex(AddingPolygonIndex))
+	{
+		EditorContext.BeginTransaction(LOCTEXT("DeletePolygon", "Delete Polygon"));
+
+		Geometry.Shapes.RemoveAt(AddingPolygonIndex);
+
+		EditorContext.MarkTransactionAsDirty();
+		EditorContext.EndTransaction();
+	}
+
+	ResetAddPolygonMode();
+}
+
+void FSpriteGeometryEditingHelper::TEMP_HandleAddPolygonClick(const FVector2D& TexturePoint, bool bWantsSubtractive)
+{
+	EditorContext.BeginTransaction(LOCTEXT("AddPolygonVertexTransaction", "Add Vertex to Polygon"));
+
+	FSpriteGeometryCollection& Geometry = GetGeometryChecked();
+	if (AddingPolygonIndex == INDEX_NONE)
+	{
+		FSpriteGeometryShape NewPolygon;
+		NewPolygon.ShapeType = ESpriteShapeType::Polygon;
+		NewPolygon.bNegativeWinding = CanAddSubtractivePolygon() && bWantsSubtractive;
+
+		Geometry.Shapes.Add(NewPolygon);
+		AddingPolygonIndex = Geometry.Shapes.Num() - 1;
+	}
+
+	if (Geometry.Shapes.IsValidIndex(AddingPolygonIndex))
+	{
+		FSpriteGeometryShape& Shape = Geometry.Shapes[AddingPolygonIndex];
+		Shape.ShapeType = ESpriteShapeType::Polygon;
+		Shape.Vertices.Add(Shape.ConvertTextureSpaceToShapeSpace(TexturePoint));
+	}
+	else
+	{
+		ResetAddPolygonMode();
+	}
+
+	Geometry.GeometryType = ESpritePolygonMode::FullyCustom;
+
+	EditorContext.MarkTransactionAsDirty();
+	EditorContext.EndTransaction();
+}
+
+bool FSpriteGeometryEditingHelper::ClosestPointOnLine(const FVector2D& Point, const FVector2D& LineStart, const FVector2D& LineEnd, FVector2D& OutClosestPoint)
+{
+	const FVector2D DL = LineEnd - LineStart;
+	const FVector2D DP = Point - LineStart;
+	const float DotL = FVector2D::DotProduct(DP, DL);
+	const float DotP = FVector2D::DotProduct(DL, DL);
+	if (DotP > 0.0001f)
+	{
+		float T = DotL / DotP;
+		if ((T >= -0.0001f) && (T <= 1.0001f))
+		{
+			T = FMath::Clamp(T, 0.0f, 1.0f);
+			OutClosestPoint = LineStart + T * DL;
+			return true;
+		}
+	}
+
+	return false;
+}
+
+bool FSpriteGeometryEditingHelper::DeleteVertexInPolygonInternal(FSpriteGeometryCollection& Geometry, const int32 ShapeIndex, const int32 VertexIndex)
+{
+	if (Geometry.Shapes.IsValidIndex(ShapeIndex))
+	{
+		FSpriteGeometryShape& Shape = Geometry.Shapes[ShapeIndex];
+		if (Shape.Vertices.IsValidIndex(VertexIndex))
+		{
+			Geometry.GeometryType = ESpritePolygonMode::FullyCustom;
+			Shape.ShapeType = ESpriteShapeType::Polygon;
+			Shape.Vertices.RemoveAt(VertexIndex);
+
+			if (Shape.Vertices.Num() == 0)
+			{
+				// Tell the caller they should delete the polygon since it has no more verts
+				return true;
+			}
+		}
+	}
+
+	return false;
+}
+
+
+FVector2D FSpriteGeometryEditingHelper::TextureSpaceToScreenSpace(const FSceneView& View, const FVector2D& SourcePoint) const
+{
+	const FVector WorldSpacePoint = EditorContext.TextureSpaceToWorldSpace(SourcePoint);
+
+	FVector2D PixelLocation;
+	View.WorldToPixel(WorldSpacePoint, /*out*/ PixelLocation);
+	return PixelLocation;
+}
+
+//////////////////////////////////////////////////////////////////////////
+
+#undef LOCTEXT_NAMESPACE
diff --git a/Engine/Plugins/2D/Paper2D/Source/Paper2DEditor/Private/PaperEditorShared/SpriteGeometryEditing.h b/Engine/Plugins/2D/Paper2D/Source/Paper2DEditor/Private/PaperEditorShared/SpriteGeometryEditing.h
--- a/Engine/Plugins/2D/Paper2D/Source/Paper2DEditor/Private/PaperEditorShared/SpriteGeometryEditing.h
+++ b/Engine/Plugins/2D/Paper2D/Source/Paper2DEditor/Private/PaperEditorShared/SpriteGeometryEditing.h
@@ -2,15 +2,19 @@
 
 #pragma once
 
+class FSelectedItem;
+struct FSpriteGeometryCollection;
+
 //////////////////////////////////////////////////////////////////////////
 // FShapeVertexPair
 
 struct FShapeVertexPair
 {
-	int32 ShapeIndex;
-
-	int32 VertexIndex;
+public:
+	const int32 ShapeIndex;
+	const int32 VertexIndex;
 
+public:
 	FShapeVertexPair()
 		: ShapeIndex(INDEX_NONE)
 		, VertexIndex(INDEX_NONE)
@@ -34,21 +38,128 @@ inline uint32 GetTypeHash(const FShapeVertexPair& Item)
 	return Item.VertexIndex + (Item.ShapeIndex * 311);
 }
 
+//////////////////////////////////////////////////////////////////////////
+// FSpriteSelectionHelper
+
+class FSpriteSelectionHelper
+{
+public:
+	bool HasAnySelectedItems() const
+	{
+		return SelectedItemSet.Num() > 0;
+	}
+
+	const TSet< TSharedPtr<FSelectedItem> >& GetSelectionSet() const
+	{
+		return SelectedItemSet;
+	}
+
+	void SelectItem(TSharedPtr<FSelectedItem> NewItem);
+
+	virtual void ClearSelectionSet();
+
+	virtual void DeleteSelectedItems(bool bShouldTransact = true);
+
+	~FSpriteSelectionHelper()
+	{
+	}
+
+private:
+	// Set of selected objects
+	TSet< TSharedPtr<FSelectedItem> > SelectedItemSet;
+};
+
 //////////////////////////////////////////////////////////////////////////
 // FSpriteGeometryEditingHelper
 
-class FSpriteGeometryEditingHelper : public FGCObject
+class FSpriteGeometryEditingHelper : public FSpriteSelectionHelper, public FGCObject
 {
 public:
 	FSpriteGeometryEditingHelper(class ISpriteSelectionContext& InEditorContext);
 
+	// FSpriteSelectionHelper interface
+	virtual void ClearSelectionSet() override;
+	virtual void DeleteSelectedItems(bool bShouldTransact = true) override;
+	// End of FSpriteSelectionHelper interface
+
 	// FGCObject interface
 	virtual void AddReferencedObjects(FReferenceCollector& Collector) override;
 	// End of FGCObject interface
 
-	void DrawGeometry(const FSceneView& View, FPrimitiveDrawInterface& PDI, FSpriteGeometryCollection& Geometry, const FLinearColor& GeometryVertexColor, const FLinearColor& NegativeGeometryVertexColor, bool bIsRenderGeometry);
+	void DrawGeometry(const FSceneView& View, FPrimitiveDrawInterface& PDI, const FLinearColor& GeometryVertexColor, const FLinearColor& NegativeGeometryVertexColor);
+	void DrawGeometry_CanvasPass(FViewport& InViewport, FSceneView& View, FCanvas& Canvas, /*inout*/ int32& YPos, const FLinearColor& GeometryVertexColor, const FLinearColor& NegativeGeometryVertexColor);
+
+	// Adds a point to the currently edited index
+	// If SelectedPolygonIndex is set, only that polygon is considered
+	void AddPointToGeometry(const FVector2D& TextureSpacePoint, const int32 SelectedPolygonIndex = INDEX_NONE);
+
+	void SelectGeometry(const FShapeVertexPair& GeometyItem);
+	bool IsGeometrySelected(const FShapeVertexPair& GeometyItem) const;
+
+
+	void AddShapeToSelection(const int32 ShapeIndex);
+	void AddPolygonVertexToSelection(const int32 ShapeIndex, const int32 VertexIndex);
+	void AddPolygonEdgeToSelection(const int32 ShapeIndex, const int32 FirstVertexIndex);
+
+	// Changes the geometry being edited (clears the selection set in the process)
+	void SetGeometryBeingEdited(FSpriteGeometryCollection* NewGeometryBeingEdited, bool bInAllowCircles, bool bInAllowSubtractivePolygons);
+
+	// Adds a new circle shape and selects it
+	void AddNewCircleShape(const FVector2D& CircleLocation, float Radius);
+	bool CanAddCircleShape() const { return bAllowCircles && (GeometryBeingEdited != nullptr); }
+
+	// Adds a new box shape and selects it
+	void AddNewBoxShape(const FVector2D& BoxLocation, const FVector2D& BoxSize);
+	bool CanAddBoxShape() const { return GeometryBeingEdited != nullptr; }
+
+	void ResetAddPolygonMode() { bIsAddingPolygon = false; }
+	void ToggleAddPolygonMode();
+	bool IsAddingPolygon() const { return bIsAddingPolygon; }
+	bool CanAddPolygon() const { return GeometryBeingEdited != nullptr; }
+	bool CanAddSubtractivePolygon() const { return CanAddPolygon() && bAllowSubtractivePolygons; }
+	void AbandonAddPolygonMode();
+
+	//@TODO: The code calling this will eventually be folded in here too
+	void TEMP_HandleAddPolygonClick(const FVector2D& TexturePoint, bool bWantsSubtractive);
+
+private:
+	static bool ClosestPointOnLine(const FVector2D& Point, const FVector2D& LineStart, const FVector2D& LineEnd, FVector2D& OutClosestPoint);
+
+	// Be sure to call this with polygonIndex and VertexIndex in descending order
+	// Returns true if the shape went to zero points and should be deleted itself
+	static bool DeleteVertexInPolygonInternal(FSpriteGeometryCollection& Geometry, const int32 ShapeIndex, const int32 VertexIndex);
+
+	FSpriteGeometryCollection& GetGeometryChecked()
+	{
+		check(GeometryBeingEdited);
+		return *GeometryBeingEdited;
+	}
+
+	//@TODO: This is needed for the canvas pass, but most of that code should go into the PDI pass
+	FVector2D TextureSpaceToScreenSpace(const FSceneView& View, const FVector2D& SourcePoint) const;
 
 private:
 	class ISpriteSelectionContext& EditorContext;
 	UMaterial* WidgetVertexColorMaterial;
+
+	// Set of selected vertices/shapes
+	TSet<FShapeVertexPair> SelectedIDSet;
+
+	// Active geometry being edited
+	FSpriteGeometryCollection* GeometryBeingEdited;
+
+	// Is waiting to add geometry
+	bool bIsAddingPolygon;
+
+	// The polygon index being added to, INDEX_NONE if we don't have a polygon yet
+	int32 AddingPolygonIndex;
+
+	// Should we draw normal ticks?
+	bool bShowNormals;
+
+	// Do we allow subtractive polygons?
+	bool bAllowSubtractivePolygons;
+
+	// Do we allow circles?
+	bool bAllowCircles;
 };
diff --git a/Engine/Plugins/2D/Paper2D/Source/Paper2DEditor/Private/SpriteEditor/SpriteEditor.cpp b/Engine/Plugins/2D/Paper2D/Source/Paper2DEditor/Private/SpriteEditor/SpriteEditor.cpp
--- a/Engine/Plugins/2D/Paper2D/Source/Paper2DEditor/Private/SpriteEditor/SpriteEditor.cpp
+++ b/Engine/Plugins/2D/Paper2D/Source/Paper2DEditor/Private/SpriteEditor/SpriteEditor.cpp
@@ -94,6 +94,8 @@ void SSpriteEditorViewport::BindCommands()
 	const FSpriteEditorCommands& Commands = FSpriteEditorCommands::Get();
 
 	TSharedRef<FSpriteEditorViewportClient> EditorViewportClientRef = EditorViewportClient.ToSharedRef();
+	FSpriteGeometryEditingHelper* GeometryEditingHelper = EditorViewportClientRef->GetGeometryEditHelper();
+
 
 	// Show toggles
 	CommandList->MapAction(
@@ -176,28 +178,24 @@ void SSpriteEditorViewport::BindCommands()
 		Commands.DeleteSelection,
 		FExecuteAction::CreateSP( EditorViewportClientRef, &FSpriteEditorViewportClient::DeleteSelection ),
 		FCanExecuteAction::CreateSP( EditorViewportClientRef, &FSpriteEditorViewportClient::CanDeleteSelection ));
-	CommandList->MapAction(
-		Commands.SplitEdge,
-		FExecuteAction::CreateSP( EditorViewportClientRef, &FSpriteEditorViewportClient::SplitEdge ),
-		FCanExecuteAction::CreateSP( EditorViewportClientRef, &FSpriteEditorViewportClient::CanSplitEdge ));
 	CommandList->MapAction(
 		Commands.AddBoxShape,
 		FExecuteAction::CreateSP(EditorViewportClientRef, &FSpriteEditorViewportClient::AddBoxShape),
-		FCanExecuteAction::CreateSP(EditorViewportClientRef, &FSpriteEditorViewportClient::CanAddBoxShape),
+		FCanExecuteAction::CreateRaw(GeometryEditingHelper, &FSpriteGeometryEditingHelper::CanAddBoxShape),
 		FIsActionChecked(),
-		FIsActionButtonVisible::CreateSP(EditorViewportClientRef, &FSpriteEditorViewportClient::CanAddBoxShape));
+		FIsActionButtonVisible::CreateRaw(GeometryEditingHelper, &FSpriteGeometryEditingHelper::CanAddBoxShape));
 	CommandList->MapAction(
 		Commands.ToggleAddPolygonMode,
-		FExecuteAction::CreateSP(EditorViewportClientRef, &FSpriteEditorViewportClient::ToggleAddPolygonMode),
-		FCanExecuteAction::CreateSP(EditorViewportClientRef, &FSpriteEditorViewportClient::CanAddPolygon),
-		FIsActionChecked::CreateSP(EditorViewportClientRef, &FSpriteEditorViewportClient::IsAddingPolygon),
-		FIsActionButtonVisible::CreateSP(EditorViewportClientRef, &FSpriteEditorViewportClient::CanAddPolygon));
+		FExecuteAction::CreateRaw(GeometryEditingHelper, &FSpriteGeometryEditingHelper::ToggleAddPolygonMode),
+		FCanExecuteAction::CreateRaw(GeometryEditingHelper, &FSpriteGeometryEditingHelper::CanAddPolygon),
+		FIsActionChecked::CreateRaw(GeometryEditingHelper, &FSpriteGeometryEditingHelper::IsAddingPolygon),
+		FIsActionButtonVisible::CreateRaw(GeometryEditingHelper, &FSpriteGeometryEditingHelper::CanAddPolygon));
 	CommandList->MapAction(
 		Commands.AddCircleShape,
 		FExecuteAction::CreateSP(EditorViewportClientRef, &FSpriteEditorViewportClient::AddCircleShape),
-		FCanExecuteAction::CreateSP(EditorViewportClientRef, &FSpriteEditorViewportClient::CanAddCircleShape),
+		FCanExecuteAction::CreateRaw(GeometryEditingHelper, &FSpriteGeometryEditingHelper::CanAddCircleShape),
 		FIsActionChecked(),
-		FIsActionButtonVisible::CreateSP(EditorViewportClientRef, &FSpriteEditorViewportClient::CanAddCircleShape));
+		FIsActionButtonVisible::CreateRaw(GeometryEditingHelper, &FSpriteGeometryEditingHelper::CanAddCircleShape));
 	CommandList->MapAction(
 		Commands.SnapAllVertices,
 		FExecuteAction::CreateSP( EditorViewportClientRef, &FSpriteEditorViewportClient::SnapAllVerticesToPixelGrid ),
diff --git a/Engine/Plugins/2D/Paper2D/Source/Paper2DEditor/Private/SpriteEditor/SpriteEditorSelections.h b/Engine/Plugins/2D/Paper2D/Source/Paper2DEditor/Private/SpriteEditor/SpriteEditorSelections.h
--- a/Engine/Plugins/2D/Paper2D/Source/Paper2DEditor/Private/SpriteEditor/SpriteEditorSelections.h
+++ b/Engine/Plugins/2D/Paper2D/Source/Paper2DEditor/Private/SpriteEditor/SpriteEditorSelections.h
@@ -29,6 +29,10 @@ public:
 	virtual FVector TextureSpaceToWorldSpace(const FVector2D& SourcePoint) const = 0;
 	virtual bool SelectedItemIsSelected(const struct FShapeVertexPair& Item) const = 0;
 	virtual float SelectedItemGetUnitsPerPixel() const = 0;
+	virtual void BeginTransaction(const FText& SessionName) = 0;
+	virtual void MarkTransactionAsDirty() = 0;
+	virtual void EndTransaction() = 0;
+	virtual void InvalidateViewportAndHitProxies() = 0;
 };
 
 //////////////////////////////////////////////////////////////////////////
@@ -197,11 +201,6 @@ public:
 	{
 		return GetWorldPosIndexed(VertexIndex);
 	}
-
-	virtual void SplitEdge() override
-	{
-		// Nonsense operation on a vertex, do nothing
-	}
 };
 
 
@@ -218,24 +217,22 @@ public:
 	FSpriteGeometryCollection& Geometry;
 
 	// The index of this shape in the geometry above
-	int32 ShapeIndex;
+	const int32 ShapeIndex;
 
 	// Is this a background object that should have lower priority?
-	bool bIsBackground;
+	const bool bIsBackground;
 
 	TWeakObjectPtr<UPaperSprite> SpritePtr;
 public:
 	FSpriteSelectedShape(ISpriteSelectionContext& InEditorContext, FSpriteGeometryCollection& InGeometry, int32 InShapeIndex, bool bInIsBackground = false);
 
 	// FSelectedItem interface
-//	virtual bool IsValidInEditor(UPaperSprite* Sprite, bool bInRenderData) const override;
 	virtual uint32 GetTypeHash() const override;
 	virtual EMouseCursor::Type GetMouseCursor() const override;
 	virtual bool Equals(const FSelectedItem& OtherItem) const override;
 	virtual bool IsBackgroundObject() const override;
 	virtual void ApplyDelta(const FVector2D& Delta, const FRotator& Rotation, const FVector& Scale3D, FWidget::EWidgetMode MoveMode) override;
 	virtual FVector GetWorldPos() const override;
-	virtual void SplitEdge() override;
 	// End of FSelectedItem interface
 };
 
@@ -245,38 +242,28 @@ public:
 class FSpriteSelectedVertex : public FSelectedItem
 {
 public:
-	int32 VertexIndex;
-	int32 ShapeIndex;
+	// The editor context
+	ISpriteSelectionContext& EditorContext;
 
-	// If true, it's render data, otherwise it's collision data
-	bool bRenderData;
+	// The geometry that this shape belongs to
+	FSpriteGeometryCollection& Geometry;
+
+	const int32 ShapeIndex;
+	const int32 VertexIndex;
 
-	TWeakObjectPtr<UPaperSprite> SpritePtr;
 public:
-	FSpriteSelectedVertex()
+	FSpriteSelectedVertex(ISpriteSelectionContext& InEditorContext, FSpriteGeometryCollection& InGeometry, int32 InShapeIndex, int32 InVertexIndex)
 		: FSelectedItem(FSelectionTypes::Vertex)
-		, VertexIndex(0)
-		, ShapeIndex(0)
-		, bRenderData(false)
-	{
-	}
-
-	virtual bool IsValidInEditor(UPaperSprite* Sprite, bool bInRenderData) const
+		, EditorContext(InEditorContext)
+		, Geometry(InGeometry)
+		, ShapeIndex(InShapeIndex)
+		, VertexIndex(InVertexIndex)
 	{
-		if (Sprite && SpritePtr == Sprite && bRenderData == bInRenderData)
-		{
-			FSpriteGeometryCollection& Geometry = bRenderData ? Sprite->RenderGeometry : Sprite->CollisionGeometry;
-			if (Geometry.Shapes.IsValidIndex(ShapeIndex) && Geometry.Shapes[ShapeIndex].Vertices.IsValidIndex(VertexIndex))
-			{
-				return true;
-			}
-		}
-		return false;
 	}
 
 	virtual uint32 GetTypeHash() const override
 	{
-		return VertexIndex + (ShapeIndex * 311) + (bRenderData ? 1063 : 0);
+		return VertexIndex + (ShapeIndex * 311);
 	}
 
 	virtual bool Equals(const FSelectedItem& OtherItem) const override
@@ -286,7 +273,7 @@ public:
 			const FSpriteSelectedVertex& V1 = *this;
 			const FSpriteSelectedVertex& V2 = *(FSpriteSelectedVertex*)(&OtherItem);
 
-			return (V1.VertexIndex == V2.VertexIndex) && (V1.ShapeIndex == V2.ShapeIndex) && (V1.bRenderData == V2.bRenderData) && (V1.SpritePtr == V2.SpritePtr);
+			return (V1.VertexIndex == V2.VertexIndex) && (V1.ShapeIndex == V2.ShapeIndex) && (&(V1.Geometry) == &(V2.Geometry));
 		}
 		else
 		{
@@ -296,26 +283,22 @@ public:
 
 	virtual void ApplyDeltaIndexed(const FVector2D& Delta, int32 TargetVertexIndex)
 	{
-		if (UPaperSprite* Sprite = SpritePtr.Get())
+		if (Geometry.Shapes.IsValidIndex(ShapeIndex))
 		{
-			FSpriteGeometryCollection& Geometry = bRenderData ? Sprite->RenderGeometry : Sprite->CollisionGeometry;
-			if (Geometry.Shapes.IsValidIndex(ShapeIndex))
+			FSpriteGeometryShape& Shape = Geometry.Shapes[ShapeIndex];
+			TArray<FVector2D>& Vertices = Shape.Vertices;
+			TargetVertexIndex = (Vertices.Num() > 0) ? (TargetVertexIndex % Vertices.Num()) : 0;
+			if (Vertices.IsValidIndex(TargetVertexIndex))
 			{
-				FSpriteGeometryShape& Shape = Geometry.Shapes[ShapeIndex];
-				TArray<FVector2D>& Vertices = Shape.Vertices;
-				TargetVertexIndex = (Vertices.Num() > 0) ? (TargetVertexIndex % Vertices.Num()) : 0;
-				if (Vertices.IsValidIndex(TargetVertexIndex))
-				{
-					FVector2D& ShapeSpaceVertex = Vertices[TargetVertexIndex];
-
-					const FVector WorldSpaceDelta = PaperAxisX * Delta.X + PaperAxisY * Delta.Y;
-					const FVector2D TextureSpaceDelta = Sprite->ConvertWorldSpaceDeltaToTextureSpace(WorldSpaceDelta);
-					const FVector2D NewTextureSpacePos = Shape.ConvertShapeSpaceToTextureSpace(ShapeSpaceVertex) + TextureSpaceDelta;
-					ShapeSpaceVertex = Shape.ConvertTextureSpaceToShapeSpace(NewTextureSpacePos);
-
-					Geometry.GeometryType = ESpritePolygonMode::FullyCustom;
-					Shape.ShapeType = ESpriteShapeType::Polygon;
-				}
+				FVector2D& ShapeSpaceVertex = Vertices[TargetVertexIndex];
+
+				const FVector WorldSpaceDelta = PaperAxisX * Delta.X + PaperAxisY * Delta.Y;
+				const FVector2D TextureSpaceDelta = EditorContext.SelectedItemConvertWorldSpaceDeltaToLocalSpace(WorldSpaceDelta);
+				const FVector2D NewTextureSpacePos = Shape.ConvertShapeSpaceToTextureSpace(ShapeSpaceVertex) + TextureSpaceDelta;
+				ShapeSpaceVertex = Shape.ConvertTextureSpaceToShapeSpace(NewTextureSpacePos);
+
+				Geometry.GeometryType = ESpritePolygonMode::FullyCustom;
+				Shape.ShapeType = ESpriteShapeType::Polygon;
 			}
 		}
 	}
@@ -324,21 +307,16 @@ public:
 	{
 		FVector Result = FVector::ZeroVector;
 
-		if (UPaperSprite* Sprite = SpritePtr.Get())
+		if (Geometry.Shapes.IsValidIndex(ShapeIndex))
 		{
-			FSpriteGeometryCollection& Geometry = bRenderData ? Sprite->RenderGeometry : Sprite->CollisionGeometry;
-
-			if (Geometry.Shapes.IsValidIndex(ShapeIndex))
+			const FSpriteGeometryShape& Shape = Geometry.Shapes[ShapeIndex];
+			const TArray<FVector2D>& Vertices = Shape.Vertices;
+			TargetVertexIndex = (Vertices.Num() > 0) ? (TargetVertexIndex % Vertices.Num()) : 0;
+			if (Vertices.IsValidIndex(TargetVertexIndex))
 			{
-				const FSpriteGeometryShape& Shape = Geometry.Shapes[ShapeIndex];
-				const TArray<FVector2D>& Vertices = Shape.Vertices;
-				TargetVertexIndex = (Vertices.Num() > 0) ? (TargetVertexIndex % Vertices.Num()) : 0;
-				if (Vertices.IsValidIndex(TargetVertexIndex))
-				{
-					const FVector2D& ShapeSpacePos = Vertices[TargetVertexIndex];
-					const FVector2D TextureSpacePos = Shape.ConvertShapeSpaceToTextureSpace(ShapeSpacePos);
-					Result = Sprite->ConvertTextureSpaceToWorldSpace(TextureSpacePos);
-				}
+				const FVector2D& ShapeSpacePos = Vertices[TargetVertexIndex];
+				const FVector2D TextureSpacePos = Shape.ConvertShapeSpaceToTextureSpace(ShapeSpacePos);
+				Result = EditorContext.TextureSpaceToWorldSpace(TextureSpacePos);
 			}
 		}
 
@@ -357,24 +335,18 @@ public:
 	{
 		return GetWorldPosIndexed(VertexIndex);
 	}
-
-	virtual void SplitEdge() override
-	{
-		// Nonsense operation on a vertex, do nothing
-	}
 };
 
 
 //////////////////////////////////////////////////////////////////////////
 // FSpriteSelectedEdge
-//@TODO: Much hackery lies here; need a real data structure!
-
 
+// Note: Defined based on a vertex index; this is the edge between the vertex and the next one.
 class FSpriteSelectedEdge : public FSpriteSelectedVertex
 {
-	// Note: Defined based on a vertex index; this is the edge between the vertex and the next one.
 public:
-	FSpriteSelectedEdge()
+	FSpriteSelectedEdge(ISpriteSelectionContext& InEditorContext, FSpriteGeometryCollection& InGeometry, int32 InShapeIndex, int32 InVertexIndex)
+		: FSpriteSelectedVertex(InEditorContext, InGeometry, InShapeIndex, InVertexIndex)
 	{
 		TypeName = FSelectionTypes::Edge;
 	}
@@ -391,7 +363,7 @@ public:
 			const FSpriteSelectedEdge& E1 = *this;
 			const FSpriteSelectedEdge& E2 = *(FSpriteSelectedEdge*)(&OtherItem);
 
-			return (E1.VertexIndex == E2.VertexIndex) && (E1.ShapeIndex == E2.ShapeIndex) && (E1.bRenderData == E2.bRenderData) && (E1.SpritePtr == E2.SpritePtr);
+			return (E1.VertexIndex == E2.VertexIndex) && (E1.ShapeIndex == E2.ShapeIndex) && (&(E1.Geometry) == &(E2.Geometry));
 		}
 		else
 		{
@@ -412,29 +384,4 @@ public:
 
 		return (Pos1 + Pos2) * 0.5f;
 	}
-
-	virtual void SplitEdge() override
-	{
-		if (UPaperSprite* Sprite = SpritePtr.Get())
-		{
-			FSpriteGeometryCollection& Geometry = bRenderData ? Sprite->RenderGeometry : Sprite->CollisionGeometry;
-
-			if (Geometry.Shapes.IsValidIndex(ShapeIndex))
-			{
-				FSpriteGeometryShape& Shape = Geometry.Shapes[ShapeIndex];
-				TArray<FVector2D>& Vertices = Shape.Vertices;
-				if (Vertices.IsValidIndex(VertexIndex))
-				{
-					const int32 NextVertexIndex = (VertexIndex + 1) % Vertices.Num();
-
-					const FVector2D NewShapeSpacePos = (Vertices[VertexIndex] + Vertices[NextVertexIndex]) * 0.5f;
-					
-					Vertices.Insert(NewShapeSpacePos, VertexIndex + 1);
-
-					Geometry.GeometryType = ESpritePolygonMode::FullyCustom;
-					Shape.ShapeType = ESpriteShapeType::Polygon;
-				}
-			}
-		}
-	}
 };
diff --git a/Engine/Plugins/2D/Paper2D/Source/Paper2DEditor/Private/SpriteEditor/SpriteEditorViewportClient.cpp b/Engine/Plugins/2D/Paper2D/Source/Paper2DEditor/Private/SpriteEditor/SpriteEditorViewportClient.cpp
--- a/Engine/Plugins/2D/Paper2D/Source/Paper2DEditor/Private/SpriteEditor/SpriteEditorViewportClient.cpp
+++ b/Engine/Plugins/2D/Paper2D/Source/Paper2DEditor/Private/SpriteEditor/SpriteEditorViewportClient.cpp
@@ -42,10 +42,6 @@ namespace SpriteEditingConstants
 	const FLinearColor BakedCollisionRenderColor(1.0f, 1.0f, 0.0f, 0.5f);
 	const float BakedCollisionVertexSize = 3.0f;
 
-	// The length and color to draw a normal tick in
-	const FLinearColor GeometryNormalColor(0.0f, 1.0f, 0.0f, 0.5f);
-	const float GeometryNormalLength = 15.0f;
-
 	const FLinearColor MarqueeColor(1.0f, 1.0f, 1.0f, 0.5f);
 
 	const FLinearColor SourceRegionBoundsColor(1.0f, 1.0f, 1.0f, 0.8f);
@@ -55,12 +51,6 @@ namespace SpriteEditingConstants
 	const FLinearColor CollisionShapeColor(0.0f, 0.7f, 1.0f, 1.0f);
 	const FLinearColor RenderShapeColor(1.0f, 0.2f, 0.0f, 1.0f);
 	const FLinearColor SubtractiveRenderShapeColor(0.0f, 0.2f, 1.0f, 1.0f);
-
-	// Geometry edit rendering stuff
-	const float GeometryVertexSize = 8.0f;
-	const float GeometryBorderLineThickness = 2.0f;
-	const int32 CircleShapeNumSides = 64;
-	const FLinearColor GeometrySelectedColor(FLinearColor::White);
 }
 
 //////////////////////////////////////////////////////////////////////////
@@ -192,15 +182,6 @@ void FSpriteEditorViewportClient::UpdateSourceTextureSpriteFromSprite(UPaperSpri
 	}
 }
 
-FVector2D FSpriteEditorViewportClient::TextureSpaceToScreenSpace(const FSceneView& View, const FVector2D& SourcePoint) const
-{
-	const FVector WorldSpacePoint = TextureSpaceToWorldSpace(SourcePoint);
-
-	FVector2D PixelLocation;
-	View.WorldToPixel(WorldSpacePoint, /*out*/ PixelLocation);
-	return PixelLocation;
-}
-
 FVector2D FSpriteEditorViewportClient::SourceTextureSpaceToScreenSpace(const FSceneView& View, const FVector2D& SourcePoint) const
 {
 	const FVector WorldSpacePoint = SourceTextureSpaceToWorldSpace(SourcePoint);
@@ -326,144 +307,6 @@ void FSpriteEditorViewportClient::DrawMarquee(FViewport& InViewport, FSceneView&
 	}
 }
 
-void FSpriteEditorViewportClient::DrawGeometry_CanvasPass(FViewport& InViewport, FSceneView& View, FCanvas& Canvas, FSpriteGeometryCollection& Geometry, const FLinearColor& GeometryVertexColor, const FLinearColor& NegativeGeometryVertexColor, bool bIsRenderGeometry)
-{
-	const bool bIsHitTesting = Canvas.IsHitTesting();
-	UPaperSprite* Sprite = GetSpriteBeingEdited();
-
-	// Run thru the geometry shapes and draw hit proxies for them
-	for (int32 ShapeIndex = 0; ShapeIndex < Geometry.Shapes.Num(); ++ShapeIndex)
-	{
-		const FSpriteGeometryShape& Shape = Geometry.Shapes[ShapeIndex];
-
-		const bool bIsShapeSelected = IsShapeSelected(ShapeIndex);
-		const FLinearColor LineColorRaw = Shape.bNegativeWinding ? NegativeGeometryVertexColor : GeometryVertexColor;
-		const FLinearColor VertexColor = Shape.bNegativeWinding ? NegativeGeometryVertexColor : GeometryVertexColor;
-
-		const FLinearColor LineColor = Shape.IsShapeValid() ? LineColorRaw : FMath::Lerp(LineColorRaw, FLinearColor::Red, 0.8f);
-
-		// Draw the circle shape if necessary
-		if (Shape.ShapeType == ESpriteShapeType::Circle)
-		{
-			if (bIsHitTesting)
-			{
-				TSharedPtr<FSpriteSelectedShape> Data = MakeShareable(new FSpriteSelectedShape(*this, Geometry, ShapeIndex, /*bIsBackground=*/ false));
-				Canvas.SetHitProxy(new HSpriteSelectableObjectHitProxy(Data));
-			}
-
-			// Draw the circle
-			const float RadiusX = Shape.BoxSize.X * 0.5f;
-			const float RadiusY = Shape.BoxSize.Y * 0.5f;
-
-			const float	AngleDelta = 2.0f * PI / SpriteEditingConstants::CircleShapeNumSides;
-			
-			const float LastX = Shape.BoxPosition.X + RadiusX;
-			const float LastY = Shape.BoxPosition.Y;
-			FVector2D LastVertexPos = TextureSpaceToScreenSpace(View, FVector2D(LastX, LastY));
-
-
-			for (int32 SideIndex = 0; SideIndex < SpriteEditingConstants::CircleShapeNumSides; SideIndex++)
-			{
-				const float X = Shape.BoxPosition.X + RadiusX * FMath::Cos(AngleDelta * (SideIndex + 1));
-				const float Y = Shape.BoxPosition.Y + RadiusY * FMath::Sin(AngleDelta * (SideIndex + 1));
-				const FVector2D ScreenPos = TextureSpaceToScreenSpace(View, FVector2D(X, Y));
-
-				FCanvasLineItem LineItem(LastVertexPos, ScreenPos);
-				LineItem.SetColor(bIsShapeSelected ? SpriteEditingConstants::GeometrySelectedColor : LineColor);
-				LineItem.LineThickness = SpriteEditingConstants::GeometryBorderLineThickness;
-
-				Canvas.DrawItem(LineItem);
-
-				LastVertexPos = ScreenPos;
-			}
-
-
-			if (bIsHitTesting)
-			{
-				Canvas.SetHitProxy(nullptr);
-			}
-		}
-
-		// Draw lines connecting the vertices of the shape
-		for (int32 VertexIndex = 0; VertexIndex < Shape.Vertices.Num(); ++VertexIndex)
-		{
-			const int32 NextVertexIndex = (VertexIndex + 1) % Shape.Vertices.Num();
-
-			const FVector2D ScreenPos = TextureSpaceToScreenSpace(View, Shape.ConvertShapeSpaceToTextureSpace(Shape.Vertices[VertexIndex]));
-			const FVector2D NextScreenPos = TextureSpaceToScreenSpace(View, Shape.ConvertShapeSpaceToTextureSpace(Shape.Vertices[NextVertexIndex]));
-
-			const bool bIsEdgeSelected = bIsShapeSelected || (IsPolygonVertexSelected(ShapeIndex, VertexIndex) && IsPolygonVertexSelected(ShapeIndex, NextVertexIndex));
-
-			// Draw the normal tick
-			if (bShowNormals)
-			{
-				const FVector2D Direction = (NextScreenPos - ScreenPos).GetSafeNormal();
-				const FVector2D Normal = FVector2D(-Direction.Y, Direction.X);
-
-				const FVector2D Midpoint = (ScreenPos + NextScreenPos) * 0.5f;
-				const FVector2D NormalPoint = Midpoint - Normal * SpriteEditingConstants::GeometryNormalLength;
-				FCanvasLineItem LineItem(Midpoint, NormalPoint);
-				LineItem.SetColor(SpriteEditingConstants::GeometryNormalColor);
-
-				Canvas.DrawItem(LineItem);
-			}
-
-			// Draw the edge
-			{
-				if (bIsHitTesting)
-				{
-					TSharedPtr<FSpriteSelectedEdge> Data = MakeShareable(new FSpriteSelectedEdge());
-					Data->SpritePtr = Sprite;
-					Data->ShapeIndex = ShapeIndex;
-					Data->VertexIndex = VertexIndex;
-					Data->bRenderData = bIsRenderGeometry;
-					Canvas.SetHitProxy(new HSpriteSelectableObjectHitProxy(Data));
-				}
-
-				FCanvasLineItem LineItem(ScreenPos, NextScreenPos);
-				LineItem.SetColor(bIsEdgeSelected ? SpriteEditingConstants::GeometrySelectedColor : LineColor);
-				LineItem.LineThickness = SpriteEditingConstants::GeometryBorderLineThickness;
-				Canvas.DrawItem(LineItem);
-
-				if (bIsHitTesting)
-				{
-					Canvas.SetHitProxy(nullptr);
-				}
-			}
-		}
-
-		// Draw the vertices
-		for (int32 VertexIndex = 0; VertexIndex < Shape.Vertices.Num(); ++VertexIndex)
-		{
-			const FVector2D ScreenPos = TextureSpaceToScreenSpace(View, Shape.ConvertShapeSpaceToTextureSpace(Shape.Vertices[VertexIndex]));
-			const float X = ScreenPos.X;
-			const float Y = ScreenPos.Y;
-
-			const bool bIsVertexSelected = IsPolygonVertexSelected(ShapeIndex, VertexIndex);
-			const bool bIsVertexLastAdded = IsAddingPolygon() && (AddingPolygonIndex == ShapeIndex) && (VertexIndex == Shape.Vertices.Num() - 1);
-			const bool bNeedHighlightVertex = bIsShapeSelected || bIsVertexSelected || bIsVertexLastAdded;
-
-			if (bIsHitTesting)
-			{
-				TSharedPtr<FSpriteSelectedVertex> Data = MakeShareable(new FSpriteSelectedVertex());
-				Data->SpritePtr = Sprite;
-				Data->ShapeIndex = ShapeIndex;
-				Data->VertexIndex = VertexIndex;
-				Data->bRenderData = bIsRenderGeometry;
-				Canvas.SetHitProxy(new HSpriteSelectableObjectHitProxy(Data));
-			}
-
-			const float VertSize = SpriteEditingConstants::GeometryVertexSize;
-			Canvas.DrawTile(ScreenPos.X - VertSize*0.5f, ScreenPos.Y - VertSize*0.5f, VertSize, VertSize, 0.f, 0.f, 1.f, 1.f, bNeedHighlightVertex ? SpriteEditingConstants::GeometrySelectedColor : VertexColor, GWhiteTexture);
-
-			if (bIsHitTesting)
-			{
-				Canvas.SetHitProxy(nullptr);
-			}
-		}
-	}
-}
-
 void FSpriteEditorViewportClient::DrawGeometryStats(FViewport& InViewport, FSceneView& View, FCanvas& Canvas, const FSpriteGeometryCollection& Geometry, bool bIsRenderGeometry, int32& YPos)
 {
 	// Draw the type of geometry we're displaying stats for
@@ -707,10 +550,6 @@ void FSpriteEditorViewportClient::DrawCanvas(FViewport& Viewport, FSceneView& Vi
 
 	int32 YPos = 42;
 
-	static const FText GeomHelpStr = LOCTEXT("GeomEditHelp", "Shift + click to insert a vertex.\nSelect one or more vertices and press Delete to remove them.\nDouble click a vertex to select a polygon\n");
-	static const FText GeomAddPolygonHelpStr = LOCTEXT("GeomClickAddPolygon", "Click to start creating a polygon\nCtrl + Click to start creating a subtractive polygon\n");
-	static const FText GeomAddCollisionPolygonHelpStr = LOCTEXT("GeomClickAddCollisionPolygon", "Click to start creating a polygon");
-	static const FText GeomAddVerticesHelpStr = LOCTEXT("GeomClickAddVertices", "Click to add points to the polygon\nEnter to finish adding a polygon\nEscape to cancel this polygon\n");
 	static const FText SourceRegionHelpStr = LOCTEXT("SourceRegionHelp", "Drag handles to adjust source region\nDouble-click on an image region to select all connected pixels\nHold down Ctrl and drag a rectangle to create a new sprite at that position\nDouble-click to create a new sprite under the cursor\nClick on other sprite rectangles to change the active sprite");
 
 	switch (CurrentMode)
@@ -744,26 +583,8 @@ void FSpriteEditorViewportClient::DrawCanvas(FViewport& Viewport, FSceneView& Vi
 		break;
 	case ESpriteEditorMode::EditCollisionMode:
 		{
- 			// Display tool help
-			{
-				const FText* HelpStr;
-				if (IsAddingPolygon())
-				{
-					HelpStr = (AddingPolygonIndex == -1) ? &GeomAddCollisionPolygonHelpStr : &GeomAddVerticesHelpStr;
-				}
-				else
-				{
-					HelpStr = &GeomHelpStr;
-				}
-
-				FCanvasTextItem TextItem(FVector2D(6, YPos), *HelpStr, GEngine->GetSmallFont(), FLinearColor::White);
-				TextItem.EnableShadow(FLinearColor::Black);
-				TextItem.Draw(&Canvas);
-				YPos += 54;
-			}
-
 			// Draw the custom collision geometry
-			DrawGeometry_CanvasPass(Viewport, View, Canvas, Sprite->CollisionGeometry, SpriteEditingConstants::CollisionShapeColor, FLinearColor::White, false);
+			SpriteGeometryHelper.DrawGeometry_CanvasPass(Viewport, View, Canvas, /*inout*/ YPos, SpriteEditingConstants::CollisionShapeColor, FLinearColor::White);
 			if (Sprite->BodySetup != nullptr)
 			{
 				DrawGeometryStats(Viewport, View, Canvas, Sprite->CollisionGeometry, false, /*inout*/ YPos);
@@ -779,26 +600,8 @@ void FSpriteEditorViewportClient::DrawCanvas(FViewport& Viewport, FSceneView& Vi
 		break;
 	case ESpriteEditorMode::EditRenderingGeomMode:
 		{
- 			// Display tool help
-			{
-				const FText* HelpStr;
-				if (IsAddingPolygon())
-				{
-					HelpStr = (AddingPolygonIndex == -1) ? &GeomAddPolygonHelpStr : &GeomAddVerticesHelpStr;
-				}
-				else
-				{
-					HelpStr = &GeomHelpStr;
-				}
-
-				FCanvasTextItem TextItem(FVector2D(6, YPos), *HelpStr, GEngine->GetSmallFont(), FLinearColor::White);
-				TextItem.EnableShadow(FLinearColor::Black);
-				TextItem.Draw(&Canvas);
-				YPos += 54;
-			}
-
 			// Draw the custom render geometry
-			DrawGeometry_CanvasPass(Viewport, View, Canvas, Sprite->RenderGeometry, SpriteEditingConstants::RenderShapeColor, SpriteEditingConstants::SubtractiveRenderShapeColor, true);
+			SpriteGeometryHelper.DrawGeometry_CanvasPass(Viewport, View, Canvas, /*inout*/ YPos, SpriteEditingConstants::RenderShapeColor, SpriteEditingConstants::SubtractiveRenderShapeColor);
 			DrawGeometryStats(Viewport, View, Canvas, Sprite->RenderGeometry, true, /*inout*/ YPos);
 			DrawRenderStats(Viewport, View, Canvas, Sprite, /*inout*/ YPos);
 
@@ -857,19 +660,17 @@ void FSpriteEditorViewportClient::Draw(const FSceneView* View, FPrimitiveDrawInt
 		FSocketEditingHelper::DrawSockets(RenderSpriteComponent, View, PDI);
 	}
 
-	UPaperSprite* Sprite = GetSpriteBeingEdited();
 	switch (CurrentMode)
 	{
-	case ESpriteEditorMode::ViewMode:
-	default:
-		break;
 	case ESpriteEditorMode::EditCollisionMode:
 		// Draw the custom collision geometry
-		SpriteGeometryHelper.DrawGeometry(*View, *PDI, Sprite->CollisionGeometry, SpriteEditingConstants::CollisionShapeColor, FLinearColor::White, /*bIsRenderGeometry=*/ false);
+		SpriteGeometryHelper.DrawGeometry(*View, *PDI, SpriteEditingConstants::CollisionShapeColor, FLinearColor::White);
 		break;
 	case ESpriteEditorMode::EditRenderingGeomMode:
 		// Draw the custom render geometry
-		SpriteGeometryHelper.DrawGeometry(*View, *PDI, Sprite->RenderGeometry, SpriteEditingConstants::RenderShapeColor, SpriteEditingConstants::SubtractiveRenderShapeColor, /*bIsRenderGeometry=*/ true);
+		SpriteGeometryHelper.DrawGeometry(*View, *PDI, SpriteEditingConstants::RenderShapeColor, SpriteEditingConstants::SubtractiveRenderShapeColor);
+		break;
+	default:
 		break;
 	}
 }
@@ -881,7 +682,7 @@ void FSpriteEditorViewportClient::Draw(FViewport* Viewport, FCanvas* Canvas)
 
 void FSpriteEditorViewportClient::Tick(float DeltaSeconds)
 {
-	if (UPaperSprite* Sprite = RenderSpriteComponent->GetSprite())
+	if (UPaperSprite* Sprite = GetSpriteBeingEdited())
 	{
 		//@TODO: Doesn't need to happen every frame, only when properties are updated
 
@@ -998,11 +799,6 @@ void FSpriteEditorViewportClient::UpdateRelatedSpritesList()
 	}
 }
 
-void FSpriteEditorViewportClient::UpdateMouseDelta()
-{
-	FPaperEditorViewportClient::UpdateMouseDelta();
-}
-
 void FSpriteEditorViewportClient::ProcessClick(FSceneView& View, HHitProxy* HitProxy, FKey Key, EInputEvent Event, uint32 HitX, uint32 HitY)
 {
 	const FViewportClick Click(&View, this, Key, Event, HitX, HitY);
@@ -1011,7 +807,7 @@ void FSpriteEditorViewportClient::ProcessClick(FSceneView& View, HHitProxy* HitP
 	const bool bIsAltKeyDown = Viewport->KeyState(EKeys::LeftAlt) || Viewport->KeyState(EKeys::RightAlt);
 	bool bHandled = false;
 
-	const bool bAllowSelectVertex = !(IsEditingGeometry() && IsAddingPolygon()) && !bIsShiftKeyDown;
+	const bool bAllowSelectVertex = !(IsEditingGeometry() && SpriteGeometryHelper.IsAddingPolygon()) && !bIsShiftKeyDown;
 
 	const bool bClearSelectionModifier = bIsCtrlKeyDown;
 	const bool bDeleteClickedVertex = bIsAltKeyDown;
@@ -1022,7 +818,7 @@ void FSpriteEditorViewportClient::ProcessClick(FSceneView& View, HHitProxy* HitP
 	{
 		if (!bClearSelectionModifier)
 		{
-			ClearSelectionSet();
+			SpriteGeometryHelper.ClearSelectionSet();
 		}
 
 		if (bDeleteClickedVertex)
@@ -1030,16 +826,14 @@ void FSpriteEditorViewportClient::ProcessClick(FSceneView& View, HHitProxy* HitP
 			// Delete selection
 			if (const FSpriteSelectedVertex* SelectedVertex = SelectedItemProxy->Data->CastTo<const FSpriteSelectedVertex>(FSelectionTypes::Vertex))
 			{
-				ClearSelectionSet();
-				SelectedItemSet.Add(SelectedItemProxy->Data);
-				SelectedIDSet.Add(FShapeVertexPair(SelectedVertex->ShapeIndex, SelectedVertex->VertexIndex));
+				SpriteGeometryHelper.ClearSelectionSet();
+				SpriteGeometryHelper.AddPolygonVertexToSelection(SelectedVertex->ShapeIndex, SelectedVertex->VertexIndex);
 				DeleteSelection();
 			}
 			else if (const FSpriteSelectedShape* SelectedShape = SelectedItemProxy->Data->CastTo<const FSpriteSelectedShape>(FSelectionTypes::GeometryShape))
 			{
-				ClearSelectionSet();
-				SelectedItemSet.Add(SelectedItemProxy->Data);
-				SelectedIDSet.Add(FShapeVertexPair(SelectedShape->ShapeIndex, INDEX_NONE));
+				SpriteGeometryHelper.ClearSelectionSet();
+				SpriteGeometryHelper.AddShapeToSelection(SelectedShape->ShapeIndex);
 				DeleteSelection();
 			}
 		}
@@ -1048,10 +842,8 @@ void FSpriteEditorViewportClient::ProcessClick(FSceneView& View, HHitProxy* HitP
 			// Double click to select a polygon
 			if (const FSpriteSelectedVertex* SelectedVertex = SelectedItemProxy->Data->CastTo<const FSpriteSelectedVertex>(FSelectionTypes::Vertex))
 			{
-				if (SelectedVertex->IsValidInEditor(GetSpriteBeingEdited(), IsInRenderingEditMode()))
-				{
-					SelectShape(SelectedVertex->ShapeIndex);
-				}
+				SpriteGeometryHelper.ClearSelectionSet();
+				SpriteGeometryHelper.AddShapeToSelection(SelectedVertex->ShapeIndex);
 			}
 		}
 		else
@@ -1060,19 +852,19 @@ void FSpriteEditorViewportClient::ProcessClick(FSceneView& View, HHitProxy* HitP
 			if (const FSpriteSelectedEdge* SelectedEdge = SelectedItemProxy->Data->CastTo<const FSpriteSelectedEdge>(FSelectionTypes::Edge))
 			{
 				// Add the next vertex defined by this edge
-				AddPolygonEdgeToSelection(SelectedEdge->ShapeIndex, SelectedEdge->VertexIndex);
+				SpriteGeometryHelper.AddPolygonEdgeToSelection(SelectedEdge->ShapeIndex, SelectedEdge->VertexIndex);
 			}
 			else if (const FSpriteSelectedVertex* SelectedVertex = SelectedItemProxy->Data->CastTo<const FSpriteSelectedVertex>(FSelectionTypes::Vertex))
 			{
-				AddPolygonVertexToSelection(SelectedVertex->ShapeIndex, SelectedVertex->VertexIndex);
+				SpriteGeometryHelper.AddPolygonVertexToSelection(SelectedVertex->ShapeIndex, SelectedVertex->VertexIndex);
 			}
 			else if (const FSpriteSelectedShape* SelectedShape = SelectedItemProxy->Data->CastTo<const FSpriteSelectedShape>(FSelectionTypes::GeometryShape))
 			{
-				SelectShape(SelectedShape->ShapeIndex);
+				SpriteGeometryHelper.AddShapeToSelection(SelectedShape->ShapeIndex);
 			}
 			else
 			{
-				SelectedItemSet.Add(SelectedItemProxy->Data);
+				SpriteGeometryHelper.SelectItem(SelectedItemProxy->Data);
 			}
 		}
 
@@ -1094,7 +886,7 @@ void FSpriteEditorViewportClient::ProcessClick(FSceneView& View, HHitProxy* HitP
 	{
 		if (IsInSourceRegionEditMode())
 		{
-			ClearSelectionSet();
+			SpriteGeometryHelper.ClearSelectionSet();
 
 			if (Event == EInputEvent::IE_DoubleClick)
 			{
@@ -1141,32 +933,32 @@ void FSpriteEditorViewportClient::ProcessClick(FSceneView& View, HHitProxy* HitP
 				}
 			}
 		}
-		else if (IsEditingGeometry() && !IsAddingPolygon())
+		else if (IsEditingGeometry() && !SpriteGeometryHelper.IsAddingPolygon())
 		{
+			FSpriteGeometryCollection* Geometry = GetGeometryBeingEdited();
+
 			if (bInsertVertexModifier)
 			{
 				FVector4 WorldPoint = View.PixelToWorld(HitX, HitY, 0);
-				UPaperSprite* Sprite = RenderSpriteComponent->GetSprite();
-				if (Sprite != nullptr)
+				if (UPaperSprite* Sprite = GetSpriteBeingEdited())
 				{
 					FVector2D SpriteSpaceClickPoint = Sprite->ConvertWorldSpaceToTextureSpace(WorldPoint);
 
 					// find a polygon to add vert to
-					if (SelectedItemSet.Num() > 0)
+					bool bFoundShapeToAddTo = false;
+					for (TSharedPtr<FSelectedItem> SelectedItemPtr : SpriteGeometryHelper.GetSelectionSet())
 					{
-						TSharedPtr<FSelectedItem> SelectedItemPtr = *SelectedItemSet.CreateIterator();
-
 						if (const FSpriteSelectedVertex* SelectedVertex = SelectedItemPtr->CastTo<const FSpriteSelectedVertex>(FSelectionTypes::Vertex)) //@TODO:  Inflexible?
 						{
-							if (SelectedVertex->IsValidInEditor(Sprite, IsInRenderingEditMode()))
-							{
-								AddPointToGeometry(SpriteSpaceClickPoint, SelectedVertex->ShapeIndex);
-							}
+							SpriteGeometryHelper.AddPointToGeometry(SpriteSpaceClickPoint, SelectedVertex->ShapeIndex);
+							bFoundShapeToAddTo = true;
+							break;
 						}
 					}
-					else
+
+					if (!bFoundShapeToAddTo)
 					{
-						AddPointToGeometry(SpriteSpaceClickPoint);
+						SpriteGeometryHelper.AddPointToGeometry(SpriteSpaceClickPoint);
 					}
 					bHandled = true;
 				}
@@ -1175,7 +967,7 @@ void FSpriteEditorViewportClient::ProcessClick(FSceneView& View, HHitProxy* HitP
 		else if (!IsEditingGeometry())
 		{
 			// Clicked on the background (missed any proxies), deselect the socket or whatever was selected
-			ClearSelectionSet();
+			SpriteGeometryHelper.ClearSelectionSet();
 		}
 
 		if (!bHandled)
@@ -1281,12 +1073,12 @@ bool FSpriteEditorViewportClient::InputKey(FViewport* Viewport, int32 Controller
 	}
 	else if (IsEditingGeometry())
 	{
-		if (IsAddingPolygon())
+		if (SpriteGeometryHelper.IsAddingPolygon())
 		{
 			if (Key == EKeys::LeftMouseButton && Event == IE_Pressed)
 			{
-				int32 HitX = Viewport->GetMouseX();
-				int32 HitY = Viewport->GetMouseY();
+				const int32 HitX = Viewport->GetMouseX();
+				const int32 HitY = Viewport->GetMouseY();
 
 				// Calculate world space positions
 				FSceneViewFamilyContext ViewFamily(FSceneViewFamily::ConstructionValues(Viewport, GetScene(), EngineShowFlags));
@@ -1296,58 +1088,23 @@ bool FSpriteEditorViewportClient::InputKey(FViewport* Viewport, int32 Controller
 				UPaperSprite* Sprite = GetSpriteBeingEdited();
 				FVector2D TexturePoint = Sprite->ConvertWorldSpaceToTextureSpace(WorldPoint);
 
-				BeginTransaction(LOCTEXT("AddPolygonVertexTransaction", "Add Vertex to Polygon"));
-
-				FSpriteGeometryCollection* Geometry = GetGeometryBeingEdited();
-				if (AddingPolygonIndex == -1)
-				{
-					FSpriteGeometryShape NewPolygon;
-					
-					const bool bIsCtrlKeyDown = Viewport->KeyState(EKeys::LeftControl) || Viewport->KeyState(EKeys::RightControl);
-					NewPolygon.bNegativeWinding = CanAddSubtractivePolygon() && bIsCtrlKeyDown;
-
-					Geometry->Shapes.Add(NewPolygon);
-					AddingPolygonIndex = Geometry->Shapes.Num() - 1;
-				}
-
-				if (Geometry->Shapes.IsValidIndex(AddingPolygonIndex))
-				{
-					FSpriteGeometryShape& Shape = Geometry->Shapes[AddingPolygonIndex];
-					Shape.ShapeType = ESpriteShapeType::Polygon;
-					Shape.Vertices.Add(Shape.ConvertTextureSpaceToShapeSpace(TexturePoint));
-				}
-				else
-				{
-					ResetAddPolygonMode();
-				}
-
-				Geometry->GeometryType = ESpritePolygonMode::FullyCustom;
-
-				bManipulationDirtiedSomething = true;
-				EndTransaction();
+				const bool bMakeSubtractiveIfAllowed = Viewport->KeyState(EKeys::LeftControl) || Viewport->KeyState(EKeys::RightControl);
+				SpriteGeometryHelper.TEMP_HandleAddPolygonClick(TexturePoint, bMakeSubtractiveIfAllowed);
 			}
 			else if (Key == EKeys::Enter)
 			{
-				ResetAddPolygonMode();
+				SpriteGeometryHelper.ResetAddPolygonMode();
 			}
-			else if ((Key == EKeys::Escape) && (AddingPolygonIndex != -1))
+			else if (Key == EKeys::Escape)
 			{
-				FSpriteGeometryCollection* Geometry = GetGeometryBeingEdited();
-				if (Geometry->Shapes.IsValidIndex(AddingPolygonIndex))
-				{
-					BeginTransaction(LOCTEXT("DeletePolygon", "Delete Polygon"));
-					Geometry->Shapes.RemoveAt(AddingPolygonIndex);
-					bManipulationDirtiedSomething = true;
-					EndTransaction();
-				}
-				ResetAddPolygonMode();
+				SpriteGeometryHelper.AbandonAddPolygonMode();
 			}
 		}
 		else
 		{
 			if (ProcessMarquee(Viewport, ControllerId, Key, Event, true))
 			{
-				bool bAddingToSelection = InputState.IsShiftButtonPressed(); //@TODO: control button moves widget? Hopefully make this more consistent when that is changed
+				const bool bAddingToSelection = InputState.IsShiftButtonPressed(); //@TODO: control button moves widget? Hopefully make this more consistent when that is changed
 				SelectVerticesInMarquee(bAddingToSelection);
 			}
 		}
@@ -1376,19 +1133,16 @@ bool FSpriteEditorViewportClient::InputWidgetDelta(FViewport* Viewport, EAxisLis
 		const FVector2D Drag2D(FVector::DotProduct(Drag, PaperAxisX), -FVector::DotProduct(Drag, PaperAxisY));
 
 		// Apply the delta to all of the selected objects
-		for (auto SelectionIt = SelectedItemSet.CreateConstIterator(); SelectionIt; ++SelectionIt)
+		for (TSharedPtr<FSelectedItem> SelectedItem : SpriteGeometryHelper.GetSelectionSet())
 		{
-			TSharedPtr<FSelectedItem> SelectedItem = *SelectionIt;
 			SelectedItem->ApplyDelta(Drag2D, Rot, Scale, MoveMode);
 		}
 
-		if (SelectedItemSet.Num() > 0)
+		if (SpriteGeometryHelper.HasAnySelectedItems())
 		{
 			if (!IsInSourceRegionEditMode())
 			{
 				UPaperSprite* Sprite = GetSpriteBeingEdited();
-				// Sprite->RebuildCollisionData();
-				// Sprite->RebuildRenderData();
 				Sprite->PostEditChange();
 				Invalidate();
 			}
@@ -1421,7 +1175,7 @@ void FSpriteEditorViewportClient::TrackingStopped()
 
 FWidget::EWidgetMode FSpriteEditorViewportClient::GetWidgetMode() const
 {
-	return (SelectedItemSet.Num() > 0) ? DesiredWidgetMode : FWidget::WM_None;
+	return SpriteGeometryHelper.HasAnySelectedItems() ? DesiredWidgetMode : FWidget::WM_None;
 }
 
 void FSpriteEditorViewportClient::SetWidgetMode(FWidget::EWidgetMode NewMode)
@@ -1438,28 +1192,26 @@ bool FSpriteEditorViewportClient::CanCycleWidgetMode() const
 {
 	if (!Widget->IsDragging())
 	{
-		return (SelectedItemSet.Num() > 0);
+		return SpriteGeometryHelper.HasAnySelectedItems();
 	}
 	return false;
 }
 
 FVector FSpriteEditorViewportClient::GetWidgetLocation() const
 {
-	if (SelectedItemSet.Num() > 0)
-	{
-		UPaperSprite* Sprite = GetSpriteBeingEdited();
+	FVector SummedPos(ForceInitToZero);
 
+	if (SpriteGeometryHelper.HasAnySelectedItems())
+	{
 		// Find the center of the selection set
-		FVector SummedPos(ForceInitToZero);
-		for (auto SelectionIt = SelectedItemSet.CreateConstIterator(); SelectionIt; ++SelectionIt)
+		for (TSharedPtr<FSelectedItem> SelectedItem : SpriteGeometryHelper.GetSelectionSet())
 		{
-			TSharedPtr<FSelectedItem> SelectedItem = *SelectionIt;
 			SummedPos += SelectedItem->GetWorldPos();
 		}
-		return (SummedPos / SelectedItemSet.Num());
+		return (SummedPos / SpriteGeometryHelper.GetSelectionSet().Num());
 	}
 
-	return FVector::ZeroVector;
+	return SummedPos;
 }
 
 FMatrix FSpriteEditorViewportClient::GetWidgetCoordSystem() const
@@ -1485,13 +1237,13 @@ FVector2D FSpriteEditorViewportClient::SelectedItemConvertWorldSpaceDeltaToLocal
 
 FVector FSpriteEditorViewportClient::TextureSpaceToWorldSpace(const FVector2D& SourcePoint) const
 {
-	UPaperSprite* Sprite = RenderSpriteComponent->GetSprite();
+	UPaperSprite* Sprite = GetSpriteBeingEdited();
 	return Sprite->ConvertTextureSpaceToWorldSpace(SourcePoint);
 }
 
 bool FSpriteEditorViewportClient::SelectedItemIsSelected(const struct FShapeVertexPair& Item) const
 {
-	return SelectedIDSet.Contains(Item);
+	return SpriteGeometryHelper.IsGeometrySelected(Item);
 }
 
 float FSpriteEditorViewportClient::SelectedItemGetUnitsPerPixel() const
@@ -1511,11 +1263,16 @@ void FSpriteEditorViewportClient::BeginTransaction(const FText& SessionName)
 	}
 }
 
+void FSpriteEditorViewportClient::MarkTransactionAsDirty()
+{
+	bManipulationDirtiedSomething = true;
+}
+
 void FSpriteEditorViewportClient::EndTransaction()
 {
 	if (bManipulationDirtiedSomething)
 	{
-		UPaperSprite* Sprite = RenderSpriteComponent->GetSprite();
+		UPaperSprite* Sprite = GetSpriteBeingEdited();
 
 		if (IsInSourceRegionEditMode())
 		{
@@ -1538,12 +1295,16 @@ void FSpriteEditorViewportClient::EndTransaction()
 	}
 }
 
+void FSpriteEditorViewportClient::InvalidateViewportAndHitProxies()
+{
+	Viewport->Invalidate();
+}
+
 void FSpriteEditorViewportClient::NotifySpriteBeingEditedHasChanged()
 {
 	//@TODO: Ideally we do this before switching
 	EndTransaction();
-	ClearSelectionSet();
-	ResetAddPolygonMode();
+	SpriteGeometryHelper.ClearSelectionSet();
 
 	// Update components to know about the new sprite being edited
 	UPaperSprite* Sprite = GetSpriteBeingEdited();
@@ -1551,12 +1312,9 @@ void FSpriteEditorViewportClient::NotifySpriteBeingEditedHasChanged()
 	RenderSpriteComponent->SetSprite(Sprite);
 	UpdateSourceTextureSpriteFromSprite(Sprite);
 
-	if (IsInSourceRegionEditMode()) 
-	{
-		UpdateRelatedSpritesList();
-	}
+	InternalActivateNewMode(CurrentMode);
 
-	//
+	//@TODO: Only do this if the sprite isn't visible (may consider doing a flashing pulse around the source region rect?)
 	bDeferZoomToSprite = true;
 }
 
@@ -1565,193 +1323,37 @@ void FSpriteEditorViewportClient::FocusOnSprite()
 	FocusViewportOnBox(RenderSpriteComponent->Bounds.GetBox());
 }
 
-// Be sure to call this with polygonIndex and VertexIndex in descending order
-// Returns true if the shape went to zero points and should be deleted itself
-static bool DeleteVertexInPolygonInternal(UPaperSprite* Sprite, FSpriteGeometryCollection* Geometry, const int32 ShapeIndex, const int32 VertexIndex)
-{
-	if (Geometry->Shapes.IsValidIndex(ShapeIndex))
-	{
-		FSpriteGeometryShape& Shape = Geometry->Shapes[ShapeIndex];
-		if (Shape.Vertices.IsValidIndex(VertexIndex))
-		{
-			Shape.ShapeType = ESpriteShapeType::Polygon;
-			Shape.Vertices.RemoveAt(VertexIndex);
-			
-			// Delete polygon
-			if (Shape.Vertices.Num() == 0)
-			{
-				return true;
-			}
-		}
-	}
-
-	return false;
-}
-
 void FSpriteEditorViewportClient::DeleteSelection()
 {
-	// Make a set of polygon-vertices "pairs" to delete
-	TSet<FShapeVertexPair> CompositeIndicesSet;
-	TSet<int32> ShapesToDeleteSet;
-
-	UPaperSprite* Sprite = GetSpriteBeingEdited();
-	FSpriteGeometryCollection* Geometry = GetGeometryBeingEdited();
-	for (TSharedPtr<FSelectedItem> SelectionIt : SelectedItemSet)
-	{
-		if (const FSpriteSelectedVertex* SelectedVertex = SelectionIt->CastTo<const FSpriteSelectedVertex>(FSelectionTypes::Vertex))
-		{
-			const bool bValidSelectedVertex = SelectedVertex->IsValidInEditor(Sprite, IsInRenderingEditMode());
-			if (bValidSelectedVertex)
-			{
-				CompositeIndicesSet.Add(FShapeVertexPair(SelectedVertex->ShapeIndex, SelectedVertex->VertexIndex));
-	
-				if (SelectedVertex->IsA(FSelectionTypes::Edge)) // add the "next" point for the edge
-				{
-					const int32 NextIndex = (SelectedVertex->VertexIndex + 1) % Geometry->Shapes[SelectedVertex->ShapeIndex].Vertices.Num();
-					CompositeIndicesSet.Add(FShapeVertexPair(SelectedVertex->ShapeIndex, NextIndex));
-				}
-			}
-		}
-		else if (const FSpriteSelectedShape* SelectedShape = SelectionIt->CastTo<const FSpriteSelectedShape>(FSelectionTypes::GeometryShape))
-		{
-			ShapesToDeleteSet.Add(SelectedShape->ShapeIndex);
-		}
-	}
-
-	if ((CompositeIndicesSet.Num() > 0) || (ShapesToDeleteSet.Num() > 0))
-	{
-		BeginTransaction(LOCTEXT("DeleteSelectionTransaction", "Delete Selection"));
-
-		// Delete the selected vertices first, as they may cause entire shapes to need to be deleted (sort so we delete from the back first)
-		TArray<FShapeVertexPair> CompositeIndices = CompositeIndicesSet.Array();
-		CompositeIndices.Sort([](const FShapeVertexPair& A, const FShapeVertexPair& B) { return (A.VertexIndex > B.VertexIndex); });
-		for (const FShapeVertexPair& Composite : CompositeIndices)
-		{
-			const int32 ShapeIndex = Composite.ShapeIndex;
-			const int32 VertexIndex = Composite.VertexIndex;
-			UE_LOG(LogPaper2DEditor, Verbose, TEXT("Deleting vertex %d from shape %d"), VertexIndex, ShapeIndex);
-			if (DeleteVertexInPolygonInternal(Sprite, Geometry, ShapeIndex, VertexIndex))
-			{
-				ShapesToDeleteSet.Add(ShapeIndex);
-			}
-		}
-
-		// Delete the selected shapes (plus any shapes that became empty due to selected vertices)
-		if (ShapesToDeleteSet.Num() > 0)
-		{
-			TArray<int32> ShapesToDeleteIndicies = ShapesToDeleteSet.Array();
-			ShapesToDeleteIndicies.Sort([](const int32& A, const int32& B) { return (A > B); });
-			for (const int32 ShapeToDeleteIndex : ShapesToDeleteIndicies)
-			{
-				UE_LOG(LogPaper2DEditor, Verbose, TEXT("Deleting shape %d"), ShapeToDeleteIndex);
-				Geometry->Shapes.RemoveAt(ShapeToDeleteIndex);
-			}
-		}
-
-		Geometry->GeometryType = ESpritePolygonMode::FullyCustom;
-		bManipulationDirtiedSomething = true;
-
-		ClearSelectionSet();
-		EndTransaction();
-		ResetAddPolygonMode();
-	}
-}
-
-void FSpriteEditorViewportClient::SplitEdge()
-{
-	BeginTransaction(LOCTEXT("SplitEdgeTransaction", "Split Edge"));
-
-	//@TODO: Really need a proper data structure!!!
-	if (SelectedItemSet.Num() == 1)
-	{
-		// Split this item
-		TSharedPtr<FSelectedItem> SelectedItem = *SelectedItemSet.CreateIterator();
-		SelectedItem->SplitEdge();
-		ClearSelectionSet();
-		bManipulationDirtiedSomething = true;
-	}
-	else
-	{
-		// Can't handle this without a better DS, as the indices will get screwed up
-	}
-
-	EndTransaction();
-}
-
-void FSpriteEditorViewportClient::ToggleAddPolygonMode()
-{
-	if (IsAddingPolygon())
-	{
-		ResetAddPolygonMode();
-	}
-	else
-	{
-		ClearSelectionSet();
-		
-		bIsAddingPolygon = true; 
-		AddingPolygonIndex = -1; 
-		Invalidate();
-	}
+	SpriteGeometryHelper.DeleteSelectedItems();
 }
 
 void FSpriteEditorViewportClient::AddBoxShape()
 {
-	if (FSpriteGeometryCollection* Geometry = GetGeometryBeingEdited())
-	{
-		UPaperSprite* Sprite = GetSpriteBeingEdited();
-		check(Sprite);
-
-		BeginTransaction(LOCTEXT("AddBoxShapeTransaction", "Add Box Shape"));
-
-		const FVector2D BoxSize(Sprite->GetSourceSize());
-		const FVector2D BoxLocation(BoxSize * 0.5f);
-		Geometry->GeometryType = ESpritePolygonMode::FullyCustom;
-		Geometry->AddRectangleShape(BoxLocation, BoxSize);
-
-		// Select the new geometry
-		ClearSelectionSet();
-		SelectShape(Geometry->Shapes.Num() - 1);
+	//@TODO: These are in the wrong location if SourceUV is not zero
+	UPaperSprite* Sprite = GetSpriteBeingEdited();
+	check(Sprite);
 
-		Sprite->PostEditChange();
-		EndTransaction();
-	}
-}
+	const FVector2D BoxSize(Sprite->GetSourceSize());
+	const FVector2D BoxLocation(BoxSize * 0.5f);
 
-bool FSpriteEditorViewportClient::CanAddBoxShape() const
-{
-	return IsEditingGeometry();
+	SpriteGeometryHelper.AddNewBoxShape(BoxLocation, BoxSize);
 }
 
 void FSpriteEditorViewportClient::AddCircleShape()
 {
+	//@TODO: These are in the wrong location if SourceUV is not zero
 	check(IsInCollisionEditMode());
-	if (FSpriteGeometryCollection* Geometry = GetGeometryBeingEdited())
-	{
-		UPaperSprite* Sprite = GetSpriteBeingEdited();
-		check(Sprite);
-
-		BeginTransaction(LOCTEXT("AddCircleShapeTransaction", "Add Circle Shape"));
-
-		const FVector2D SpriteBounds = Sprite->GetSourceSize();
-		const float SmallerBoundingAxisSize = SpriteBounds.GetMin();
-		const FVector2D EllipseSize(SmallerBoundingAxisSize, SmallerBoundingAxisSize);
-		const FVector2D EllipseLocation = SpriteBounds * 0.5f;
 
-		Geometry->GeometryType = ESpritePolygonMode::FullyCustom;
-		Geometry->AddCircleShape(EllipseLocation, EllipseSize);
-
-		// Select the new geometry
-		ClearSelectionSet();
-		SelectShape(Geometry->Shapes.Num() - 1);
+	UPaperSprite* Sprite = GetSpriteBeingEdited();
+	check(Sprite);
 
-		Sprite->PostEditChange();
-		EndTransaction();
-	}
-}
+	const FVector2D SpriteBounds = Sprite->GetSourceSize();
+	const float SmallerBoundingAxisSize = SpriteBounds.GetMin();
+	const float CircleRadius = SmallerBoundingAxisSize * 0.5f;
+	const FVector2D EllipseLocation = SpriteBounds * 0.5f;
 
-bool FSpriteEditorViewportClient::CanAddCircleShape() const
-{
-	return IsInCollisionEditMode();
+	SpriteGeometryHelper.AddNewCircleShape(EllipseLocation, CircleRadius);
 }
 
 void FSpriteEditorViewportClient::SnapAllVerticesToPixelGrid()
@@ -1798,9 +1400,34 @@ void FSpriteEditorViewportClient::InternalActivateNewMode(ESpriteEditorMode::Typ
 {
 	CurrentMode = NewMode;
 	Viewport->InvalidateHitProxy();
-	ClearSelectionSet();
 	ResetMarqueeTracking();
-	ResetAddPolygonMode();
+
+	UPaperSprite* Sprite = GetSpriteBeingEdited();
+
+	// Note: This has side effects (clearing the selection set, ensuring the geometry is correct if the sprite being edited changed, etc...).
+	// Do not skip even if the mode is not really changing.
+	SpriteGeometryHelper.SetGeometryBeingEdited(nullptr, /*bAllowCircles=*/ false, /*bAllowSubtractivePolygons=*/ false);
+
+	switch (CurrentMode)
+	{
+	case ESpriteEditorMode::ViewMode:
+		break;
+	case ESpriteEditorMode::EditSourceRegionMode:
+		UpdateRelatedSpritesList();
+		break;
+	case ESpriteEditorMode::EditCollisionMode:
+		if (Sprite != nullptr)
+		{
+			SpriteGeometryHelper.SetGeometryBeingEdited(&(Sprite->CollisionGeometry), /*bAllowCircles=*/ true, /*bAllowSubtractivePolygons=*/ false);
+		}
+		break;
+	case ESpriteEditorMode::EditRenderingGeomMode:
+		if (Sprite != nullptr)
+		{
+			SpriteGeometryHelper.SetGeometryBeingEdited(&(Sprite->RenderGeometry), /*bAllowCircles=*/ false, /*bAllowSubtractivePolygons=*/ true);
+		}
+		break;
+	}
 }
 
 FSpriteGeometryCollection* FSpriteEditorViewportClient::GetGeometryBeingEdited() const
@@ -1817,173 +1444,14 @@ FSpriteGeometryCollection* FSpriteEditorViewportClient::GetGeometryBeingEdited()
 	}
 }
 
-void FSpriteEditorViewportClient::SelectShape(const int32 ShapeIndex)
-{
-	Viewport->InvalidateHitProxy();
-
-	FSpriteGeometryCollection* Geometry = GetGeometryBeingEdited();
-	if (ensure((Geometry != nullptr) && Geometry->Shapes.IsValidIndex(ShapeIndex)))
-	{
-		if (!IsShapeSelected(ShapeIndex))
-		{
-			FSpriteGeometryShape& Shape = Geometry->Shapes[ShapeIndex];
-
-			TSharedPtr<FSpriteSelectedShape> SelectedShape = MakeShareable(new FSpriteSelectedShape(*this, *Geometry, ShapeIndex));
-
-			SelectedItemSet.Add(SelectedShape);
-			SelectedIDSet.Add(FShapeVertexPair(ShapeIndex, INDEX_NONE));
-		}
-	}
-}
-
 bool FSpriteEditorViewportClient::IsShapeSelected(const int32 ShapeIndex) const
 {
-	return SelectedIDSet.Contains(FShapeVertexPair(ShapeIndex, INDEX_NONE));
+	return SpriteGeometryHelper.IsGeometrySelected(FShapeVertexPair(ShapeIndex, INDEX_NONE));
 }
 
 bool FSpriteEditorViewportClient::IsPolygonVertexSelected(const int32 ShapeIndex, const int32 VertexIndex) const
 {
-	return SelectedIDSet.Contains(FShapeVertexPair(ShapeIndex, VertexIndex));
-}
-
-void FSpriteEditorViewportClient::AddPolygonVertexToSelection(const int32 ShapeIndex, const int32 VertexIndex)
-{
-	FSpriteGeometryCollection* Geometry = GetGeometryBeingEdited();
-	if (ensure((Geometry != nullptr) && Geometry->Shapes.IsValidIndex(ShapeIndex)))
-	{
-		FSpriteGeometryShape& Shape = Geometry->Shapes[ShapeIndex];
-		if (Shape.Vertices.IsValidIndex(VertexIndex))
-		{
-			if (!IsPolygonVertexSelected(ShapeIndex, VertexIndex))
-			{
-				TSharedPtr<FSpriteSelectedVertex> Vertex = MakeShareable(new FSpriteSelectedVertex());
-				Vertex->SpritePtr = GetSpriteBeingEdited();
-				Vertex->ShapeIndex = ShapeIndex;
-				Vertex->VertexIndex = VertexIndex;
-				Vertex->bRenderData = IsInRenderingEditMode();
-
-				SelectedItemSet.Add(Vertex);
-				SelectedIDSet.Add(FShapeVertexPair(ShapeIndex, VertexIndex));
-			}
-		}
-	}
-}
-
-void FSpriteEditorViewportClient::AddPolygonEdgeToSelection(const int32 ShapeIndex, const int32 FirstVertexIndex)
-{
-	FSpriteGeometryCollection* Geometry = GetGeometryBeingEdited();
-	if (ensure((Geometry != nullptr) && Geometry->Shapes.IsValidIndex(ShapeIndex)))
-	{
-		FSpriteGeometryShape& Shape = Geometry->Shapes[ShapeIndex];
-		const int32 NextVertexIndex = (FirstVertexIndex + 1) % Shape.Vertices.Num();
-
-		AddPolygonVertexToSelection(ShapeIndex, FirstVertexIndex);
-		AddPolygonVertexToSelection(ShapeIndex, NextVertexIndex);
-	}
-}
-
-void FSpriteEditorViewportClient::ClearSelectionSet()
-{
-	SelectedItemSet.Empty();
-	SelectedIDSet.Empty();
-}
-
-static bool ClosestPointOnLine(const FVector2D& Point, const FVector2D& LineStart, const FVector2D& LineEnd, FVector2D& OutClosestPoint)
-{
-	FVector2D DL = LineEnd - LineStart;
-	FVector2D DP = Point - LineStart;
-	float DotL = FVector2D::DotProduct(DP, DL);
-	float DotP = FVector2D::DotProduct(DL, DL);
-	if (DotP > 0.0001f)
-	{
-		float T = DotL / DotP;
-		if (T >= -0.0001f && T <= 1.0001f)
-		{
-			T = FMath::Clamp(T, 0.0f, 1.0f);
-			OutClosestPoint = LineStart + T * DL;
-			return true;
-		}
-	}
-
-	return false;
-}
-
-// Adds a point to the currently edited index
-// If SelectedPolygonIndex is set, only that polygon is considered
-void FSpriteEditorViewportClient::AddPointToGeometry(const FVector2D& TextureSpacePoint, const int32 SelectedPolygonIndex)
-{
-	FSpriteGeometryCollection* Geometry = GetGeometryBeingEdited();
-
-	int32 ClosestShapeIndex = INDEX_NONE;
-	int32 ClosestVertexInsertIndex = INDEX_NONE;
-	float ClosestDistanceSquared = MAX_FLT;
-
-	int32 StartPolygonIndex = 0;
-	int32 EndPolygonIndex = Geometry->Shapes.Num();
-	if ((SelectedPolygonIndex >= 0) && (SelectedPolygonIndex < Geometry->Shapes.Num()))
-	{
-		StartPolygonIndex = SelectedPolygonIndex;
-		EndPolygonIndex = SelectedPolygonIndex + 1;
-	}
-
-	// Determine where we should insert the vertex
-	for (int32 PolygonIndex = StartPolygonIndex; PolygonIndex < EndPolygonIndex; ++PolygonIndex)
-	{
-		FSpriteGeometryShape& Polygon = Geometry->Shapes[PolygonIndex];
-		if (Polygon.Vertices.Num() >= 3)
-		{
-			for (int32 VertexIndex = 0; VertexIndex < Polygon.Vertices.Num(); ++VertexIndex)
-			{
-				FVector2D ClosestPoint;
-				const FVector2D LineStart = Polygon.ConvertShapeSpaceToTextureSpace(Polygon.Vertices[VertexIndex]);
-				int32 NextVertexIndex = (VertexIndex + 1) % Polygon.Vertices.Num();
-				const FVector2D LineEnd = Polygon.ConvertShapeSpaceToTextureSpace(Polygon.Vertices[NextVertexIndex]);
-				if (ClosestPointOnLine(TextureSpacePoint, LineStart, LineEnd, /*out*/ ClosestPoint))
-				{
-					const float CurrentDistanceSquared = FVector2D::DistSquared(ClosestPoint, TextureSpacePoint);
-					if (CurrentDistanceSquared < ClosestDistanceSquared)
-					{
-						ClosestShapeIndex = PolygonIndex;
-						ClosestDistanceSquared = CurrentDistanceSquared;
-						ClosestVertexInsertIndex = NextVertexIndex;
-					}
-				}
-			}
-		}
-		else
-		{
-			// Simply insert the vertex
-			for (int32 VertexIndex = 0; VertexIndex < Polygon.Vertices.Num(); ++VertexIndex)
-			{
-				const FVector2D CurrentVertexTS = Polygon.ConvertShapeSpaceToTextureSpace(Polygon.Vertices[VertexIndex]);
-				const float CurrentDistanceSquared = FVector2D::DistSquared(CurrentVertexTS, TextureSpacePoint);
-				if (CurrentDistanceSquared < ClosestDistanceSquared)
-				{
-					ClosestShapeIndex = PolygonIndex;
-					ClosestDistanceSquared = CurrentDistanceSquared;
-					ClosestVertexInsertIndex = VertexIndex + 1;
-				}
-			}
-		}
-	}
-
-	if ((ClosestVertexInsertIndex != INDEX_NONE) && (ClosestShapeIndex != INDEX_NONE))
-	{
-		FSpriteGeometryShape& Shape = Geometry->Shapes[ClosestShapeIndex];
-		if (Shape.ShapeType != ESpriteShapeType::Circle)
-		{
-			BeginTransaction(LOCTEXT("AddPolygonVertexTransaction", "Add Vertex to Polygon"));
-			Shape.Vertices.Insert(Shape.ConvertTextureSpaceToShapeSpace(TextureSpacePoint), ClosestVertexInsertIndex);
-			Shape.ShapeType = ESpriteShapeType::Polygon;
-			Geometry->GeometryType = ESpritePolygonMode::FullyCustom;
-			bManipulationDirtiedSomething = true;
-			EndTransaction();
-
-			// Select this vertex
-			ClearSelectionSet();
-			AddPolygonVertexToSelection(ClosestShapeIndex, ClosestVertexInsertIndex);
-		}
-	}
+	return SpriteGeometryHelper.IsGeometrySelected(FShapeVertexPair(ShapeIndex, VertexIndex));
 }
 
 void FSpriteEditorViewportClient::ResetMarqueeTracking()
@@ -2040,7 +1508,7 @@ void FSpriteEditorViewportClient::SelectVerticesInMarquee(bool bAddToSelection)
 {
 	if (!bAddToSelection)
 	{
-		ClearSelectionSet();
+		SpriteGeometryHelper.ClearSelectionSet();
 	}
 
 	if (UPaperSprite* Sprite = GetSpriteBeingEdited())
@@ -2088,7 +1556,7 @@ void FSpriteEditorViewportClient::SelectVerticesInMarquee(bool bAddToSelection)
 				
 				if (bSelectWholeShape)
 				{
-					SelectShape(ShapeIndex);
+					SpriteGeometryHelper.AddShapeToSelection(ShapeIndex);
 				}
 				else
 				{
@@ -2098,7 +1566,7 @@ void FSpriteEditorViewportClient::SelectVerticesInMarquee(bool bAddToSelection)
 						const FVector2D TextureSpaceVertex = Shape.ConvertShapeSpaceToTextureSpace(Shape.Vertices[VertexIndex]);
 						if (QueryBounds.IsInside(TextureSpaceVertex))
 						{
-							AddPolygonVertexToSelection(ShapeIndex, VertexIndex);
+							SpriteGeometryHelper.AddPolygonVertexToSelection(ShapeIndex, VertexIndex);
 						}
 					}
 				}
diff --git a/Engine/Plugins/2D/Paper2D/Source/Paper2DEditor/Private/SpriteEditor/SpriteEditorViewportClient.h b/Engine/Plugins/2D/Paper2D/Source/Paper2DEditor/Private/SpriteEditor/SpriteEditorViewportClient.h
--- a/Engine/Plugins/2D/Paper2D/Source/Paper2DEditor/Private/SpriteEditor/SpriteEditorViewportClient.h
+++ b/Engine/Plugins/2D/Paper2D/Source/Paper2DEditor/Private/SpriteEditor/SpriteEditorViewportClient.h
@@ -54,7 +54,6 @@ public:
 	// End of FViewportClient interface
 
 	// FEditorViewportClient interface
-	virtual void UpdateMouseDelta() override;
 	virtual void ProcessClick(FSceneView& View, HHitProxy* HitProxy, FKey Key, EInputEvent Event, uint32 HitX, uint32 HitY) override;
 	virtual bool InputKey(FViewport* Viewport, int32 ControllerId, FKey Key, EInputEvent Event, float AmountDepressed, bool bGamepad) override;
 	virtual bool InputWidgetDelta(FViewport* Viewport, EAxisList::Type CurrentAxis, FVector& Drag, FRotator& Rot, FVector& Scale) override;
@@ -75,6 +74,10 @@ public:
 	virtual FVector TextureSpaceToWorldSpace(const FVector2D& SourcePoint) const override;
 	virtual bool SelectedItemIsSelected(const struct FShapeVertexPair& Item) const override;
 	virtual float SelectedItemGetUnitsPerPixel() const override;
+	virtual void BeginTransaction(const FText& SessionName) override;
+	virtual void MarkTransactionAsDirty() override;
+	virtual void EndTransaction() override;
+	virtual void InvalidateViewportAndHitProxies() override;
 	// End of ISpriteSelectionContext interface
 
 	// Process marquee selection, return true of a selection has been performed
@@ -93,7 +96,7 @@ public:
 	bool IsShowMeshEdgesChecked() const;
 
 	void EnterViewMode() { InternalActivateNewMode(ESpriteEditorMode::ViewMode); }
-	void EnterSourceRegionEditMode() { InternalActivateNewMode(ESpriteEditorMode::EditSourceRegionMode); UpdateRelatedSpritesList(); }
+	void EnterSourceRegionEditMode() { InternalActivateNewMode(ESpriteEditorMode::EditSourceRegionMode); }
 	void EnterCollisionEditMode() { InternalActivateNewMode(ESpriteEditorMode::EditCollisionMode); }
 	void EnterRenderingEditMode() { InternalActivateNewMode(ESpriteEditorMode::EditRenderingGeomMode); }
 
@@ -115,19 +118,8 @@ public:
 	void DeleteSelection();
 	bool CanDeleteSelection() const { return IsEditingGeometry(); } //@TODO: Need a selection
 
-	void SplitEdge();
-	bool CanSplitEdge() const { return IsEditingGeometry(); } //@TODO: Need an edge
-
 	void AddBoxShape();
-	bool CanAddBoxShape() const;
 	void AddCircleShape();
-	bool CanAddCircleShape() const;
-
-	void ResetAddPolygonMode() { bIsAddingPolygon = false; }
-	void ToggleAddPolygonMode();
-	bool IsAddingPolygon() const { return bIsAddingPolygon; }
-	bool CanAddPolygon() const { return IsEditingGeometry(); }
-	bool CanAddSubtractivePolygon() const { return CanAddPolygon() && IsInRenderingEditMode(); }
 
 	void SnapAllVerticesToPixelGrid();
 	bool CanSnapVerticesToPixelGrid() const { return IsEditingGeometry(); }
@@ -171,12 +163,6 @@ private:
 	// Pointer back to the sprite editor viewport control that owns us
 	TWeakPtr<class SSpriteEditorViewport> SpriteEditorViewportPtr;
 
-	// Set of selected objects
-	TSet< TSharedPtr<FSelectedItem> > SelectedItemSet;
-
-	// Set of selected vertices/shapes
-	TSet<FShapeVertexPair> SelectedIDSet;
-
 	// The current transaction for undo/redo
 	class FScopedTransaction* ScopedTransaction;
 
@@ -198,12 +184,6 @@ private:
 	// Should we show related sprites in the source texture?
 	bool bShowRelatedSprites;
 
-	// Is waiting to add geometry
-	bool bIsAddingPolygon;
-
-	// The polygon index being added to, -1 if we don't have a polygon yet
-	int32 AddingPolygonIndex;
-
 	// Marquee tracking
 	bool bIsMarqueeTracking;
 	FVector2D MarqueeStartPos, MarqueeEndPos;
@@ -213,20 +193,24 @@ private:
 
 	// Sprite geometry editing/rendering helper
 	FSpriteGeometryEditingHelper SpriteGeometryHelper;
+
+public:
+	FSpriteGeometryEditingHelper* GetGeometryEditHelper()
+	{
+		return &SpriteGeometryHelper;
+	}
+
 private:
 	UPaperSprite* GetSpriteBeingEdited() const
 	{
 		return SpriteEditorPtr.Pin()->GetSpriteBeingEdited();
 	}
 
-	FVector2D TextureSpaceToScreenSpace(const FSceneView& View, const FVector2D& SourcePoint) const;
-
 	// Position relative to source texture (ignoring rotation and other transformations applied to extract the sprite)
 	FVector2D SourceTextureSpaceToScreenSpace(const FSceneView& View, const FVector2D& SourcePoint) const;
 	FVector SourceTextureSpaceToWorldSpace(const FVector2D& SourcePoint) const;
 
 	void DrawBoundsAsText(FViewport& InViewport, FSceneView& View, FCanvas& Canvas, int32& YPos);
-	void DrawGeometry_CanvasPass(FViewport& InViewport, FSceneView& View, FCanvas& Canvas, FSpriteGeometryCollection& Geometry, const FLinearColor& GeometryVertexColor, const FLinearColor& NegativeGeometryVertexColor, bool bIsRenderGeometry);
 
 	static void DrawGeometryStats(FViewport& InViewport, FSceneView& View, FCanvas& Canvas, const FSpriteGeometryCollection& Geometry, bool bIsRenderGeometry, int32& YPos);
 	static void DrawCollisionStats(FViewport& InViewport, FSceneView& View, FCanvas& Canvas, class UBodySetup* BodySetup, int32& YPos);
@@ -235,31 +219,16 @@ private:
 	void DrawRelatedSprites(FViewport& InViewport, FSceneView& View, FCanvas& Canvas, const FLinearColor& GeometryVertexColor);
 	void DrawMarquee(FViewport& InViewport, FSceneView& View, FCanvas& Canvas, const FLinearColor& MarqueeColor);
 
-	void BeginTransaction(const FText& SessionName);
-	void EndTransaction();
-
 	void UpdateSourceTextureSpriteFromSprite(UPaperSprite* SourceSprite);
 	
 	// Selection handling
 
-	// Selects a shape in the geometry for currently selected mode
-	void SelectShape(const int32 ShapeIndex);
-
 	// Indicates if the specified shape is currently selected
 	bool IsShapeSelected(const int32 ShapeIndex) const;
 
-
-	void AddPolygonEdgeToSelection(const int32 ShapeIndex, const int32 FirstVertexIndex);
-
-	void AddPolygonVertexToSelection(const int32 ShapeIndex, const int32 VertexIndex);
-
 	// Indicates if the specified vertex is selected
 	bool IsPolygonVertexSelected(const int32 PolygonIndex, const int32 VertexIndex) const;
 
-	void AddPointToGeometry(const FVector2D& TextureSpacePoint, const int32 SelectedPolygonIndex = INDEX_NONE);
-	void ClearSelectionSet();
-
-
 	void ResetMarqueeTracking();
 	bool ConvertMarqueeToSourceTextureSpace(/*out*/FVector2D& OutStartPos, /*out*/FVector2D& OutDimension);
 	void SelectVerticesInMarquee(bool bAddToSelection);
