Sprite Editor: Fixed edge normals not being toggleable anymore and added initialization of some other variables

diff --git a/Engine/Plugins/2D/Paper2D/Source/Paper2DEditor/Private/PaperEditorShared/SpriteGeometryEditing.cpp b/Engine/Plugins/2D/Paper2D/Source/Paper2DEditor/Private/PaperEditorShared/SpriteGeometryEditing.cpp
--- a/Engine/Plugins/2D/Paper2D/Source/Paper2DEditor/Private/PaperEditorShared/SpriteGeometryEditing.cpp
+++ b/Engine/Plugins/2D/Paper2D/Source/Paper2DEditor/Private/PaperEditorShared/SpriteGeometryEditing.cpp
@@ -183,6 +183,11 @@ FVector FSpriteSelectedShape::GetWorldPos() const
 FSpriteGeometryEditingHelper::FSpriteGeometryEditingHelper(ISpriteSelectionContext& InEditorContext)
 	: EditorContext(InEditorContext)
 	, GeometryBeingEdited(nullptr)
+	, bIsAddingPolygon(false)
+	, AddingPolygonIndex(INDEX_NONE)
+	, bShowNormals(true)
+	, bAllowSubtractivePolygons(false)
+	, bAllowCircles(true)
 {
 	WidgetVertexColorMaterial = (UMaterial*)StaticLoadObject(UMaterial::StaticClass(), NULL, TEXT("/Engine/EditorMaterials/WidgetVertexColorMaterial.WidgetVertexColorMaterial"), NULL, LOAD_None, NULL);
 }
@@ -657,6 +662,17 @@ void FSpriteGeometryEditingHelper::AddPolygonEdgeToSelection(const int32 ShapeIn
 	}
 }
 
+void FSpriteGeometryEditingHelper::SetShowNormals(bool bShouldShowNormals)
+{
+	bShowNormals = bShouldShowNormals;
+	EditorContext.InvalidateViewportAndHitProxies();
+}
+
+void FSpriteGeometryEditingHelper::ToggleShowNormals()
+{
+	SetShowNormals(!bShowNormals);
+}
+
 void FSpriteGeometryEditingHelper::SetGeometryBeingEdited(FSpriteGeometryCollection* NewGeometryBeingEdited, bool bInAllowCircles, bool bInAllowSubtractivePolygons)
 {
 	ClearSelectionSet();
diff --git a/Engine/Plugins/2D/Paper2D/Source/Paper2DEditor/Private/PaperEditorShared/SpriteGeometryEditing.h b/Engine/Plugins/2D/Paper2D/Source/Paper2DEditor/Private/PaperEditorShared/SpriteGeometryEditing.h
--- a/Engine/Plugins/2D/Paper2D/Source/Paper2DEditor/Private/PaperEditorShared/SpriteGeometryEditing.h
+++ b/Engine/Plugins/2D/Paper2D/Source/Paper2DEditor/Private/PaperEditorShared/SpriteGeometryEditing.h
@@ -101,6 +101,10 @@ public:
 	void AddPolygonVertexToSelection(const int32 ShapeIndex, const int32 VertexIndex);
 	void AddPolygonEdgeToSelection(const int32 ShapeIndex, const int32 FirstVertexIndex);
 
+	void SetShowNormals(bool bShouldShowNormals);
+	void ToggleShowNormals();
+	bool IsShowNormalsEnabled() const { return bShowNormals; }
+
 	// Changes the geometry being edited (clears the selection set in the process)
 	void SetGeometryBeingEdited(FSpriteGeometryCollection* NewGeometryBeingEdited, bool bInAllowCircles, bool bInAllowSubtractivePolygons);
 
@@ -154,7 +158,7 @@ private:
 	// The polygon index being added to, INDEX_NONE if we don't have a polygon yet
 	int32 AddingPolygonIndex;
 
-	// Should we draw normal ticks?
+	// Should we show polygon edge normals?
 	bool bShowNormals;
 
 	// Do we allow subtractive polygons?
diff --git a/Engine/Plugins/2D/Paper2D/Source/Paper2DEditor/Private/SpriteEditor/SpriteEditor.cpp b/Engine/Plugins/2D/Paper2D/Source/Paper2DEditor/Private/SpriteEditor/SpriteEditor.cpp
--- a/Engine/Plugins/2D/Paper2D/Source/Paper2DEditor/Private/SpriteEditor/SpriteEditor.cpp
+++ b/Engine/Plugins/2D/Paper2D/Source/Paper2DEditor/Private/SpriteEditor/SpriteEditor.cpp
@@ -136,9 +136,9 @@ void SSpriteEditorViewport::BindCommands()
 
 	CommandList->MapAction(
 		Commands.SetShowNormals,
-		FExecuteAction::CreateSP( EditorViewportClientRef, &FSpriteEditorViewportClient::ToggleShowNormals ),
+		FExecuteAction::CreateRaw(GeometryEditingHelper, &FSpriteGeometryEditingHelper::ToggleShowNormals),
 		FCanExecuteAction(),
-		FIsActionChecked::CreateSP( EditorViewportClientRef, &FSpriteEditorViewportClient::IsShowNormalsChecked ) );
+		FIsActionChecked::CreateRaw(GeometryEditingHelper, &FSpriteGeometryEditingHelper::IsShowNormalsEnabled));
  
 	CommandList->MapAction(
 		Commands.SetShowPivot,
diff --git a/Engine/Plugins/2D/Paper2D/Source/Paper2DEditor/Private/SpriteEditor/SpriteEditorViewportClient.cpp b/Engine/Plugins/2D/Paper2D/Source/Paper2DEditor/Private/SpriteEditor/SpriteEditorViewportClient.cpp
--- a/Engine/Plugins/2D/Paper2D/Source/Paper2DEditor/Private/SpriteEditor/SpriteEditorViewportClient.cpp
+++ b/Engine/Plugins/2D/Paper2D/Source/Paper2DEditor/Private/SpriteEditor/SpriteEditorViewportClient.cpp
@@ -75,7 +75,6 @@ FSpriteEditorViewportClient::FSpriteEditorViewportClient(TWeakPtr<FSpriteEditor>
 
 	bShowSourceTexture = false;
 	bShowSockets = true;
-	bShowNormals = true;
 	bShowPivot = true;
 	bShowRelatedSprites = true;
 
diff --git a/Engine/Plugins/2D/Paper2D/Source/Paper2DEditor/Private/SpriteEditor/SpriteEditorViewportClient.h b/Engine/Plugins/2D/Paper2D/Source/Paper2DEditor/Private/SpriteEditor/SpriteEditorViewportClient.h
--- a/Engine/Plugins/2D/Paper2D/Source/Paper2DEditor/Private/SpriteEditor/SpriteEditorViewportClient.h
+++ b/Engine/Plugins/2D/Paper2D/Source/Paper2DEditor/Private/SpriteEditor/SpriteEditorViewportClient.h
@@ -89,9 +89,6 @@ public:
 	void ToggleShowPivot() { bShowPivot = !bShowPivot; Invalidate(); }
 	bool IsShowPivotChecked() const { return bShowPivot; }
 
-	void ToggleShowNormals() { bShowNormals = !bShowNormals; Invalidate(); }
-	bool IsShowNormalsChecked() const { return bShowNormals; }
-
 	void ToggleShowMeshEdges();
 	bool IsShowMeshEdgesChecked() const;
 
@@ -172,9 +169,6 @@ private:
 	// Should we show sockets?
 	bool bShowSockets;
 
-	// Should we show polygon normals?
-	bool bShowNormals;
-
 	// Should we show the sprite pivot?
 	bool bShowPivot;
 
