Sprite Editor: Added SourceUV to initial shape position for Add Circle and Add Box commands
[UE-13182]

diff --git a/Engine/Plugins/2D/Paper2D/Source/Paper2DEditor/Private/SpriteEditor/SpriteEditorViewportClient.cpp b/Engine/Plugins/2D/Paper2D/Source/Paper2DEditor/Private/SpriteEditor/SpriteEditorViewportClient.cpp
--- a/Engine/Plugins/2D/Paper2D/Source/Paper2DEditor/Private/SpriteEditor/SpriteEditorViewportClient.cpp
+++ b/Engine/Plugins/2D/Paper2D/Source/Paper2DEditor/Private/SpriteEditor/SpriteEditorViewportClient.cpp
@@ -1329,19 +1329,17 @@ void FSpriteEditorViewportClient::DeleteSelection()
 
 void FSpriteEditorViewportClient::AddBoxShape()
 {
-	//@TODO: These are in the wrong location if SourceUV is not zero
 	UPaperSprite* Sprite = GetSpriteBeingEdited();
 	check(Sprite);
 
 	const FVector2D BoxSize(Sprite->GetSourceSize());
-	const FVector2D BoxLocation(BoxSize * 0.5f);
+	const FVector2D BoxLocation(Sprite->GetSourceUV() + (BoxSize * 0.5f));
 
 	SpriteGeometryHelper.AddNewBoxShape(BoxLocation, BoxSize);
 }
 
 void FSpriteEditorViewportClient::AddCircleShape()
 {
-	//@TODO: These are in the wrong location if SourceUV is not zero
 	check(IsInCollisionEditMode());
 
 	UPaperSprite* Sprite = GetSpriteBeingEdited();
@@ -1350,7 +1348,7 @@ void FSpriteEditorViewportClient::AddCircleShape()
 	const FVector2D SpriteBounds = Sprite->GetSourceSize();
 	const float SmallerBoundingAxisSize = SpriteBounds.GetMin();
 	const float CircleRadius = SmallerBoundingAxisSize * 0.5f;
-	const FVector2D EllipseLocation = SpriteBounds * 0.5f;
+	const FVector2D EllipseLocation = Sprite->GetSourceUV() + (SpriteBounds * 0.5f);
 
 	SpriteGeometryHelper.AddNewCircleShape(EllipseLocation, CircleRadius);
 }
