Sprite Editor:
- Reordered Rendering and Collision categories to match the order of the mode buttons
- Added header summaries for Rendering and Collision categories
  - Rendering shows what kind of material is being used (Masked/Opaque/Translucent)
  - Collision shows the collision domain (none, 2d, 3d)
- Reduced confusion over which kind of geometry is being edited
  - Hid collision geometry properties when in Edit RenderGeom mode
  - Hid rendering geometry properties when in Edit Collision mode
[UE-13174]

diff --git a/Engine/Plugins/2D/Paper2D/Source/Paper2DEditor/Private/SpriteEditor/SpriteDetailsCustomization.cpp b/Engine/Plugins/2D/Paper2D/Source/Paper2DEditor/Private/SpriteEditor/SpriteDetailsCustomization.cpp
--- a/Engine/Plugins/2D/Paper2D/Source/Paper2DEditor/Private/SpriteEditor/SpriteDetailsCustomization.cpp
+++ b/Engine/Plugins/2D/Paper2D/Source/Paper2DEditor/Private/SpriteEditor/SpriteDetailsCustomization.cpp
@@ -17,7 +17,18 @@
 
 TSharedRef<IDetailCustomization> FSpriteDetailsCustomization::MakeInstance()
 {
-	return MakeShareable(new FSpriteDetailsCustomization);
+	TAttribute<ESpriteEditorMode::Type> DummyEditMode(ESpriteEditorMode::ViewMode);
+	return MakeInstanceForSpriteEditor(DummyEditMode);
+}
+
+TSharedRef<IDetailCustomization> FSpriteDetailsCustomization::MakeInstanceForSpriteEditor(TAttribute<ESpriteEditorMode::Type> InEditMode)
+{
+	return MakeShareable(new FSpriteDetailsCustomization(InEditMode));
+}
+
+FSpriteDetailsCustomization::FSpriteDetailsCustomization(TAttribute<ESpriteEditorMode::Type> InEditMode)
+{
+	SpriteEditMode = InEditMode;
 }
 
 FDetailWidgetRow& FSpriteDetailsCustomization::GenerateWarningRow(IDetailCategoryBuilder& WarningCategory, bool bExperimental, const FText& WarningText, const FText& Tooltip, const FString& ExcerptLink, const FString& ExcerptName)
@@ -60,13 +71,13 @@ void FSpriteDetailsCustomization::CustomizeDetails(IDetailLayoutBuilder& DetailL
 	IDetailCategoryBuilder& SpriteCategory = DetailLayout.EditCategory("Sprite", FText::GetEmpty(), ECategoryPriority::Important);
 	BuildSpriteSection(SpriteCategory, DetailLayout);
 
-	// Build the rendering category
-	IDetailCategoryBuilder& RenderingCategory = DetailLayout.EditCategory("Rendering");
-	BuildRenderingSection(RenderingCategory, DetailLayout);
-
 	// Build the collision category
 	IDetailCategoryBuilder& CollisionCategory = DetailLayout.EditCategory("Collision");
 	BuildCollisionSection(CollisionCategory, DetailLayout);
+
+	// Build the rendering category
+	IDetailCategoryBuilder& RenderingCategory = DetailLayout.EditCategory("Rendering");
+	BuildRenderingSection(RenderingCategory, DetailLayout);
 }
 
 void FSpriteDetailsCustomization::BuildSpriteSection(IDetailCategoryBuilder& SpriteCategory, IDetailLayoutBuilder& DetailLayout)
@@ -95,11 +106,57 @@ void FSpriteDetailsCustomization::BuildSpriteSection(IDetailCategoryBuilder& Spr
 
 void FSpriteDetailsCustomization::BuildRenderingSection(IDetailCategoryBuilder& RenderingCategory, IDetailLayoutBuilder& DetailLayout)
 {
+	TAttribute<EVisibility> HideWhenInCollisionMode = TAttribute<EVisibility>::Create(TAttribute<EVisibility>::FGetter::CreateSP(this, &FSpriteDetailsCustomization::EditorModeIsNot, ESpriteEditorMode::EditCollisionMode));
+	TAttribute<EVisibility> ShowWhenInCollisionMode = TAttribute<EVisibility>::Create(TAttribute<EVisibility>::FGetter::CreateSP(this, &FSpriteDetailsCustomization::EditorModeMatches, ESpriteEditorMode::EditCollisionMode));
+
+	static const FText EditRenderingInRenderingMode = LOCTEXT("RenderingPropertiesHiddenInCollisionMode", "Switch to 'Edit RenderGeom' mode\nto edit Rendering settings");
+	RenderingCategory.AddCustomRow(EditRenderingInRenderingMode)
+		.Visibility(ShowWhenInCollisionMode)
+		.WholeRowContent()
+		.HAlign(HAlign_Center)
+		[
+			SNew(STextBlock)
+			.Font(DetailLayout.GetDetailFontItalic())
+			.Justification(ETextJustify::Center)
+			.Text(EditRenderingInRenderingMode)
+		];
+
+	TArray<TWeakObjectPtr<UObject>> ObjectsBeingCustomized;
+	DetailLayout.GetObjectsBeingCustomized(/*out*/ ObjectsBeingCustomized);
+
+	if (ObjectsBeingCustomized.Num() > 0)
+	{
+		if (UPaperSprite* SpriteBeingEdited = Cast<UPaperSprite>(ObjectsBeingCustomized[0].Get()))
+		{
+			static const FText TypesOfMaterialsTooltip = LOCTEXT("TypesOfMaterialsTooltip", "Translucent materials can have smooth alpha edges, blending with the background\nMasked materials have on or off alpha, useful for cutouts\nOpaque materials have no transparency but render faster");
+
+			RenderingCategory.HeaderContent
+			(
+				SNew(SBox)
+				.HAlign(HAlign_Right)
+				[
+					SNew(SHorizontalBox)
+					+SHorizontalBox::Slot()
+					.Padding(FMargin(5.0f, 0.0f))
+					.AutoWidth()
+					[
+						SNew(STextBlock)
+						.Font(FEditorStyle::GetFontStyle("TinyText"))
+						.Text(this, &FSpriteDetailsCustomization::GetRenderingHeaderContentText, TWeakObjectPtr<UPaperSprite>(SpriteBeingEdited))
+						.ToolTipText(TypesOfMaterialsTooltip)
+					]
+				]
+			);
+
+		}
+	}
+
 	// Add the rendering geometry mode into the parent container (renamed)
 	const FString RenderGeometryTypePropertyPath = FString::Printf(TEXT("%s.%s"), GET_MEMBER_NAME_STRING_CHECKED(UPaperSprite, RenderGeometry), GET_MEMBER_NAME_STRING_CHECKED(FSpriteGeometryCollection, GeometryType));
 	TSharedPtr<IPropertyHandle> RenderGeometryTypeProperty = DetailLayout.GetProperty(*RenderGeometryTypePropertyPath);
 	RenderingCategory.AddProperty(RenderGeometryTypeProperty)
-		.DisplayName(LOCTEXT("RenderGeometryType", "Render Geometry Type"));
+		.DisplayName(LOCTEXT("RenderGeometryType", "Render Geometry Type"))
+		.Visibility(HideWhenInCollisionMode);
 
 	// Show the alternate material, but only when the mode is Diced
 	TAttribute<EVisibility> ShowWhenModeIsDiced = TAttribute<EVisibility>::Create(TAttribute<EVisibility>::FGetter::CreateSP(this, &FSpriteDetailsCustomization::PolygonModeMatches, RenderGeometryTypeProperty, ESpritePolygonMode::Diced));
@@ -109,22 +166,56 @@ void FSpriteDetailsCustomization::BuildRenderingSection(IDetailCategoryBuilder&
 
 	// Show the rendering geometry settings
 	TSharedRef<IPropertyHandle> RenderGeometry = DetailLayout.GetProperty(GET_MEMBER_NAME_CHECKED(UPaperSprite, RenderGeometry));
-	IDetailPropertyRow& RenderGeometryProperty = RenderingCategory.AddProperty(RenderGeometry);
+	IDetailPropertyRow& RenderGeometryProperty = RenderingCategory.AddProperty(RenderGeometry)
+		.Visibility(HideWhenInCollisionMode);
 
 	// Add the render polygons into advanced (renamed)
 	const FString RenderGeometryPolygonsPropertyPath = FString::Printf(TEXT("%s.%s"), GET_MEMBER_NAME_STRING_CHECKED(UPaperSprite, RenderGeometry), GET_MEMBER_NAME_STRING_CHECKED(FSpriteGeometryCollection, Shapes));
 	RenderingCategory.AddProperty(DetailLayout.GetProperty(*RenderGeometryPolygonsPropertyPath), EPropertyLocation::Advanced)
-		.DisplayName(LOCTEXT("RenderShapes", "Render Shapes"));
+		.DisplayName(LOCTEXT("RenderShapes", "Render Shapes"))
+		.Visibility(HideWhenInCollisionMode);
 }
 
 void FSpriteDetailsCustomization::BuildCollisionSection(IDetailCategoryBuilder& CollisionCategory, IDetailLayoutBuilder& DetailLayout)
 {
 	TSharedPtr<IPropertyHandle> SpriteCollisionDomainProperty = DetailLayout.GetProperty(GET_MEMBER_NAME_CHECKED(UPaperSprite, SpriteCollisionDomain));
-	TAttribute<EVisibility> ParticipatesInPhysics = TAttribute<EVisibility>::Create( TAttribute<EVisibility>::FGetter::CreateSP( this, &FSpriteDetailsCustomization::AnyPhysicsMode, SpriteCollisionDomainProperty) ) ;
+
+	CollisionCategory.HeaderContent
+	(
+		SNew(SBox)
+		.HAlign(HAlign_Right)
+		[
+			SNew(SHorizontalBox)
+			+SHorizontalBox::Slot()
+			.Padding(FMargin(5.0f, 0.0f))
+			.AutoWidth()
+			[
+				SNew(STextBlock)
+				.Font(FEditorStyle::GetFontStyle("TinyText"))
+				.Text(this, &FSpriteDetailsCustomization::GetCollisionHeaderContentText, SpriteCollisionDomainProperty)
+			]
+		]
+	);
+
+	TAttribute<EVisibility> ParticipatesInPhysics = TAttribute<EVisibility>::Create(TAttribute<EVisibility>::FGetter::CreateSP( this, &FSpriteDetailsCustomization::AnyPhysicsMode, SpriteCollisionDomainProperty));
 	TAttribute<EVisibility> ParticipatesInPhysics3D = TAttribute<EVisibility>::Create(TAttribute<EVisibility>::FGetter::CreateSP(this, &FSpriteDetailsCustomization::PhysicsModeMatches, SpriteCollisionDomainProperty, ESpriteCollisionMode::Use3DPhysics));
 	TAttribute<EVisibility> ParticipatesInPhysics2D = TAttribute<EVisibility>::Create(TAttribute<EVisibility>::FGetter::CreateSP(this, &FSpriteDetailsCustomization::PhysicsModeMatches, SpriteCollisionDomainProperty, ESpriteCollisionMode::Use2DPhysics));
+	TAttribute<EVisibility> HideWhenInRenderingMode = TAttribute<EVisibility>::Create(TAttribute<EVisibility>::FGetter::CreateSP(this, &FSpriteDetailsCustomization::EditorModeIsNot, ESpriteEditorMode::EditRenderingGeomMode));
+	TAttribute<EVisibility> ShowWhenInRenderingMode = TAttribute<EVisibility>::Create(TAttribute<EVisibility>::FGetter::CreateSP(this, &FSpriteDetailsCustomization::EditorModeMatches, ESpriteEditorMode::EditRenderingGeomMode));
 
-	CollisionCategory.AddProperty(SpriteCollisionDomainProperty);
+	static const FText EditCollisionInCollisionMode = LOCTEXT("CollisionPropertiesHiddenInRenderingMode", "Switch to 'Edit Collsion' mode\nto edit Collision settings");
+	CollisionCategory.AddCustomRow(EditCollisionInCollisionMode)
+		.Visibility(ShowWhenInRenderingMode)
+		.WholeRowContent()
+		.HAlign(HAlign_Center)
+		[
+			SNew(STextBlock)
+			.Font(DetailLayout.GetDetailFontItalic())
+			.Justification(ETextJustify::Center)
+			.Text(EditCollisionInCollisionMode)
+		];
+
+	CollisionCategory.AddProperty(SpriteCollisionDomainProperty).Visibility(HideWhenInRenderingMode);
 
 	// Add a warning bar about 2D collision being experimental
 	FText WarningFor2D = LOCTEXT("Experimental2DPhysicsWarning", "2D collision support is *experimental*");
@@ -156,20 +247,10 @@ void FSpriteDetailsCustomization::BuildCollisionSection(IDetailCategoryBuilder&
 			.Visibility(ParticipatesInPhysics);
 	}
 
-	// Show the collision geometry when not None
-	CollisionCategory.AddProperty( DetailLayout.GetProperty(GET_MEMBER_NAME_CHECKED(UPaperSprite, CollisionGeometry)) )
-		.Visibility(ParticipatesInPhysics);
-
 	// Show the collision thickness only in 3D mode
 	CollisionCategory.AddProperty( DetailLayout.GetProperty(GET_MEMBER_NAME_CHECKED(UPaperSprite, CollisionThickness)) )
 		.Visibility(ParticipatesInPhysics3D);
 
-	// Add the collision polygons into advanced (renamed)
-	const FString CollisionGeometryPolygonsPropertyPath = FString::Printf(TEXT("%s.%s"), GET_MEMBER_NAME_STRING_CHECKED(UPaperSprite, CollisionGeometry), GET_MEMBER_NAME_STRING_CHECKED(FSpriteGeometryCollection, Shapes));
-	CollisionCategory.AddProperty(DetailLayout.GetProperty(*CollisionGeometryPolygonsPropertyPath), EPropertyLocation::Advanced)
-		.DisplayName(LOCTEXT("CollisionShapes", "Collision Shapes"))
-		.Visibility(ParticipatesInPhysics);
-
 	// Show the default body instance (and only it) from the body setup (if it exists)
 	DetailLayout.HideProperty("BodySetup");
 	IDetailPropertyRow& BodySetupDefaultInstance = CollisionCategory.AddProperty("BodySetup.DefaultInstance");
@@ -197,6 +278,16 @@ void FSpriteDetailsCustomization::BuildCollisionSection(IDetailCategoryBuilder&
 			DefaultInstanceRow->Visibility(ParticipatesInPhysics);
 		}
 	}
+
+	// Show the collision geometry when not None
+	CollisionCategory.AddProperty(DetailLayout.GetProperty(GET_MEMBER_NAME_CHECKED(UPaperSprite, CollisionGeometry)))
+		.Visibility(ParticipatesInPhysics);
+
+	// Add the collision polygons into advanced (renamed)
+	const FString CollisionGeometryPolygonsPropertyPath = FString::Printf(TEXT("%s.%s"), GET_MEMBER_NAME_STRING_CHECKED(UPaperSprite, CollisionGeometry), GET_MEMBER_NAME_STRING_CHECKED(FSpriteGeometryCollection, Shapes));
+	CollisionCategory.AddProperty(DetailLayout.GetProperty(*CollisionGeometryPolygonsPropertyPath), EPropertyLocation::Advanced)
+		.DisplayName(LOCTEXT("CollisionShapes", "Collision Shapes"))
+		.Visibility(ParticipatesInPhysics);
 }
 
 void FSpriteDetailsCustomization::BuildTextureSection(IDetailCategoryBuilder& SpriteCategory, IDetailLayoutBuilder& DetailLayout)
@@ -313,8 +404,23 @@ TSharedRef<SWidget> FSpriteDetailsCustomization::CreateTextureNameWidget(TShared
 	}
 }
 
+EVisibility FSpriteDetailsCustomization::EditorModeMatches(ESpriteEditorMode::Type DesiredMode) const
+{
+	return (SpriteEditMode.Get() == DesiredMode) ? EVisibility::Visible : EVisibility::Collapsed;
+}
+
+EVisibility FSpriteDetailsCustomization::EditorModeIsNot(ESpriteEditorMode::Type DesiredMode) const
+{
+	return (SpriteEditMode.Get() != DesiredMode) ? EVisibility::Visible : EVisibility::Collapsed;
+}
+
 EVisibility FSpriteDetailsCustomization::PhysicsModeMatches(TSharedPtr<IPropertyHandle> Property, ESpriteCollisionMode::Type DesiredMode) const
 {
+	if (SpriteEditMode.Get() == ESpriteEditorMode::EditRenderingGeomMode)
+	{
+		return EVisibility::Collapsed;
+	}
+
 	if (Property.IsValid())
 	{
 		uint8 ValueAsByte;
@@ -332,6 +438,11 @@ EVisibility FSpriteDetailsCustomization::PhysicsModeMatches(TSharedPtr<IProperty
 
 EVisibility FSpriteDetailsCustomization::AnyPhysicsMode(TSharedPtr<IPropertyHandle> Property) const
 {
+	if (SpriteEditMode.Get() == ESpriteEditorMode::EditRenderingGeomMode)
+	{
+		return EVisibility::Collapsed;
+	}
+
 	if (Property.IsValid())
 	{
 		uint8 ValueAsByte;
@@ -347,6 +458,88 @@ EVisibility FSpriteDetailsCustomization::AnyPhysicsMode(TSharedPtr<IPropertyHand
 	return EVisibility::Visible;
 }
 
+FText FSpriteDetailsCustomization::GetCollisionHeaderContentText(TSharedPtr<IPropertyHandle> Property) const
+{
+	FText CollisionContentText;
+
+	if (Property.IsValid())
+	{
+		uint8 ValueAsByte;
+		if (Property->GetValue(/*out*/ ValueAsByte) == FPropertyAccess::Success)
+		{
+			ESpriteCollisionMode::Type CollisionMode = ((ESpriteCollisionMode::Type)ValueAsByte);
+
+			switch (CollisionMode)
+			{
+			case ESpriteCollisionMode::None:
+				CollisionContentText = LOCTEXT("CollisionHeader_NoCollision", "(no collision)");
+				break;
+			case ESpriteCollisionMode::Use2DPhysics:
+				CollisionContentText = LOCTEXT("CollisionHeader_Use2D", "Uses 2D Physics");
+				break;
+			case ESpriteCollisionMode::Use3DPhysics:
+				CollisionContentText = LOCTEXT("CollisionHeader_Use3D", "Uses 3D Physics");
+				break;
+			}
+		}
+	}
+
+	return CollisionContentText;
+}
+
+
+FText FSpriteDetailsCustomization::GetRenderingHeaderContentText(TWeakObjectPtr<UPaperSprite> WeakSprite) const
+{
+	FText HeaderDisplayText;
+
+	if (UPaperSprite* SpriteBeingEdited = WeakSprite.Get())
+	{
+		int32 NumOpaqueTriangles = 0;
+		int32 NumMaskedTriangles = 0;
+		int32 NumTranslucentTriangles = 0;
+		FSpriteEditorViewportClient::AnalyzeSpriteMaterialType(SpriteBeingEdited, /*out*/ NumOpaqueTriangles, /*out*/ NumMaskedTriangles, /*out*/ NumTranslucentTriangles);
+
+		static const FText Opaque = LOCTEXT("OpaqueMaterial", "Opaque");
+		static const FText Translucent = LOCTEXT("TranslucentMaterial", "Translucent");
+		static const FText Masked = LOCTEXT("MaskedMaterial", "Masked");
+		static const FText TwoItems = LOCTEXT("SpriteWithTwoMaterialsRenderHeaderText", "{0} and {1}");
+		static const FText ThreeItems = LOCTEXT("SpriteWithThreeMaterialsRenderHeaderText", "{0}, {1}, {2}"); // Should never happen right now!
+
+		TArray<const FText*, TInlineAllocator<3>> MaterialTypes;
+		if (NumOpaqueTriangles > 0)
+		{
+			MaterialTypes.Add(&Opaque);
+		}
+
+		if (NumMaskedTriangles > 0)
+		{
+			MaterialTypes.Add(&Masked);
+		}
+
+		if (NumTranslucentTriangles > 0)
+		{
+			MaterialTypes.Add(&Translucent);
+		}
+
+		switch (MaterialTypes.Num())
+		{
+		case 0:
+			break;
+		case 1:
+			HeaderDisplayText = *MaterialTypes[0];
+			break;
+		case 2:
+			HeaderDisplayText = FText::Format(TwoItems, *MaterialTypes[0], *MaterialTypes[1]);
+			break;
+		case 3:
+			HeaderDisplayText = FText::Format(ThreeItems, *MaterialTypes[0], *MaterialTypes[1], *MaterialTypes[2]);
+			break;
+		};
+	}
+
+	return HeaderDisplayText;
+}
+
 EVisibility FSpriteDetailsCustomization::GetAtlasGroupVisibility()
 {
 	return GetDefault<UPaperRuntimeSettings>()->bEnableSpriteAtlasGroups ? EVisibility::Visible : EVisibility::Collapsed;
@@ -377,6 +570,14 @@ EVisibility FSpriteDetailsCustomization::Get2DPhysicsNotEnabledWarningVisibility
 
 EVisibility FSpriteDetailsCustomization::PolygonModeMatches(TSharedPtr<IPropertyHandle> Property, ESpritePolygonMode::Type DesiredMode) const
 {
+	if (DesiredMode == ESpritePolygonMode::Diced)
+	{
+		if (SpriteEditMode.Get() == ESpriteEditorMode::EditCollisionMode)
+		{
+			return EVisibility::Collapsed;
+		}
+	}
+
 	if (Property.IsValid())
 	{
 		uint8 ValueAsByte;
diff --git a/Engine/Plugins/2D/Paper2D/Source/Paper2DEditor/Private/SpriteEditor/SpriteDetailsCustomization.h b/Engine/Plugins/2D/Paper2D/Source/Paper2DEditor/Private/SpriteEditor/SpriteDetailsCustomization.h
--- a/Engine/Plugins/2D/Paper2D/Source/Paper2DEditor/Private/SpriteEditor/SpriteDetailsCustomization.h
+++ b/Engine/Plugins/2D/Paper2D/Source/Paper2DEditor/Private/SpriteEditor/SpriteDetailsCustomization.h
@@ -3,6 +3,7 @@
 #pragma once
 
 #include "PropertyEditing.h"
+#include "SpriteEditorViewportClient.h"
 
 //////////////////////////////////////////////////////////////////////////
 // FSpriteDetailsCustomization
@@ -13,11 +14,18 @@ public:
 	/** Makes a new instance of this detail layout class for a specific detail view requesting it */
 	static TSharedRef<IDetailCustomization> MakeInstance();
 
+	/** Makes a new instance of this detail layout class for a specific detail view requesting it */
+	static TSharedRef<IDetailCustomization> MakeInstanceForSpriteEditor(TAttribute<ESpriteEditorMode::Type> InEditMode);
+
 	// IDetailCustomization interface
 	virtual void CustomizeDetails(IDetailLayoutBuilder& DetailLayout) override;
 	// End of IDetailCustomization interface
 
 protected:
+	FSpriteDetailsCustomization(TAttribute<ESpriteEditorMode::Type> InEditMode);
+
+	EVisibility EditorModeMatches(ESpriteEditorMode::Type DesiredMode) const;
+	EVisibility EditorModeIsNot(ESpriteEditorMode::Type DesiredMode) const;
 	EVisibility PhysicsModeMatches(TSharedPtr<IPropertyHandle> Property, ESpriteCollisionMode::Type DesiredMode) const;
 	EVisibility AnyPhysicsMode(TSharedPtr<IPropertyHandle> Property) const;
 	EVisibility GetCustomPivotVisibility(TSharedPtr<IPropertyHandle> Property) const;
@@ -25,6 +33,9 @@ protected:
 
 	EVisibility PolygonModeMatches(TSharedPtr<IPropertyHandle> Property, ESpritePolygonMode::Type DesiredMode) const;
 
+	FText GetCollisionHeaderContentText(TSharedPtr<IPropertyHandle> Property) const;
+	FText GetRenderingHeaderContentText(TWeakObjectPtr<class UPaperSprite> WeakSprite) const;
+
 	static EVisibility GetAtlasGroupVisibility();
 
 	static FDetailWidgetRow& GenerateWarningRow(IDetailCategoryBuilder& WarningCategory, bool bExperimental, const FText& WarningText, const FText& Tooltip, const FString& ExcerptLink, const FString& ExcerptName);
@@ -40,6 +51,8 @@ protected:
 	TSharedRef<SWidget> CreateTextureNameWidget(TSharedPtr<IPropertyHandle> PropertyHandle, const FText& OverrideText);
 
 protected:
+	TAttribute<ESpriteEditorMode::Type> SpriteEditMode;
+
 	// These are labels harvested from the material
 	TMap<int32, FText> AdditionalTextureLabels;
-};
+};
\ No newline at end of file
diff --git a/Engine/Plugins/2D/Paper2D/Source/Paper2DEditor/Private/SpriteEditor/SpriteEditor.cpp b/Engine/Plugins/2D/Paper2D/Source/Paper2DEditor/Private/SpriteEditor/SpriteEditor.cpp
--- a/Engine/Plugins/2D/Paper2D/Source/Paper2DEditor/Private/SpriteEditor/SpriteEditor.cpp
+++ b/Engine/Plugins/2D/Paper2D/Source/Paper2DEditor/Private/SpriteEditor/SpriteEditor.cpp
@@ -13,6 +13,7 @@
 #include "WorkspaceMenuStructureModule.h"
 #include "Paper2DEditorModule.h"
 #include "SSpriteEditorViewportToolbar.h"
+#include "SpriteDetailsCustomization.h"
 
 #include "SSpriteList.h"
 #include "SDockTab.h"
@@ -68,6 +69,11 @@ public:
 	{
 		EditorViewportClient->NotifySpriteBeingEditedHasChanged();
 	}
+
+	ESpriteEditorMode::Type GetCurrentMode() const
+	{
+		return EditorViewportClient->GetCurrentMode();
+	}
 private:
 	// Pointer back to owning sprite editor instance (the keeper of state)
 	TWeakPtr<class FSpriteEditor> SpriteEditorPtr;
@@ -259,6 +265,11 @@ public:
 		SpriteEditorPtr = InSpriteEditor;
 
 		SSingleObjectDetailsPanel::Construct(SSingleObjectDetailsPanel::FArguments(), /*bAutomaticallyObserveViaGetObjectToObserve=*/ true, /*bAllowSearch=*/ true);
+
+		TAttribute<ESpriteEditorMode::Type> SpriteEditorMode = TAttribute<ESpriteEditorMode::Type>::Create(
+			TAttribute<ESpriteEditorMode::Type>::FGetter::CreateSP(InSpriteEditor.ToSharedRef(), &FSpriteEditor::GetCurrentMode));
+		FOnGetDetailCustomizationInstance CustomizeSpritesForEditor = FOnGetDetailCustomizationInstance::CreateStatic(&FSpriteDetailsCustomization::MakeInstanceForSpriteEditor, SpriteEditorMode);
+		PropertyView->RegisterInstancedCustomPropertyLayout(UPaperSprite::StaticClass(), CustomizeSpritesForEditor);
 	}
 
 	// SSingleObjectDetailsPanel interface
@@ -549,6 +560,10 @@ void FSpriteEditor::SetSpriteBeingEdited(UPaperSprite* NewSprite)
 	}
 }
 
+ESpriteEditorMode::Type FSpriteEditor::GetCurrentMode() const
+{
+	return ViewportPtr->GetCurrentMode();
+}
 
 void FSpriteEditor::CreateModeToolbarWidgets(FToolBarBuilder& IgnoredBuilder)
 {
diff --git a/Engine/Plugins/2D/Paper2D/Source/Paper2DEditor/Private/SpriteEditor/SpriteEditor.h b/Engine/Plugins/2D/Paper2D/Source/Paper2DEditor/Private/SpriteEditor/SpriteEditor.h
--- a/Engine/Plugins/2D/Paper2D/Source/Paper2DEditor/Private/SpriteEditor/SpriteEditor.h
+++ b/Engine/Plugins/2D/Paper2D/Source/Paper2DEditor/Private/SpriteEditor/SpriteEditor.h
@@ -5,6 +5,20 @@
 #include "Toolkits/AssetEditorToolkit.h"
 #include "Toolkits/AssetEditorManager.h"
 
+//////////////////////////////////////////////////////////////////////////
+// 
+
+namespace ESpriteEditorMode
+{
+	enum Type
+	{
+		ViewMode,
+		EditSourceRegionMode,
+		EditCollisionMode,
+		EditRenderingGeomMode
+	};
+}
+
 //////////////////////////////////////////////////////////////////////////
 // FSpriteEditor
 
@@ -36,6 +50,9 @@ public:
 
 	UPaperSprite* GetSpriteBeingEdited() const { return SpriteBeingEdited; }
 	void SetSpriteBeingEdited(UPaperSprite* NewSprite);
+
+	ESpriteEditorMode::Type GetCurrentMode() const;
+
 protected:
 	UPaperSprite* SpriteBeingEdited;
 	TSharedPtr<class SSpriteEditorViewport> ViewportPtr;
diff --git a/Engine/Plugins/2D/Paper2D/Source/Paper2DEditor/Private/SpriteEditor/SpriteEditorViewportClient.cpp b/Engine/Plugins/2D/Paper2D/Source/Paper2DEditor/Private/SpriteEditor/SpriteEditorViewportClient.cpp
--- a/Engine/Plugins/2D/Paper2D/Source/Paper2DEditor/Private/SpriteEditor/SpriteEditorViewportClient.cpp
+++ b/Engine/Plugins/2D/Paper2D/Source/Paper2DEditor/Private/SpriteEditor/SpriteEditorViewportClient.cpp
@@ -336,27 +336,6 @@ void FSpriteEditorViewportClient::DrawGeometryStats(FViewport& InViewport, FScen
 	YPos = (int32)TextItem.Position.Y;
 }
 
-void AttributeTrianglesByMaterialType(int32 NumTriangles, UMaterialInterface* Material, int32& NumOpaqueTriangles, int32& NumMaskedTriangles, int32& NumTranslucentTriangles)
-{
-	if (Material != nullptr)
-	{
-		switch (Material->GetBlendMode())
-		{
-		case EBlendMode::BLEND_Opaque:
-			NumOpaqueTriangles += NumTriangles;
-			break;
-		case EBlendMode::BLEND_Translucent:
-		case EBlendMode::BLEND_Additive:
-		case EBlendMode::BLEND_Modulate:
-			NumTranslucentTriangles += NumTriangles;
-			break;
-		case EBlendMode::BLEND_Masked:
-			NumMaskedTriangles += NumTriangles;
-			break;
-		}
-	}
-}
-
 void FSpriteEditorViewportClient::DrawCollisionStats(FViewport& InViewport, FSceneView& View, FCanvas& Canvas, class UBodySetup* BodySetup, int32& YPos)
 {
 	FCanvasTextItem TextItem(FVector2D(6, YPos), LOCTEXT("CollisionGeomBaked", "Collision Geometry (baked)"), GEngine->GetSmallFont(), FLinearColor::White);
@@ -440,25 +419,35 @@ void FSpriteEditorViewportClient::DrawCollisionStats(FViewport& InViewport, FSce
 	YPos = (int32)TextItem.Position.Y;
 }
 
-void FSpriteEditorViewportClient::DrawRenderStats(FViewport& InViewport, FSceneView& View, FCanvas& Canvas, class UPaperSprite* Sprite, int32& YPos)
+void FSpriteEditorViewportClient::AnalyzeSpriteMaterialType(UPaperSprite* Sprite, int32& OutNumOpaque, int32& OutNumMasked, int32& OutNumTranslucent)
 {
-	FCanvasTextItem TextItem(FVector2D(6, YPos), LOCTEXT("RenderGeomBaked", "Render Geometry (baked)"), GEngine->GetSmallFont(), FLinearColor::White);
-	TextItem.EnableShadow(FLinearColor::Black);
-
-	TextItem.Draw(&Canvas);
-	TextItem.Position += FVector2D(6.0f, 18.0f);
-
-	int32 NumSections = (Sprite->AlternateMaterialSplitIndex != INDEX_NONE) ? 2 : 1;
-	if (NumSections > 1)
+	struct Local
 	{
-		TextItem.Text = FText::Format(LOCTEXT("SectionCount", "Sections: {0}"), FText::AsNumber(NumSections));
-		TextItem.Draw(&Canvas);
-		TextItem.Position.Y += 18.0f;
-	}
+		static void AttributeTrianglesByMaterialType(int32 NumTriangles, UMaterialInterface* Material, int32& NumOpaqueTriangles, int32& NumMaskedTriangles, int32& NumTranslucentTriangles)
+		{
+			if (Material != nullptr)
+			{
+				switch (Material->GetBlendMode())
+				{
+				case EBlendMode::BLEND_Opaque:
+					NumOpaqueTriangles += NumTriangles;
+					break;
+				case EBlendMode::BLEND_Translucent:
+				case EBlendMode::BLEND_Additive:
+				case EBlendMode::BLEND_Modulate:
+					NumTranslucentTriangles += NumTriangles;
+					break;
+				case EBlendMode::BLEND_Masked:
+					NumMaskedTriangles += NumTriangles;
+					break;
+				}
+			}
+		}
+	};
 
-	int32 NumOpaqueTriangles = 0;
-	int32 NumMaskedTriangles = 0;
-	int32 NumTranslucentTriangles = 0;
+	OutNumOpaque = 0;
+	OutNumMasked = 0;
+	OutNumTranslucent = 0;
 
 	int32 NumVerts = Sprite->BakedRenderData.Num();
 	int32 DefaultTriangles = 0;
@@ -473,8 +462,30 @@ void FSpriteEditorViewportClient::DrawRenderStats(FViewport& InViewport, FSceneV
 		DefaultTriangles = NumVerts / 3;
 	}
 
-	AttributeTrianglesByMaterialType(DefaultTriangles, Sprite->GetDefaultMaterial(), /*inout*/ NumOpaqueTriangles, /*inout*/ NumMaskedTriangles, /*inout*/ NumTranslucentTriangles);
-	AttributeTrianglesByMaterialType(AlternateTriangles, Sprite->GetAlternateMaterial(), /*inout*/ NumOpaqueTriangles, /*inout*/ NumMaskedTriangles, /*inout*/ NumTranslucentTriangles);
+	Local::AttributeTrianglesByMaterialType(DefaultTriangles, Sprite->GetDefaultMaterial(), /*inout*/ OutNumOpaque, /*inout*/ OutNumMasked, /*inout*/ OutNumTranslucent);
+	Local::AttributeTrianglesByMaterialType(AlternateTriangles, Sprite->GetAlternateMaterial(), /*inout*/ OutNumOpaque, /*inout*/ OutNumMasked, /*inout*/ OutNumTranslucent);
+}
+
+void FSpriteEditorViewportClient::DrawRenderStats(FViewport& InViewport, FSceneView& View, FCanvas& Canvas, class UPaperSprite* Sprite, int32& YPos)
+{
+	FCanvasTextItem TextItem(FVector2D(6, YPos), LOCTEXT("RenderGeomBaked", "Render Geometry (baked)"), GEngine->GetSmallFont(), FLinearColor::White);
+	TextItem.EnableShadow(FLinearColor::Black);
+
+	TextItem.Draw(&Canvas);
+	TextItem.Position += FVector2D(6.0f, 18.0f);
+
+	int32 NumOpaqueTriangles = 0;
+	int32 NumMaskedTriangles = 0;
+	int32 NumTranslucentTriangles = 0;
+	AnalyzeSpriteMaterialType(Sprite, /*out*/ NumOpaqueTriangles, /*out*/ NumMaskedTriangles, /*out*/ NumTranslucentTriangles);
+
+	int32 NumSections = (Sprite->AlternateMaterialSplitIndex != INDEX_NONE) ? 2 : 1;
+	if (NumSections > 1)
+	{
+		TextItem.Text = FText::Format(LOCTEXT("SectionCount", "Sections: {0}"), FText::AsNumber(NumSections));
+		TextItem.Draw(&Canvas);
+		TextItem.Position.Y += 18.0f;
+	}
 
 	// Draw the number of triangles
 	if (NumOpaqueTriangles > 0)
diff --git a/Engine/Plugins/2D/Paper2D/Source/Paper2DEditor/Private/SpriteEditor/SpriteEditorViewportClient.h b/Engine/Plugins/2D/Paper2D/Source/Paper2DEditor/Private/SpriteEditor/SpriteEditorViewportClient.h
--- a/Engine/Plugins/2D/Paper2D/Source/Paper2DEditor/Private/SpriteEditor/SpriteEditorViewportClient.h
+++ b/Engine/Plugins/2D/Paper2D/Source/Paper2DEditor/Private/SpriteEditor/SpriteEditorViewportClient.h
@@ -13,20 +13,6 @@
 #include "AssetData.h"
 #include "PaperEditorShared/SpriteGeometryEditing.h"
 
-//////////////////////////////////////////////////////////////////////////
-// 
-
-namespace ESpriteEditorMode
-{
-	enum Type
-	{
-		ViewMode,
-		EditSourceRegionMode,
-		EditCollisionMode,
-		EditRenderingGeomMode
-	};
-}
-
 //////////////////////////////////////////////////////////////////////////
 // FRelatedSprite
 
@@ -132,6 +118,8 @@ public:
 	{
 		return CurrentMode;
 	}
+
+	static void AnalyzeSpriteMaterialType(UPaperSprite* Sprite, int32& OutNumOpaque, int32& OutNumMasked, int32& OutNumTranslucent);
 private:
 	// Editor mode
 	ESpriteEditorMode::Type CurrentMode;
