Paper2D: Improved various comments in PaperSprite.h

diff --git a/Engine/Plugins/2D/Paper2D/Source/Paper2D/Classes/PaperSprite.h b/Engine/Plugins/2D/Paper2D/Source/Paper2D/Classes/PaperSprite.h
--- a/Engine/Plugins/2D/Paper2D/Source/Paper2D/Classes/PaperSprite.h
+++ b/Engine/Plugins/2D/Paper2D/Source/Paper2D/Classes/PaperSprite.h
@@ -12,9 +12,14 @@ namespace ESpriteCollisionMode
 {
 	enum Type
 	{
+		// Should this have no collison and not participate in physics?
 		None,
-		Use2DPhysics,
-		Use3DPhysics
+
+		// EXPERIMENTAL: Should this have 2D collision geometry and participate in the 2D physics world?
+		Use2DPhysics UMETA(DisplayName = "Use 2D Physics"),
+
+		// Should this have 3D collision geometry and participate in the 3D physics world?
+		Use3DPhysics UMETA(DisplayName = "Use 3D Physics")
 	};
 }
 
@@ -28,6 +33,7 @@ struct FPaperSpriteSocket
 	UPROPERTY(EditAnywhere, BlueprintReadOnly, Category=Sockets)
 	FTransform LocalTransform;
 
+	// Name of the socket
 	UPROPERTY(EditAnywhere, BlueprintReadOnly, Category=Sockets)
 	FName SocketName;
 };
@@ -39,6 +45,8 @@ typedef TArray<class UTexture2D*, TInlineAllocator<4>> FAdditionalSpriteTextureA
  *
  * Stores the data necessary to render a single 2D sprite (from a region of a texture)
  * Can also contain collision shapes for the sprite.
+ *
+ * @see UPaperSpriteComponent
  */
 
 UCLASS(BlueprintType, meta=(DisplayThumbnail = "true"))
