Flipbooks: Clamped maximum framerate to 1000 [UE-13157]

diff --git a/Engine/Plugins/2D/Paper2D/Source/Paper2D/Classes/PaperFlipbook.h b/Engine/Plugins/2D/Paper2D/Source/Paper2D/Classes/PaperFlipbook.h
--- a/Engine/Plugins/2D/Paper2D/Source/Paper2D/Classes/PaperFlipbook.h
+++ b/Engine/Plugins/2D/Paper2D/Source/Paper2D/Classes/PaperFlipbook.h
@@ -50,7 +50,7 @@ class PAPER2D_API UPaperFlipbook : public UObject
 
 protected:
 	// The nominal frame rate to play this flipbook animation back at
-	UPROPERTY(Category=Sprite, EditAnywhere, BlueprintReadOnly)
+	UPROPERTY(Category=Sprite, EditAnywhere, BlueprintReadOnly, meta=(ClampMax=1000))
 	float FramesPerSecond;
 
 	// The set of key frames for this flipbook animation (each one has a duration and a sprite to display)
