Paper2D: Added a class comment to UPaperTileSet [UE-13239]

diff --git a/Engine/Plugins/2D/Paper2D/Source/Paper2D/Classes/PaperTileSet.h b/Engine/Plugins/2D/Paper2D/Source/Paper2D/Classes/PaperTileSet.h
--- a/Engine/Plugins/2D/Paper2D/Source/Paper2D/Classes/PaperTileSet.h
+++ b/Engine/Plugins/2D/Paper2D/Source/Paper2D/Classes/PaperTileSet.h
@@ -7,6 +7,11 @@
 
 #include "PaperTileSet.generated.h"
 
+/**
+ * A tile set is a collection of tiles pulled from a texture that can be used to fill out a tile map.
+ *
+ * @see UPaperTileMap, UPaperTileMapComponent
+ */
 UCLASS()
 class PAPER2D_API UPaperTileSet : public UObject
 {
