Tile Map Editor: Switched default focus key to F from Home, matching the rest of the editor
- Moved Fill tool from F to G to avoid conflicts (matches Photoshop)

diff --git a/Engine/Plugins/2D/Paper2D/Source/Paper2DEditor/Private/TileMapEditing/TileMapEditorCommands.cpp b/Engine/Plugins/2D/Paper2D/Source/Paper2DEditor/Private/TileMapEditing/TileMapEditorCommands.cpp
--- a/Engine/Plugins/2D/Paper2D/Source/Paper2DEditor/Private/TileMapEditing/TileMapEditorCommands.cpp
+++ b/Engine/Plugins/2D/Paper2D/Source/Paper2DEditor/Private/TileMapEditing/TileMapEditorCommands.cpp
@@ -14,7 +14,7 @@ void FTileMapEditorCommands::RegisterCommands()
 
 	UI_COMMAND(SelectPaintTool, "Paint", "Paint", EUserInterfaceActionType::RadioButton, FInputChord(EKeys::B));
 	UI_COMMAND(SelectEraserTool, "Eraser", "Eraser", EUserInterfaceActionType::RadioButton, FInputChord(EKeys::E));
-	UI_COMMAND(SelectFillTool, "Fill", "Paint Bucket", EUserInterfaceActionType::RadioButton, FInputChord(EKeys::F));
+	UI_COMMAND(SelectFillTool, "Fill", "Paint Bucket", EUserInterfaceActionType::RadioButton, FInputChord(EKeys::G));
 
 	// Show toggles
 	UI_COMMAND(SetShowCollision, "Collision", "Toggles display of the simplified collision mesh of the static mesh, if one has been assigned.", EUserInterfaceActionType::ToggleButton, FInputChord(EKeys::C, EModifierKey::Alt));
@@ -22,7 +22,7 @@ void FTileMapEditorCommands::RegisterCommands()
 	UI_COMMAND(SetShowPivot, "Pivot", "Display the pivot location of the sprite.", EUserInterfaceActionType::ToggleButton, FInputChord());
 
 	// Misc. actions
-	UI_COMMAND(FocusOnTileMap, "Focus on tile map", "Centers and zooms the view to focus on the tile map.", EUserInterfaceActionType::Button, FInputChord(EKeys::Home, EModifierKey::None));
+	UI_COMMAND(FocusOnTileMap, "Focus on tile map", "Centers and zooms the view to focus on the tile map.", EUserInterfaceActionType::Button, FInputChord(EKeys::F, EModifierKey::None));
 
 	// Selection actions
 	UI_COMMAND(FlipSelectionHorizontally, "Flip selection horizontally", "Flips the selection horizontally", EUserInterfaceActionType::Button, FInputChord(EKeys::X));
