Sprite Editor: Switched default focus key to F from Home, matching the rest of the editor

diff --git a/Engine/Plugins/2D/Paper2D/Source/Paper2DEditor/Private/SpriteEditor/SpriteEditorCommands.cpp b/Engine/Plugins/2D/Paper2D/Source/Paper2DEditor/Private/SpriteEditor/SpriteEditorCommands.cpp
--- a/Engine/Plugins/2D/Paper2D/Source/Paper2DEditor/Private/SpriteEditor/SpriteEditorCommands.cpp
+++ b/Engine/Plugins/2D/Paper2D/Source/Paper2DEditor/Private/SpriteEditor/SpriteEditorCommands.cpp
@@ -29,7 +29,7 @@ void FSpriteEditorCommands::RegisterCommands()
 	UI_COMMAND(EnterRenderingEditMode, "Edit RenderGeom", "Edit the rendering geometry (useful to reduce overdraw).", EUserInterfaceActionType::ToggleButton, FInputChord());
 
 	// Misc. actions
-	UI_COMMAND(FocusOnSprite, "Focus on sprite", "Centers and zooms the view to focus on the current sprite.", EUserInterfaceActionType::Button, FInputChord(EKeys::Home, EModifierKey::None));
+	UI_COMMAND(FocusOnSprite, "Focus on sprite", "Centers and zooms the view to focus on the current sprite.", EUserInterfaceActionType::Button, FInputChord(EKeys::F, EModifierKey::None));
 
 	// Geometry editing commands
 	UI_COMMAND(DeleteSelection, "Delete", "Delete the selection.", EUserInterfaceActionType::Button, FInputChord(EKeys::Platform_Delete, EModifierKey::None));
