Paper2D: Confgured the perspective orientation for Paper2D editors so that switching to it will be looking at the sprite/flipbook/etc... instead of along it [UE-13235]
- Removed per-editor focus commands from Paper2D in favor of the inherited one from SEditorViewport (still F for all editors)

diff --git a/Engine/Plugins/2D/Paper2D/Source/Paper2DEditor/Private/FlipbookEditor/FlipbookEditor.cpp b/Engine/Plugins/2D/Paper2D/Source/Paper2DEditor/Private/FlipbookEditor/FlipbookEditor.cpp
--- a/Engine/Plugins/2D/Paper2D/Source/Paper2DEditor/Private/FlipbookEditor/FlipbookEditor.cpp
+++ b/Engine/Plugins/2D/Paper2D/Source/Paper2DEditor/Private/FlipbookEditor/FlipbookEditor.cpp
@@ -59,6 +59,7 @@ public:
 	virtual TSharedRef<FEditorViewportClient> MakeEditorViewportClient() override;
 	virtual TSharedPtr<SWidget> MakeViewportToolbar() override;
 	virtual EVisibility GetTransformToolbarVisibility() const override;
+	virtual void OnFocusViewportToSelection() override;
 	// End of SEditorViewport interface
 
 	// ICommonEditorViewportToolbarInfoProvider interface
@@ -150,10 +151,15 @@ TSharedPtr<SWidget> SFlipbookEditorViewport::MakeViewportToolbar()
 
 EVisibility SFlipbookEditorViewport::GetTransformToolbarVisibility() const
 {
-	// Nothing to transform in the flipbook editor, yet...
+	//@TODO: Nothing to transform in the flipbook editor, yet...
 	return EVisibility::Hidden;
 }
 
+void SFlipbookEditorViewport::OnFocusViewportToSelection()
+{
+	EditorViewportClient->RequestFocusOnSelection(/*bInstant=*/ false);
+}
+
 TSharedRef<class SEditorViewport> SFlipbookEditorViewport::GetViewportWidget()
 {
 	return SharedThis(this);
diff --git a/Engine/Plugins/2D/Paper2D/Source/Paper2DEditor/Private/FlipbookEditor/FlipbookEditorViewportClient.cpp b/Engine/Plugins/2D/Paper2D/Source/Paper2DEditor/Private/FlipbookEditor/FlipbookEditorViewportClient.cpp
--- a/Engine/Plugins/2D/Paper2D/Source/Paper2DEditor/Private/FlipbookEditor/FlipbookEditorViewportClient.cpp
+++ b/Engine/Plugins/2D/Paper2D/Source/Paper2DEditor/Private/FlipbookEditor/FlipbookEditorViewportClient.cpp
@@ -31,7 +31,6 @@ FFlipbookEditorViewportClient::FFlipbookEditorViewportClient(const TAttribute<UP
 
 	bShowPivot = false;
 	bShowSockets = true;
-	bDeferZoomToSprite = true;
 	DrawHelper.bDrawGrid = false;
 
 	EngineShowFlags.DisableAdvancedFeatures();
@@ -86,6 +85,11 @@ void FFlipbookEditorViewportClient::Draw(const FSceneView* View, FPrimitiveDrawI
 	}
 }
 
+FBox FFlipbookEditorViewportClient::GetDesiredFocusBounds() const
+{
+	return AnimatedRenderComponent->Bounds.GetBox();
+}
+
 void FFlipbookEditorViewportClient::Tick(float DeltaSeconds)
 {
 	if (AnimatedRenderComponent.IsValid())
@@ -97,16 +101,6 @@ void FFlipbookEditorViewportClient::Tick(float DeltaSeconds)
 			AnimatedRenderComponent->UpdateBounds();
 			FlipbookBeingEditedLastFrame = Flipbook;
 		}
-
-		// Zoom in on the sprite
-		//@TODO: This doesn't work correctly, only partially zooming in or something
-		//@TODO: Fix this properly so it doesn't need to be deferred, or wait for the viewport to initialize
-		FIntPoint Size = Viewport->GetSizeXY();
-		if (bDeferZoomToSprite && (Size.X > 0) && (Size.Y > 0))
-		{
-			FocusViewportOnBox(AnimatedRenderComponent->Bounds.GetBox(), true);
-			bDeferZoomToSprite = false;
-		}
 	}
 
 	FPaperEditorViewportClient::Tick(DeltaSeconds);
@@ -121,12 +115,6 @@ bool FFlipbookEditorViewportClient::InputKey(FViewport* Viewport, int32 Controll
 {
 	bool bHandled = false;
 
-	if ((Event == IE_Pressed) && (Key == EKeys::F) && (AnimatedRenderComponent.IsValid()))
-	{
-		bHandled = true;
-		FocusViewportOnBox(AnimatedRenderComponent->Bounds.GetBox());
-	}
-
 	// Pass keys to standard controls, if we didn't consume input
 	return (bHandled) ? true : FEditorViewportClient::InputKey(Viewport, ControllerId, Key, Event, AmountDepressed, bGamepad);
 }
diff --git a/Engine/Plugins/2D/Paper2D/Source/Paper2DEditor/Private/FlipbookEditor/FlipbookEditorViewportClient.h b/Engine/Plugins/2D/Paper2D/Source/Paper2DEditor/Private/FlipbookEditor/FlipbookEditorViewportClient.h
--- a/Engine/Plugins/2D/Paper2D/Source/Paper2DEditor/Private/FlipbookEditor/FlipbookEditorViewportClient.h
+++ b/Engine/Plugins/2D/Paper2D/Source/Paper2DEditor/Private/FlipbookEditor/FlipbookEditorViewportClient.h
@@ -42,6 +42,11 @@ public:
 		return AnimatedRenderComponent.Get();
 	}
 
+protected:
+	// FPaperEditorViewportClient interface
+	virtual FBox GetDesiredFocusBounds() const override;
+	// End of FPaperEditorViewportClient interface
+
 private:
 
 	// The preview scene
@@ -61,7 +66,4 @@ private:
 
 	// Should we show sockets?
 	bool bShowSockets;
-
-	// Should we zoom to the sprite next tick?
-	bool bDeferZoomToSprite;
 };
diff --git a/Engine/Plugins/2D/Paper2D/Source/Paper2DEditor/Private/PaperEditorViewportClient.cpp b/Engine/Plugins/2D/Paper2D/Source/Paper2DEditor/Private/PaperEditorViewportClient.cpp
--- a/Engine/Plugins/2D/Paper2D/Source/Paper2DEditor/Private/PaperEditorViewportClient.cpp
+++ b/Engine/Plugins/2D/Paper2D/Source/Paper2DEditor/Private/PaperEditorViewportClient.cpp
@@ -47,6 +47,12 @@ FPaperEditorViewportClient::FPaperEditorViewportClient(const TWeakPtr<SEditorVie
 	}
 	SetViewModes(VMI_Lit, VMI_Lit);
 	SetViewportType(NewViewportType);
+
+	bDeferZoomToSprite = true;
+	bDeferZoomToSpriteIsInstant = true;
+
+	// Get the correct general direction of the perspective mode; the distance doesn't matter much as we've queued up a deferred zoom that will calculate a much better distance
+ 	SetInitialViewTransform(LVT_Perspective, -100.0f * PaperAxisZ, PaperAxisZ.Rotation(), 0.0f);
 }
 
 FPaperEditorViewportClient::~FPaperEditorViewportClient()
@@ -64,6 +70,27 @@ void FPaperEditorViewportClient::Draw(FViewport* Viewport, FCanvas* Canvas)
 	Canvas->Clear(GetBackgroundColor());
 }
 
+void FPaperEditorViewportClient::Tick(float DeltaSeconds)
+{
+	// Zoom in on the sprite
+	//@TODO: Fix this properly so it doesn't need to be deferred, or wait for the viewport to initialize
+	FIntPoint Size = Viewport->GetSizeXY();
+	if (bDeferZoomToSprite && (Size.X > 0) && (Size.Y > 0))
+	{
+		FBox BoundsToFocus = GetDesiredFocusBounds();
+		if (ViewportType != LVT_Perspective)
+		{
+			TGuardValue<ELevelViewportType> SaveViewportType(ViewportType, LVT_Perspective);
+			FocusViewportOnBox(BoundsToFocus, bDeferZoomToSpriteIsInstant);
+		}
+
+		FocusViewportOnBox(BoundsToFocus, bDeferZoomToSpriteIsInstant);
+		bDeferZoomToSprite = false;
+	}
+
+	FEditorViewportClient::Tick(DeltaSeconds);
+}
+
 void FPaperEditorViewportClient::DrawSelectionRectangles(FViewport* Viewport, FCanvas* Canvas)
 {
 	for (auto RectangleIt = SelectionRectangles.CreateConstIterator(); RectangleIt; ++RectangleIt)
@@ -87,6 +114,13 @@ void FPaperEditorViewportClient::AddReferencedObjects(FReferenceCollector& Colle
 	Collector.AddReferencedObject(CheckerboardTexture);
 }
 
+// Called to request a focus on the current selection
+void FPaperEditorViewportClient::RequestFocusOnSelection(bool bInstant)
+{
+	bDeferZoomToSprite = true;
+	bDeferZoomToSpriteIsInstant = bInstant;
+}
+
 void FPaperEditorViewportClient::ModifyCheckerboardTextureColors()
 {
 	const FColor ColorOne = FColor(128, 128, 128);//TextureEditorPtr.Pin()->GetCheckeredBackground_ColorOne();
diff --git a/Engine/Plugins/2D/Paper2D/Source/Paper2DEditor/Private/PaperEditorViewportClient.h b/Engine/Plugins/2D/Paper2D/Source/Paper2DEditor/Private/PaperEditorViewportClient.h
--- a/Engine/Plugins/2D/Paper2D/Source/Paper2DEditor/Private/PaperEditorViewportClient.h
+++ b/Engine/Plugins/2D/Paper2D/Source/Paper2DEditor/Private/PaperEditorViewportClient.h
@@ -24,6 +24,7 @@ public:
 
 	// FViewportClient interface
 	virtual void Draw(FViewport* Viewport, FCanvas* Canvas) override;
+	virtual void Tick(float DeltaSeconds) override;
 	// End of FViewportClient interface
 
 	// FEditorViewportClient interface
@@ -34,6 +35,9 @@ public:
 	virtual void AddReferencedObjects(FReferenceCollector& Collector) override;
 	// End of FSerializableObject interface
 
+	// Called to request a focus on the current selection
+	virtual void RequestFocusOnSelection(bool bInstant);
+
 	/** Modifies the checkerboard texture's data */
 	void ModifyCheckerboardTextureColors();
 
@@ -55,9 +59,19 @@ private:
 protected:
 	void DrawSelectionRectangles(FViewport* Viewport, FCanvas* Canvas);
 
+	virtual FBox GetDesiredFocusBounds() const
+	{
+		return FBox(ForceInitToZero);
+	}
+
 protected:
 	/** Checkerboard texture */
 	UTexture2D* CheckerboardTexture;
 	FVector2D ZoomPos;
 	float ZoomAmount;
+
+private:
+	// Should we zoom to the focus bounds next tick?
+	bool bDeferZoomToSprite;
+	bool bDeferZoomToSpriteIsInstant;
 };
diff --git a/Engine/Plugins/2D/Paper2D/Source/Paper2DEditor/Private/SpriteEditor/SpriteEditor.cpp b/Engine/Plugins/2D/Paper2D/Source/Paper2DEditor/Private/SpriteEditor/SpriteEditor.cpp
--- a/Engine/Plugins/2D/Paper2D/Source/Paper2DEditor/Private/SpriteEditor/SpriteEditor.cpp
+++ b/Engine/Plugins/2D/Paper2D/Source/Paper2DEditor/Private/SpriteEditor/SpriteEditor.cpp
@@ -56,6 +56,7 @@ public:
 	virtual TSharedRef<FEditorViewportClient> MakeEditorViewportClient() override;
 	virtual TSharedPtr<SWidget> MakeViewportToolbar() override;
 	virtual EVisibility GetTransformToolbarVisibility() const override;
+	virtual void OnFocusViewportToSelection() override;
 	// End of SEditorViewport interface
 
 	// ICommonEditorViewportToolbarInfoProvider interface
@@ -174,11 +175,6 @@ void SSpriteEditorViewport::BindCommands()
 		FCanExecuteAction(),
 		FIsActionChecked::CreateSP( EditorViewportClientRef, &FSpriteEditorViewportClient::IsInRenderingEditMode ) );
 
-	// Misc. actions
-	CommandList->MapAction(
-		Commands.FocusOnSprite,
-		FExecuteAction::CreateSP( EditorViewportClientRef, &FSpriteEditorViewportClient::FocusOnSprite ));
-
 	// Geometry editing commands
 	CommandList->MapAction(
 		Commands.DeleteSelection,
@@ -232,6 +228,11 @@ EVisibility SSpriteEditorViewport::GetTransformToolbarVisibility() const
 	return EVisibility::Visible;
 }
 
+void SSpriteEditorViewport::OnFocusViewportToSelection()
+{
+	EditorViewportClient->RequestFocusOnSelection(/*bInstant=*/ false);
+}
+
 TSharedRef<class SEditorViewport> SSpriteEditorViewport::GetViewportWidget()
 {
 	return SharedThis(this);
diff --git a/Engine/Plugins/2D/Paper2D/Source/Paper2DEditor/Private/SpriteEditor/SpriteEditorCommands.cpp b/Engine/Plugins/2D/Paper2D/Source/Paper2DEditor/Private/SpriteEditor/SpriteEditorCommands.cpp
--- a/Engine/Plugins/2D/Paper2D/Source/Paper2DEditor/Private/SpriteEditor/SpriteEditorCommands.cpp
+++ b/Engine/Plugins/2D/Paper2D/Source/Paper2DEditor/Private/SpriteEditor/SpriteEditorCommands.cpp
@@ -28,9 +28,6 @@ void FSpriteEditorCommands::RegisterCommands()
 	UI_COMMAND(EnterCollisionEditMode, "Edit Collision", "Edit the collision geometry.", EUserInterfaceActionType::ToggleButton, FInputChord());
 	UI_COMMAND(EnterRenderingEditMode, "Edit RenderGeom", "Edit the rendering geometry (useful to reduce overdraw).", EUserInterfaceActionType::ToggleButton, FInputChord());
 
-	// Misc. actions
-	UI_COMMAND(FocusOnSprite, "Focus on sprite", "Centers and zooms the view to focus on the current sprite.", EUserInterfaceActionType::Button, FInputChord(EKeys::F, EModifierKey::None));
-
 	// Geometry editing commands
 	UI_COMMAND(DeleteSelection, "Delete", "Delete the selection.", EUserInterfaceActionType::Button, FInputChord(EKeys::Platform_Delete, EModifierKey::None));
 	UI_COMMAND(SplitEdge, "Split", "Split edge.", EUserInterfaceActionType::Button, FInputChord(EKeys::Insert, EModifierKey::None));
diff --git a/Engine/Plugins/2D/Paper2D/Source/Paper2DEditor/Private/SpriteEditor/SpriteEditorCommands.h b/Engine/Plugins/2D/Paper2D/Source/Paper2DEditor/Private/SpriteEditor/SpriteEditorCommands.h
--- a/Engine/Plugins/2D/Paper2D/Source/Paper2DEditor/Private/SpriteEditor/SpriteEditorCommands.h
+++ b/Engine/Plugins/2D/Paper2D/Source/Paper2DEditor/Private/SpriteEditor/SpriteEditorCommands.h
@@ -40,9 +40,6 @@ public:
 	TSharedPtr<FUICommandInfo> EnterCollisionEditMode;
 	TSharedPtr<FUICommandInfo> EnterRenderingEditMode;
 
-	// Misc. actions
-	TSharedPtr<FUICommandInfo> FocusOnSprite;
-
 	// Geometry editing commands
 	TSharedPtr<FUICommandInfo> DeleteSelection;
 	TSharedPtr<FUICommandInfo> SplitEdge;
diff --git a/Engine/Plugins/2D/Paper2D/Source/Paper2DEditor/Private/SpriteEditor/SpriteEditorViewportClient.cpp b/Engine/Plugins/2D/Paper2D/Source/Paper2DEditor/Private/SpriteEditor/SpriteEditorViewportClient.cpp
--- a/Engine/Plugins/2D/Paper2D/Source/Paper2DEditor/Private/SpriteEditor/SpriteEditorViewportClient.cpp
+++ b/Engine/Plugins/2D/Paper2D/Source/Paper2DEditor/Private/SpriteEditor/SpriteEditorViewportClient.cpp
@@ -78,8 +78,6 @@ FSpriteEditorViewportClient::FSpriteEditorViewportClient(TWeakPtr<FSpriteEditor>
 	bShowPivot = true;
 	bShowRelatedSprites = true;
 
-	bDeferZoomToSprite = true;
-
 	bIsMarqueeTracking = false;
 
 	DrawHelper.bDrawGrid = false;
@@ -131,7 +129,7 @@ void FSpriteEditorViewportClient::UpdateSourceTextureSpriteFromSprite(UPaperSpri
 			TargetSprite->PixelsPerUnrealUnit = SourceSprite->PixelsPerUnrealUnit;
 			TargetSprite->InitializeSprite(SpriteReinitParams);
 
-			bDeferZoomToSprite = true;
+			RequestFocusOnSelection(/*bInstant=*/ true);
 		}
 
 
@@ -687,9 +685,16 @@ void FSpriteEditorViewportClient::Draw(const FSceneView* View, FPrimitiveDrawInt
 
 void FSpriteEditorViewportClient::Draw(FViewport* Viewport, FCanvas* Canvas)
 {
+	// Skipping the parent on purpose
 	FEditorViewportClient::Draw(Viewport, Canvas);
 }
 
+FBox FSpriteEditorViewportClient::GetDesiredFocusBounds() const
+{
+	UPaperSpriteComponent* ComponentToFocusOn = SourceTextureViewComponent->IsVisible() ? SourceTextureViewComponent : RenderSpriteComponent;
+	return ComponentToFocusOn->Bounds.GetBox();
+}
+
 void FSpriteEditorViewportClient::Tick(float DeltaSeconds)
 {
 	if (UPaperSprite* Sprite = GetSpriteBeingEdited())
@@ -707,26 +712,16 @@ void FSpriteEditorViewportClient::Tick(float DeltaSeconds)
 		bool bSourceTextureViewComponentVisibility = bShowSourceTexture || IsInSourceRegionEditMode();
 		if (bSourceTextureViewComponentVisibility != SourceTextureViewComponent->IsVisible())
 		{
-			bDeferZoomToSprite = true;
+			RequestFocusOnSelection(/*bInstant=*/ true);
 			SourceTextureViewComponent->SetVisibility(bSourceTextureViewComponentVisibility);
 		}
 
 		bool bRenderTextureViewComponentVisibility = !IsInSourceRegionEditMode();
 		if (bRenderTextureViewComponentVisibility != RenderSpriteComponent->IsVisible())
 		{
-			bDeferZoomToSprite = true;
+			RequestFocusOnSelection(/*bInstant=*/ true);
 			RenderSpriteComponent->SetVisibility(bRenderTextureViewComponentVisibility);
 		}
-
-		// Zoom in on the sprite
-		//@TODO: Fix this properly so it doesn't need to be deferred, or wait for the viewport to initialize
-		FIntPoint Size = Viewport->GetSizeXY();
-		if (bDeferZoomToSprite && (Size.X > 0) && (Size.Y > 0))
-		{
-			UPaperSpriteComponent* ComponentToFocusOn = SourceTextureViewComponent->IsVisible() ? SourceTextureViewComponent : RenderSpriteComponent;
-			FocusViewportOnBox(ComponentToFocusOn->Bounds.GetBox(), true);
-			bDeferZoomToSprite = false;
-		}		
 	}
 
 	if (bIsMarqueeTracking)
@@ -1325,7 +1320,7 @@ void FSpriteEditorViewportClient::NotifySpriteBeingEditedHasChanged()
 	InternalActivateNewMode(CurrentMode);
 
 	//@TODO: Only do this if the sprite isn't visible (may consider doing a flashing pulse around the source region rect?)
-	bDeferZoomToSprite = true;
+	RequestFocusOnSelection(/*bInstant=*/ true);
 }
 
 void FSpriteEditorViewportClient::FocusOnSprite()
diff --git a/Engine/Plugins/2D/Paper2D/Source/Paper2DEditor/Private/SpriteEditor/SpriteEditorViewportClient.h b/Engine/Plugins/2D/Paper2D/Source/Paper2DEditor/Private/SpriteEditor/SpriteEditorViewportClient.h
--- a/Engine/Plugins/2D/Paper2D/Source/Paper2DEditor/Private/SpriteEditor/SpriteEditorViewportClient.h
+++ b/Engine/Plugins/2D/Paper2D/Source/Paper2DEditor/Private/SpriteEditor/SpriteEditorViewportClient.h
@@ -120,6 +120,12 @@ public:
 	}
 
 	static void AnalyzeSpriteMaterialType(UPaperSprite* Sprite, int32& OutNumOpaque, int32& OutNumMasked, int32& OutNumTranslucent);
+
+protected:
+	// FPaperEditorViewportClient interface
+	virtual FBox GetDesiredFocusBounds() const override;
+	// End of FPaperEditorViewportClient interface
+
 private:
 	// Editor mode
 	ESpriteEditorMode::Type CurrentMode;
@@ -160,9 +166,6 @@ private:
 	// Should we show the sprite pivot?
 	bool bShowPivot;
 
-	// Should we zoom to the sprite next tick?
-	bool bDeferZoomToSprite;
-
 	// Should we show related sprites in the source texture?
 	bool bShowRelatedSprites;
 
diff --git a/Engine/Plugins/2D/Paper2D/Source/Paper2DEditor/Private/TileMapEditing/TileMapEditor.cpp b/Engine/Plugins/2D/Paper2D/Source/Paper2DEditor/Private/TileMapEditing/TileMapEditor.cpp
--- a/Engine/Plugins/2D/Paper2D/Source/Paper2DEditor/Private/TileMapEditing/TileMapEditor.cpp
+++ b/Engine/Plugins/2D/Paper2D/Source/Paper2DEditor/Private/TileMapEditing/TileMapEditor.cpp
@@ -54,6 +54,7 @@ public:
 	virtual TSharedRef<FEditorViewportClient> MakeEditorViewportClient() override;
 	virtual TSharedPtr<SWidget> MakeViewportToolbar() override;
 	virtual EVisibility GetTransformToolbarVisibility() const override;
+	virtual void OnFocusViewportToSelection() override;
 	// End of SEditorViewport interface
 
 	// ICommonEditorViewportToolbarInfoProvider interface
@@ -111,11 +112,6 @@ void STileMapEditorViewport::BindCommands()
 		FExecuteAction::CreateSP(EditorViewportClientRef, &FTileMapEditorViewportClient::ToggleShowPivot),
 		FCanExecuteAction(),
 		FIsActionChecked::CreateSP(EditorViewportClientRef, &FTileMapEditorViewportClient::IsShowPivotChecked));
-
-	// Misc. actions
-	CommandList->MapAction(
-		Commands.FocusOnTileMap,
-		FExecuteAction::CreateSP(EditorViewportClientRef, &FTileMapEditorViewportClient::FocusOnTileMap));
 }
 
 TSharedRef<FEditorViewportClient> STileMapEditorViewport::MakeEditorViewportClient()
@@ -142,6 +138,11 @@ EVisibility STileMapEditorViewport::GetTransformToolbarVisibility() const
 	return EVisibility::Visible;
 }
 
+void STileMapEditorViewport::OnFocusViewportToSelection()
+{
+	EditorViewportClient->RequestFocusOnSelection(/*bInstant=*/ false);
+}
+
 TSharedRef<class SEditorViewport> STileMapEditorViewport::GetViewportWidget()
 {
 	return SharedThis(this);
diff --git a/Engine/Plugins/2D/Paper2D/Source/Paper2DEditor/Private/TileMapEditing/TileMapEditorCommands.cpp b/Engine/Plugins/2D/Paper2D/Source/Paper2DEditor/Private/TileMapEditing/TileMapEditorCommands.cpp
--- a/Engine/Plugins/2D/Paper2D/Source/Paper2DEditor/Private/TileMapEditing/TileMapEditorCommands.cpp
+++ b/Engine/Plugins/2D/Paper2D/Source/Paper2DEditor/Private/TileMapEditing/TileMapEditorCommands.cpp
@@ -21,9 +21,6 @@ void FTileMapEditorCommands::RegisterCommands()
 
 	UI_COMMAND(SetShowPivot, "Pivot", "Display the pivot location of the sprite.", EUserInterfaceActionType::ToggleButton, FInputChord());
 
-	// Misc. actions
-	UI_COMMAND(FocusOnTileMap, "Focus on tile map", "Centers and zooms the view to focus on the tile map.", EUserInterfaceActionType::Button, FInputChord(EKeys::F, EModifierKey::None));
-
 	// Selection actions
 	UI_COMMAND(FlipSelectionHorizontally, "Flip selection horizontally", "Flips the selection horizontally", EUserInterfaceActionType::Button, FInputChord(EKeys::X));
 	UI_COMMAND(FlipSelectionVertically, "Flip selection vertically", "Flips the selection vertically", EUserInterfaceActionType::Button, FInputChord(EKeys::Y));
diff --git a/Engine/Plugins/2D/Paper2D/Source/Paper2DEditor/Private/TileMapEditing/TileMapEditorCommands.h b/Engine/Plugins/2D/Paper2D/Source/Paper2DEditor/Private/TileMapEditing/TileMapEditorCommands.h
--- a/Engine/Plugins/2D/Paper2D/Source/Paper2DEditor/Private/TileMapEditing/TileMapEditorCommands.h
+++ b/Engine/Plugins/2D/Paper2D/Source/Paper2DEditor/Private/TileMapEditing/TileMapEditorCommands.h
@@ -35,9 +35,6 @@ public:
 	
 	TSharedPtr<FUICommandInfo> SetShowPivot;
 
-	// Misc. actions
-	TSharedPtr<FUICommandInfo> FocusOnTileMap;
-
 	// Selection actions
 	TSharedPtr<FUICommandInfo> FlipSelectionHorizontally;
 	TSharedPtr<FUICommandInfo> FlipSelectionVertically;
diff --git a/Engine/Plugins/2D/Paper2D/Source/Paper2DEditor/Private/TileMapEditing/TileMapEditorViewportClient.cpp b/Engine/Plugins/2D/Paper2D/Source/Paper2DEditor/Private/TileMapEditing/TileMapEditorViewportClient.cpp
--- a/Engine/Plugins/2D/Paper2D/Source/Paper2DEditor/Private/TileMapEditing/TileMapEditorViewportClient.cpp
+++ b/Engine/Plugins/2D/Paper2D/Source/Paper2DEditor/Private/TileMapEditing/TileMapEditorViewportClient.cpp
@@ -38,8 +38,6 @@ FTileMapEditorViewportClient::FTileMapEditorViewportClient(TWeakPtr<FTileMapEdit
 	DrawHelper.bDrawPivot = false;
 	bShowPivot = true;
 
-	bDeferZoomToTileMap = true;
-
 	EngineShowFlags.DisableAdvancedFeatures();
 	EngineShowFlags.CompositeEditorPrimitives = true;
 
@@ -127,19 +125,6 @@ void FTileMapEditorViewportClient::Draw(FViewport* Viewport, FCanvas* Canvas)
 
 void FTileMapEditorViewportClient::Tick(float DeltaSeconds)
 {
-	if (UPaperTileMap* TileMap = RenderTileMapComponent->TileMap)
-	{
-		// Zoom in on the tile map
-		//@TODO: Fix this properly so it doesn't need to be deferred, or wait for the viewport to initialize
-		FIntPoint Size = Viewport->GetSizeXY();
-		if (bDeferZoomToTileMap && (Size.X > 0) && (Size.Y > 0))
-		{
-			UPaperTileMapComponent* ComponentToFocusOn = RenderTileMapComponent;
-			FocusViewportOnBox(ComponentToFocusOn->Bounds.GetBox(), true);
-			bDeferZoomToTileMap = false;
-		}		
-	}
-
 	FPaperEditorViewportClient::Tick(DeltaSeconds);
 
 	if (!GIntraFrameDebuggingGameThread)
@@ -191,23 +176,17 @@ void FTileMapEditorViewportClient::NotifyTileMapBeingEditedHasChanged()
 {
 	//@TODO: Ideally we do this before switching
 	EndTransaction();
-	ClearSelectionSet();
 
 	// Update components to know about the new tile map being edited
 	UPaperTileMap* TileMap = GetTileMapBeingEdited();
 	RenderTileMapComponent->TileMap = TileMap;
 
-	//
-	bDeferZoomToTileMap = true;
-}
-
-void FTileMapEditorViewportClient::FocusOnTileMap()
-{
-	FocusViewportOnBox(RenderTileMapComponent->Bounds.GetBox());
+	RequestFocusOnSelection(/*bInstant=*/ true);
 }
 
-void FTileMapEditorViewportClient::ClearSelectionSet()
+FBox FTileMapEditorViewportClient::GetDesiredFocusBounds() const
 {
+	return RenderTileMapComponent->Bounds.GetBox();
 }
 
 //////////////////////////////////////////////////////////////////////////
diff --git a/Engine/Plugins/2D/Paper2D/Source/Paper2DEditor/Private/TileMapEditing/TileMapEditorViewportClient.h b/Engine/Plugins/2D/Paper2D/Source/Paper2DEditor/Private/TileMapEditing/TileMapEditorViewportClient.h
--- a/Engine/Plugins/2D/Paper2D/Source/Paper2DEditor/Private/TileMapEditing/TileMapEditorViewportClient.h
+++ b/Engine/Plugins/2D/Paper2D/Source/Paper2DEditor/Private/TileMapEditing/TileMapEditorViewportClient.h
@@ -72,8 +72,11 @@ private:
 	// Should we show the sprite pivot?
 	bool bShowPivot;
 
-	// Should we zoom to the tile map next tick?
-	bool bDeferZoomToTileMap;
+protected:
+	// FPaperEditorViewportClient interface
+	virtual FBox GetDesiredFocusBounds() const override;
+	// End of FPaperEditorViewportClient interface
+
 private:
 	UPaperTileMap* GetTileMapBeingEdited() const
 	{
@@ -84,6 +87,4 @@ private:
 	
 	void BeginTransaction(const FText& SessionName);
 	void EndTransaction();
-
-	void ClearSelectionSet();
 };
