Flipbook Editor: Enabled the Transform toolbar, mainly to get camera speed control

diff --git a/Engine/Plugins/2D/Paper2D/Source/Paper2DEditor/Private/FlipbookEditor/FlipbookEditor.cpp b/Engine/Plugins/2D/Paper2D/Source/Paper2DEditor/Private/FlipbookEditor/FlipbookEditor.cpp
--- a/Engine/Plugins/2D/Paper2D/Source/Paper2DEditor/Private/FlipbookEditor/FlipbookEditor.cpp
+++ b/Engine/Plugins/2D/Paper2D/Source/Paper2DEditor/Private/FlipbookEditor/FlipbookEditor.cpp
@@ -151,8 +151,7 @@ TSharedPtr<SWidget> SFlipbookEditorViewport::MakeViewportToolbar()
 
 EVisibility SFlipbookEditorViewport::GetTransformToolbarVisibility() const
 {
-	//@TODO: Nothing to transform in the flipbook editor, yet...
-	return EVisibility::Hidden;
+	return EVisibility::Visible;
 }
 
 void SFlipbookEditorViewport::OnFocusViewportToSelection()
