Tile Map Editor: Added additional layer commands and keyboard shortcuts
- WIP: They can't be triggered yet because there is no path to an existing command list from the details panel
- Ctrl+Shift+N creates a new layer
- Ctrl+[ sends the current layer backward
- Ctrl+Shift+[ sends the current layer to the back
- Ctrl+] brings the current layer forward
- Ctrl+Shift+] brings the current layer to the front
- Alt+[ selects the previous layer
- Alt+] selects the next layer
[UE-13176]

diff --git a/Engine/Plugins/2D/Paper2D/Source/Paper2DEditor/Private/TileMapEditing/PaperTileMapDetailsCustomization.cpp b/Engine/Plugins/2D/Paper2D/Source/Paper2DEditor/Private/TileMapEditing/PaperTileMapDetailsCustomization.cpp
--- a/Engine/Plugins/2D/Paper2D/Source/Paper2DEditor/Private/TileMapEditing/PaperTileMapDetailsCustomization.cpp
+++ b/Engine/Plugins/2D/Paper2D/Source/Paper2DEditor/Private/TileMapEditing/PaperTileMapDetailsCustomization.cpp
@@ -133,6 +133,11 @@ void FPaperTileMapDetailsCustomization::CustomizeDetails(IDetailLayoutBuilder& D
 		TileMapCategory.AddProperty(GET_MEMBER_NAME_CHECKED(UPaperTileMapComponent, TileMap));
 	}
 
+	//@TODO: No clue where to route this to a toolkit one, and anywhere I can put an OnKeyDown around here
+	// or in the STileLayerList to route it myself isn't good enough, since the editor viewport will have focus
+	// 99% of the time, not the details panel.
+	TSharedPtr<FUICommandList> CommandList = MakeShareable(new FUICommandList);
+
 	// Add the layer browser
 	if (TileMap != nullptr)
 	{
@@ -157,7 +162,7 @@ void FPaperTileMapDetailsCustomization::CustomizeDetails(IDetailLayoutBuilder& D
 			]
 			+SVerticalBox::Slot()
 			[
-				SNew(STileLayerList, TileMap, NotifyHook)
+				SNew(STileLayerList, TileMap, NotifyHook, CommandList)
 			]
 		];
 	}
diff --git a/Engine/Plugins/2D/Paper2D/Source/Paper2DEditor/Private/TileMapEditing/STileLayerList.cpp b/Engine/Plugins/2D/Paper2D/Source/Paper2DEditor/Private/TileMapEditing/STileLayerList.cpp
--- a/Engine/Plugins/2D/Paper2D/Source/Paper2DEditor/Private/TileMapEditing/STileLayerList.cpp
+++ b/Engine/Plugins/2D/Paper2D/Source/Paper2DEditor/Private/TileMapEditing/STileLayerList.cpp
@@ -14,7 +14,7 @@
 //////////////////////////////////////////////////////////////////////////
 // STileLayerList
 
-void STileLayerList::Construct(const FArguments& InArgs, UPaperTileMap* InTileMap, FNotifyHook* InNotifyHook)
+void STileLayerList::Construct(const FArguments& InArgs, UPaperTileMap* InTileMap, FNotifyHook* InNotifyHook, TSharedPtr<class FUICommandList> InCommandList)
 {
 	TileMapPtr = InTileMap;
 	NotifyHook = InNotifyHook;
@@ -22,7 +22,7 @@ void STileLayerList::Construct(const FArguments& InArgs, UPaperTileMap* InTileMa
 	FTileMapEditorCommands::Register();
 	const FTileMapEditorCommands& Commands = FTileMapEditorCommands::Get();
 
-	CommandList = MakeShareable(new FUICommandList);
+	CommandList = InCommandList;
 
 	CommandList->MapAction(
 		Commands.AddNewLayerAbove,
@@ -49,14 +49,32 @@ void STileLayerList::Construct(const FArguments& InArgs, UPaperTileMap* InTileMa
 
 	CommandList->MapAction(
 		Commands.MoveLayerUp,
-		FExecuteAction::CreateSP(this, &STileLayerList::MoveLayerUp),
+		FExecuteAction::CreateSP(this, &STileLayerList::MoveLayerUp, /*bForceToTop=*/ false),
 		FCanExecuteAction::CreateSP(this, &STileLayerList::CanExecuteActionNeedingLayerAbove));
 
 	CommandList->MapAction(
 		Commands.MoveLayerDown,
-		FExecuteAction::CreateSP(this, &STileLayerList::MoveLayerDown),
+		FExecuteAction::CreateSP(this, &STileLayerList::MoveLayerDown, /*bForceToBottom=*/ false),
 		FCanExecuteAction::CreateSP(this, &STileLayerList::CanExecuteActionNeedingLayerBelow));
-	
+
+	CommandList->MapAction(
+		Commands.MoveLayerToTop,
+		FExecuteAction::CreateSP(this, &STileLayerList::MoveLayerUp, /*bForceToTop=*/ true),
+		FCanExecuteAction::CreateSP(this, &STileLayerList::CanExecuteActionNeedingLayerAbove));
+
+	CommandList->MapAction(
+		Commands.MoveLayerToBottom,
+		FExecuteAction::CreateSP(this, &STileLayerList::MoveLayerDown, /*bForceToBottom=*/ true),
+		FCanExecuteAction::CreateSP(this, &STileLayerList::CanExecuteActionNeedingLayerBelow));
+
+	CommandList->MapAction(
+		Commands.SelectLayerAbove,
+		FExecuteAction::CreateSP(this, &STileLayerList::SelectLayerAbove, /*bTopmost=*/ false));
+
+	CommandList->MapAction(
+		Commands.SelectLayerBelow,
+		FExecuteAction::CreateSP(this, &STileLayerList::SelectLayerBelow, /*bBottommost=*/ false));
+
 	FToolBarBuilder ToolbarBuilder(CommandList, FMultiBoxCustomization("TileLayerBrowserToolbar"), TSharedPtr<FExtender>(), Orient_Horizontal, /*InForceSmallIcons=*/ true);
 	ToolbarBuilder.SetLabelVisibility(EVisibility::Collapsed);
 
@@ -353,20 +371,48 @@ void STileLayerList::MergeLayerDown()
 	}
 }
 
-void STileLayerList::MoveLayerUp()
+void STileLayerList::MoveLayerUp(bool bForceToTop)
 {
 	const int32 SelectedIndex = GetSelectionIndex();
-	const int32 NewIndex = SelectedIndex - 1;
+	const int32 NewIndex = bForceToTop ? 0 : (SelectedIndex - 1);
 	ChangeLayerOrdering(SelectedIndex, NewIndex);
 }
 
-void STileLayerList::MoveLayerDown()
+void STileLayerList::MoveLayerDown(bool bForceToBottom)
 {
 	const int32 SelectedIndex = GetSelectionIndex();
-	const int32 NewIndex = SelectedIndex + 1;
+	const int32 NewIndex = bForceToBottom ? (GetNumLayers() - 1) : (SelectedIndex + 1);
 	ChangeLayerOrdering(SelectedIndex, NewIndex);
 }
 
+void STileLayerList::SelectLayerAbove(bool bTopmost)
+{
+	const int32 SelectedIndex = GetSelectionIndex();
+	const int32 NumLayers = GetNumLayers();
+	const int32 NewIndex = bTopmost ? 0 : ((NumLayers + SelectedIndex - 1) % NumLayers);
+	SetSelectedLayerIndex(NewIndex);
+}
+
+void STileLayerList::SelectLayerBelow(bool bBottommost)
+{
+	const int32 SelectedIndex = GetSelectionIndex();
+	const int32 NumLayers = GetNumLayers();
+	const int32 NewIndex = bBottommost ? (NumLayers - 1) : ((SelectedIndex + 1) % NumLayers);
+	SetSelectedLayerIndex(NewIndex);
+}
+
+void STileLayerList::SetSelectedLayerIndex(int32 NewIndex)
+{
+	if (UPaperTileMap* TileMap = TileMapPtr.Get())
+	{
+		if (TileMap->TileLayers.IsValidIndex(NewIndex))
+		{
+			SetSelectedLayer(TileMap->TileLayers[NewIndex]);
+			PostEditNotfications();
+		}
+	}
+}
+
 int32 STileLayerList::GetNumLayers() const
 {
 	if (UPaperTileMap* TileMap = TileMapPtr.Get())
@@ -424,7 +470,7 @@ TSharedPtr<SWidget> STileLayerList::OnConstructContextMenu()
 
 	const FTileMapEditorCommands& Commands = FTileMapEditorCommands::Get();
 
- 	MenuBuilder.BeginSection("BasicOperations");
+	MenuBuilder.BeginSection("BasicOperations", LOCTEXT("BasicOperationsHeader", "Layer actions"));
  	{
 		MenuBuilder.AddMenuEntry(Commands.DuplicateLayer);
 		MenuBuilder.AddMenuEntry(Commands.DeleteLayer);
diff --git a/Engine/Plugins/2D/Paper2D/Source/Paper2DEditor/Private/TileMapEditing/STileLayerList.h b/Engine/Plugins/2D/Paper2D/Source/Paper2DEditor/Private/TileMapEditing/STileLayerList.h
--- a/Engine/Plugins/2D/Paper2D/Source/Paper2DEditor/Private/TileMapEditing/STileLayerList.h
+++ b/Engine/Plugins/2D/Paper2D/Source/Paper2DEditor/Private/TileMapEditing/STileLayerList.h
@@ -11,7 +11,7 @@ public:
 	SLATE_BEGIN_ARGS(STileLayerList) {}
 	SLATE_END_ARGS()
 
-	void Construct(const FArguments& InArgs, UPaperTileMap* TileMap, FNotifyHook* InNotifyHook);
+	void Construct(const FArguments& InArgs, UPaperTileMap* TileMap, FNotifyHook* InNotifyHook, TSharedPtr<class FUICommandList> InCommandList);
 
 protected:
 	typedef TSharedPtr<int32> FMirrorEntry;
@@ -44,8 +44,12 @@ protected:
 	void DeleteLayer();
 	void DuplicateLayer();
 	void MergeLayerDown();
-	void MoveLayerUp();
-	void MoveLayerDown();
+	void MoveLayerUp(bool bForceToTop);
+	void MoveLayerDown(bool bForceToBottom);
+	void SelectLayerAbove(bool bTopmost);
+	void SelectLayerBelow(bool bBottommost);
+
+	void SetSelectedLayerIndex(int32 NewIndex);
 
 	int32 GetNumLayers() const;
 	bool CanExecuteActionNeedingLayerAbove() const;
diff --git a/Engine/Plugins/2D/Paper2D/Source/Paper2DEditor/Private/TileMapEditing/TileMapEditorCommands.cpp b/Engine/Plugins/2D/Paper2D/Source/Paper2DEditor/Private/TileMapEditing/TileMapEditorCommands.cpp
--- a/Engine/Plugins/2D/Paper2D/Source/Paper2DEditor/Private/TileMapEditing/TileMapEditorCommands.cpp
+++ b/Engine/Plugins/2D/Paper2D/Source/Paper2DEditor/Private/TileMapEditing/TileMapEditorCommands.cpp
@@ -28,13 +28,17 @@ void FTileMapEditorCommands::RegisterCommands()
 	UI_COMMAND(RotateSelectionCCW, "Rotate selection counterclockwise", "Rotates the selection counterclockwise", EUserInterfaceActionType::Button, FInputChord(EKeys::Z, EModifierKey::Shift));
 
 	// Layer actions
-	UI_COMMAND(AddNewLayerAbove, "Add new layer", "Add a new layer above the current layer", EUserInterfaceActionType::Button, FInputChord());
+	UI_COMMAND(AddNewLayerAbove, "Add new layer", "Add a new layer above the current layer", EUserInterfaceActionType::Button, FInputChord(EKeys::N, EModifierKey::Control | EModifierKey::Shift));
 	UI_COMMAND(AddNewLayerBelow, "Add new layer below", "Add a new layer below the current layer", EUserInterfaceActionType::Button, FInputChord());
 	UI_COMMAND(DeleteLayer, "Delete layer", "Delete the current layer", EUserInterfaceActionType::Button, FInputChord());
 	UI_COMMAND(DuplicateLayer, "Duplicate layer", "Duplicate the current layer", EUserInterfaceActionType::Button, FInputChord());
 	UI_COMMAND(MergeLayerDown, "Merge layer down", "Merge the current layer down to the layer below it", EUserInterfaceActionType::Button, FInputChord());
-	UI_COMMAND(MoveLayerUp, "Move layer up", "Move layer up", EUserInterfaceActionType::Button, FInputChord());
-	UI_COMMAND(MoveLayerDown, "Move layer down", "Move layer down", EUserInterfaceActionType::Button, FInputChord());
+	UI_COMMAND(MoveLayerUp, "Bring forward", "Bring forward", EUserInterfaceActionType::Button, FInputChord(EKeys::RightBracket, EModifierKey::Control));
+	UI_COMMAND(MoveLayerDown, "Send backward", "Send backward", EUserInterfaceActionType::Button, FInputChord(EKeys::LeftBracket, EModifierKey::Control));
+	UI_COMMAND(MoveLayerToTop, "Bring to front", "Bring to front", EUserInterfaceActionType::Button, FInputChord(EKeys::LeftBracket, EModifierKey::Shift | EModifierKey::Control));
+	UI_COMMAND(MoveLayerToBottom, "Send to back", "Send to back", EUserInterfaceActionType::Button, FInputChord(EKeys::RightBracket, EModifierKey::Shift | EModifierKey::Control));
+	UI_COMMAND(SelectLayerAbove, "Select next layer", "Changes the active layer to the next layer", EUserInterfaceActionType::Button, FInputChord(EKeys::RightBracket, EModifierKey::Alt));
+	UI_COMMAND(SelectLayerBelow, "Select previous layer", "Changes the active layer to the previous layer", EUserInterfaceActionType::Button, FInputChord(EKeys::LeftBracket, EModifierKey::Alt));
 }
 
 #undef LOCTEXT_NAMESPACE
\ No newline at end of file
diff --git a/Engine/Plugins/2D/Paper2D/Source/Paper2DEditor/Private/TileMapEditing/TileMapEditorCommands.h b/Engine/Plugins/2D/Paper2D/Source/Paper2DEditor/Private/TileMapEditing/TileMapEditorCommands.h
--- a/Engine/Plugins/2D/Paper2D/Source/Paper2DEditor/Private/TileMapEditing/TileMapEditorCommands.h
+++ b/Engine/Plugins/2D/Paper2D/Source/Paper2DEditor/Private/TileMapEditing/TileMapEditorCommands.h
@@ -49,4 +49,8 @@ public:
 	TSharedPtr<FUICommandInfo> MergeLayerDown;
 	TSharedPtr<FUICommandInfo> MoveLayerUp;
 	TSharedPtr<FUICommandInfo> MoveLayerDown;
+	TSharedPtr<FUICommandInfo> MoveLayerToTop;
+	TSharedPtr<FUICommandInfo> MoveLayerToBottom;
+	TSharedPtr<FUICommandInfo> SelectLayerAbove;
+	TSharedPtr<FUICommandInfo> SelectLayerBelow;
 };
\ No newline at end of file
