Paper2D: Avoid some log spam when importing tile maps

diff --git a/Engine/Plugins/2D/Paper2D/Source/Paper2DEditor/Private/PaperJSONHelpers.cpp b/Engine/Plugins/2D/Paper2D/Source/Paper2DEditor/Private/PaperJSONHelpers.cpp
--- a/Engine/Plugins/2D/Paper2D/Source/Paper2DEditor/Private/PaperJSONHelpers.cpp
+++ b/Engine/Plugins/2D/Paper2D/Source/Paper2DEditor/Private/PaperJSONHelpers.cpp
@@ -11,8 +11,12 @@ static TArray< TSharedPtr<FJsonValue> > EmptyArray;
 
 FString FPaperJSONHelpers::ReadString(TSharedPtr<class FJsonObject> Item, const FString& Key, const FString& DefaultValue)
 {
-	const TSharedPtr<FJsonValue>& ValuePtr = Item->GetField<EJson::String>(Key);
-	return ValuePtr.IsValid() ? ValuePtr->AsString() : DefaultValue;
+	FString Result;
+	if (!Item->TryGetStringField(Key, /*out*/ Result))
+	{
+		Result = DefaultValue;
+	}
+	return Result;
 }
 
 TSharedPtr<class FJsonObject> FPaperJSONHelpers::ReadObject(TSharedPtr<class FJsonObject> Item, const FString& Key)
@@ -41,21 +45,20 @@ const TArray< TSharedPtr<FJsonValue> >& FPaperJSONHelpers::ReadArray(TSharedPtr<
 
 bool FPaperJSONHelpers::ReadBoolean(const TSharedPtr<class FJsonObject> Item, const FString& Key, bool bDefaultIfMissing)
 {
-	if (Item->HasTypedField<EJson::Boolean>(Key))
-	{
-		return Item->GetBoolField(Key);
-	}
-	else
+	bool bResult;
+	if (!Item->TryGetBoolField(Key, /*out*/ bResult))
 	{
-		return bDefaultIfMissing;
+		bResult = bDefaultIfMissing;
 	}
+	return bResult;
 }
 
 bool FPaperJSONHelpers::ReadFloatNoDefault(const TSharedPtr<class FJsonObject> Item, const FString& Key, float& Out_Value)
 {
-	if (Item->HasTypedField<EJson::Number>(Key))
+	double DoubleOutValue;
+	if (Item->TryGetNumberField(Key, /*out*/ DoubleOutValue))
 	{
-		Out_Value = Item->GetNumberField(Key);
+		Out_Value = DoubleOutValue;
 		return true;
 	}
 	else
diff --git a/Engine/Plugins/2D/Paper2D/Source/PaperTiledImporter/Private/PaperTiledImporterFactory.cpp b/Engine/Plugins/2D/Paper2D/Source/PaperTiledImporter/Private/PaperTiledImporterFactory.cpp
--- a/Engine/Plugins/2D/Paper2D/Source/PaperTiledImporter/Private/PaperTiledImporterFactory.cpp
+++ b/Engine/Plugins/2D/Paper2D/Source/PaperTiledImporter/Private/PaperTiledImporterFactory.cpp
@@ -352,7 +352,7 @@ bool UPaperTiledImporterFactory::CanReimport(UObject* Obj, TArray<FString>& OutF
 		}
 		else
 		{
-			OutFilenames.Add(TEXT(""));
+			OutFilenames.Add(FString());
 		}
 		return true;
 	}
@@ -454,7 +454,7 @@ void UPaperTiledImporterFactory::ParseGlobalInfoFromJSON(TSharedPtr<FJsonObject>
 	ParseIntegerFields(OptionalIntFields, ARRAY_COUNT(OptionalIntFields), Tree, NameForErrors, /*bSilent=*/ true);
 
 	// Parse StaggerAxis if present
-	const FString StaggerAxisStr = FPaperJSONHelpers::ReadString(Tree, TEXT("staggeraxis"), TEXT(""));
+	const FString StaggerAxisStr = FPaperJSONHelpers::ReadString(Tree, TEXT("staggeraxis"), FString());
 	if (StaggerAxisStr == TEXT("x"))
 	{
 		OutParsedInfo.StaggerAxis = ETiledStaggerAxis::X;
@@ -470,7 +470,7 @@ void UPaperTiledImporterFactory::ParseGlobalInfoFromJSON(TSharedPtr<FJsonObject>
 	}
 
 	// Parse StaggerIndex if present
-	const FString StaggerIndexStr = FPaperJSONHelpers::ReadString(Tree, TEXT("staggerindex"), TEXT(""));
+	const FString StaggerIndexStr = FPaperJSONHelpers::ReadString(Tree, TEXT("staggerindex"), FString());
 	if (StaggerIndexStr == TEXT("even"))
 	{
 		OutParsedInfo.StaggerIndex = ETiledStaggerIndex::Even;
@@ -486,7 +486,7 @@ void UPaperTiledImporterFactory::ParseGlobalInfoFromJSON(TSharedPtr<FJsonObject>
 	}
 
 	// Parse RenderOrder if present
-	const FString RenderOrderStr = FPaperJSONHelpers::ReadString(Tree, TEXT("staggerindex"), TEXT(""));
+	const FString RenderOrderStr = FPaperJSONHelpers::ReadString(Tree, TEXT("staggerindex"), FString());
 	if (RenderOrderStr == TEXT("right-down"))
 	{
 		OutParsedInfo.RenderOrder = ETiledRenderOrder::RightDown;
@@ -510,14 +510,14 @@ void UPaperTiledImporterFactory::ParseGlobalInfoFromJSON(TSharedPtr<FJsonObject>
 	}
 
 	// Parse BackgroundColor if present
-	const FString ColorHexStr = FPaperJSONHelpers::ReadString(Tree, TEXT("backgroundcolor"), TEXT(""));
+	const FString ColorHexStr = FPaperJSONHelpers::ReadString(Tree, TEXT("backgroundcolor"), FString());
 	if (!ColorHexStr.IsEmpty())
 	{
 		OutParsedInfo.BackgroundColor = FColor::FromHex(ColorHexStr);
 	}
 
 	// Parse the orientation
-	const FString OrientationModeStr = FPaperJSONHelpers::ReadString(Tree, TEXT("orientation"), TEXT(""));
+	const FString OrientationModeStr = FPaperJSONHelpers::ReadString(Tree, TEXT("orientation"), FString());
 	if (OrientationModeStr == TEXT("orthogonal"))
 	{
 		OutParsedInfo.Orientation = ETiledOrientation::Orthogonal;
@@ -649,20 +649,24 @@ void FTileSetFromTiled::ParseTileSetFromJSON(TSharedPtr<FJsonObject> Tree, const
 	if (bSuccessfullyParsed)
 	{
 		FIntPoint TileOffsetTemp;
-		if (FPaperJSONHelpers::ReadIntPoint(Tree, TEXT("tileoffset"), /*out*/ TileOffsetTemp))
-		{
-			TileOffsetX = TileOffsetTemp.X;
-			TileOffsetY = TileOffsetTemp.Y;
-		}
-		else
+
+		if (Tree->HasField(TEXT("tileoffset")))
 		{
-			TILED_IMPORT_ERROR(TEXT("Failed to parse '%s'.  Invalid or missing value for '%s'"), *NameForErrors, TEXT("tileoffset"));
-			bSuccessfullyParsed = false;
+			if (FPaperJSONHelpers::ReadIntPoint(Tree, TEXT("tileoffset"), /*out*/ TileOffsetTemp))
+			{
+				TileOffsetX = TileOffsetTemp.X;
+				TileOffsetY = TileOffsetTemp.Y;
+			}
+			else
+			{
+				TILED_IMPORT_ERROR(TEXT("Failed to parse '%s'.  Invalid or missing value for '%s'"), *NameForErrors, TEXT("tileoffset"));
+				bSuccessfullyParsed = false;
+			}
 		}
 	}
 
 	// Parse the tile set name
-	Name = FPaperJSONHelpers::ReadString(Tree, TEXT("name"), TEXT(""));
+	Name = FPaperJSONHelpers::ReadString(Tree, TEXT("name"), FString());
 	if (Name.IsEmpty())
 	{
 		TILED_IMPORT_WARNING(TEXT("Expected a non-empty name for each tile set in '%s', generating a new name"), *NameForErrors);
@@ -670,7 +674,7 @@ void FTileSetFromTiled::ParseTileSetFromJSON(TSharedPtr<FJsonObject> Tree, const
 	}
 
 	// Parse the image path
-	ImagePath = FPaperJSONHelpers::ReadString(Tree, TEXT("image"), TEXT(""));
+	ImagePath = FPaperJSONHelpers::ReadString(Tree, TEXT("image"), FString());
 	if (ImagePath.IsEmpty())
 	{
 		TILED_IMPORT_ERROR(TEXT("Failed to parse '%s'.  Invalid value for '%s' ('%s' but expected a path to an image)"), *NameForErrors, TEXT("image"), *ImagePath);
@@ -678,7 +682,7 @@ void FTileSetFromTiled::ParseTileSetFromJSON(TSharedPtr<FJsonObject> Tree, const
 	}
 
 	// Parse the transparent color if present
-	const FString TransparentColorStr = FPaperJSONHelpers::ReadString(Tree, TEXT("transparentcolor"), TEXT(""));
+	const FString TransparentColorStr = FPaperJSONHelpers::ReadString(Tree, TEXT("transparentcolor"), FString());
 	if (!TransparentColorStr.IsEmpty())
 	{
 		bRemoveTransparentColor = true;
@@ -791,7 +795,7 @@ bool FTileLayerFromTiled::ParseFromJSON(TSharedPtr<FJsonObject> Tree, const FStr
 	}
 
 	// Parse the layer type
-	const FString LayerTypeStr = FPaperJSONHelpers::ReadString(Tree, TEXT("type"), TEXT(""));
+	const FString LayerTypeStr = FPaperJSONHelpers::ReadString(Tree, TEXT("type"), FString());
 	if (LayerTypeStr == TEXT("tilelayer"))
 	{
 		LayerType = ETiledLayerType::TileLayer;
@@ -811,7 +815,7 @@ bool FTileLayerFromTiled::ParseFromJSON(TSharedPtr<FJsonObject> Tree, const FStr
 	}
 
 	// Parse the object draw order (if present)
-	const FString ObjectDrawOrderStr = FPaperJSONHelpers::ReadString(Tree, TEXT("draworder"), TEXT(""));
+	const FString ObjectDrawOrderStr = FPaperJSONHelpers::ReadString(Tree, TEXT("draworder"), FString());
 	if (ObjectDrawOrderStr == TEXT("index"))
 	{
 		ObjectDrawOrder = ETiledObjectLayerDrawOrder::Index;
@@ -874,7 +878,7 @@ bool FTileLayerFromTiled::ParseFromJSON(TSharedPtr<FJsonObject> Tree, const FStr
 	}
 	else if (LayerType == ETiledLayerType::ImageLayer)
 	{
-		OverlayImagePath = FPaperJSONHelpers::ReadString(Tree, TEXT("image"), TEXT(""));
+		OverlayImagePath = FPaperJSONHelpers::ReadString(Tree, TEXT("image"), FString());
 	}
 	else
 	{
