Paper2D: WIP of per-tile properties, including collision shapes
- Introduced FPaperTileMetadata
- Added a dense array of FPaperTileMetadata in each tile set (may end up making these sparse later)
- Importing from Tiled will map per-tile collision object data into this data struture, but nothing consumes it yet

Note: Data forward-compat. for per-tile properties is *not* guaranteed yet

diff --git a/Engine/Plugins/2D/Paper2D/Source/Paper2D/Classes/PaperTileSet.h b/Engine/Plugins/2D/Paper2D/Source/Paper2D/Classes/PaperTileSet.h
--- a/Engine/Plugins/2D/Paper2D/Source/Paper2D/Classes/PaperTileSet.h
+++ b/Engine/Plugins/2D/Paper2D/Source/Paper2D/Classes/PaperTileSet.h
@@ -2,11 +2,36 @@
 
 #pragma once
 
+#include "SpriteEditorOnlyTypes.h"
 #include "Engine/DataAsset.h"
 #include "Engine/EngineTypes.h"
 
 #include "PaperTileSet.generated.h"
 
+// Information about a single tile in a tile set
+USTRUCT()
+struct PAPER2D_API FPaperTileMetadata
+{
+public:
+	GENERATED_USTRUCT_BODY()
+
+	UPROPERTY(EditAnywhere, Category=Sprite)
+	FSpriteGeometryCollection CollisionData;
+
+public:
+	FPaperTileMetadata()
+	{
+		CollisionData.GeometryType = ESpritePolygonMode::FullyCustom;
+	}
+
+	// Does this tile have collision information?
+	bool HasCollision() const
+	{
+		return CollisionData.Shapes.Num() > 0;
+	}
+};
+
+
 /**
  * A tile set is a collection of tiles pulled from a texture that can be used to fill out a tile map.
  *
@@ -34,7 +59,7 @@ class PAPER2D_API UPaperTileSet : public UObject
 	int32 Margin;
 
 	// The amount of padding between tiles in the tile sheet (in pixels)
-	UPROPERTY(Category = TileSet, EditAnywhere, meta=(UIMin=0, ClampMin=0))
+	UPROPERTY(Category=TileSet, EditAnywhere, meta=(UIMin=0, ClampMin=0))
 	int32 Spacing;
 
 	// The drawing offset for tiles from this set (in pixels)
@@ -48,20 +73,67 @@ class PAPER2D_API UPaperTileSet : public UObject
 #endif
 
 
+protected:
+	// Cached width of this tile set (in tiles)
+	UPROPERTY()
+	int32 WidthInTiles;
+
+	// Cached height of this tile set (in tiles)
+	UPROPERTY()
+	int32 HeightInTiles;
+
+	// Allocated width of the per-tile data
+	UPROPERTY()
+	int32 AllocatedWidth;
+
+	// Allocated height of the per-tile data
+	UPROPERTY()
+	int32 AllocatedHeight;
+
+	// Per-tile information (Experimental)
+	UPROPERTY(EditAnywhere, EditFixedSize, Category=Sprite)
+	TArray<FPaperTileMetadata> ExperimentalPerTileData;
+
+protected:
+
+	// Reallocates the ExperimentalPerTileData array to the specified size.
+	// Note: This is a destructive operation!
+	void DestructiveAllocateTileData(int32 NewWidth, int32 NewHeight);
+
+	// Reallocates the per-tile data to match the current (WidthInTiles, HeightInTiles) size, copying over what it can
+	void ReallocateAndCopyTileData();
+
+public:	
+
 	// UObject interface
 #if WITH_EDITOR
+	virtual void PostLoad() override;
 	virtual void PostEditChangeProperty(FPropertyChangedEvent& PropertyChangedEvent) override;
 #endif
 	// End of UObject interface
 
 public:
+	// Returns the number of tiles in this tile set
 	int32 GetTileCount() const;
+	
+	// Returns the number of tiles per row in this tile set
 	int32 GetTileCountX() const;
+	
+	// Returns the number of tiles per column in this tile set
 	int32 GetTileCountY() const;
 
+	// Returns editable tile metadata for the specified tile index
+	FPaperTileMetadata* GetMutableTileMetadata(int32 TileIndex);
+
+	// Returns the tile metadata for the specified tile index
+	const FPaperTileMetadata* GetTileMetadata(int32 TileIndex) const;
+
+	// Returns the texture-space coordinates of the top left corner of the specified tile index
 	bool GetTileUV(int32 TileIndex, /*out*/ FVector2D& Out_TileUV) const;
+
+	// Returns the texture-space coordinates of the top left corner of the tile at (TileXY.X, TileXY.Y)
 	FIntPoint GetTileUVFromTileXY(const FIntPoint& TileXY) const;
 
-	// Converts the texture coordinates into tile coordinates
+	// Converts the texture-space coordinates into tile coordinates
 	FIntPoint GetTileXYFromTextureUV(const FVector2D& TextureUV, bool bRoundUp) const;
 };
diff --git a/Engine/Plugins/2D/Paper2D/Source/Paper2D/Private/PaperTileSet.cpp b/Engine/Plugins/2D/Paper2D/Source/Paper2D/Private/PaperTileSet.cpp
--- a/Engine/Plugins/2D/Paper2D/Source/Paper2D/Private/PaperTileSet.cpp
+++ b/Engine/Plugins/2D/Paper2D/Source/Paper2D/Private/PaperTileSet.cpp
@@ -65,6 +65,30 @@ int32 UPaperTileSet::GetTileCountY() const
 	}
 }
 
+FPaperTileMetadata* UPaperTileSet::GetMutableTileMetadata(int32 TileIndex)
+{
+	if (ExperimentalPerTileData.IsValidIndex(TileIndex))
+	{
+		return &(ExperimentalPerTileData[TileIndex]);
+	}
+	else
+	{
+		return nullptr;
+	}
+}
+
+const FPaperTileMetadata* UPaperTileSet::GetTileMetadata(int32 TileIndex) const
+{
+	if (ExperimentalPerTileData.IsValidIndex(TileIndex))
+	{
+		return &(ExperimentalPerTileData[TileIndex]);
+	}
+	else
+	{
+		return nullptr;
+	}
+}
+
 bool UPaperTileSet::GetTileUV(int32 TileIndex, /*out*/ FVector2D& Out_TileUV) const
 {
 	const int32 NumCells = GetTileCount();
@@ -105,6 +129,20 @@ FIntPoint UPaperTileSet::GetTileXYFromTextureUV(const FVector2D& TextureUV, bool
 #include "PaperTileMapComponent.h"
 #include "ComponentReregisterContext.h"
 
+void UPaperTileSet::PostLoad()
+{
+	if (TileSheet != nullptr)
+	{
+		TileSheet->ConditionalPostLoad();
+
+		WidthInTiles = GetTileCountX();
+		HeightInTiles = GetTileCountY();
+		ReallocateAndCopyTileData();
+	}
+
+	Super::PostLoad();
+}
+
 void UPaperTileSet::PostEditChangeProperty(FPropertyChangedEvent& PropertyChangedEvent)
 {
 	const FName PropertyName = (PropertyChangedEvent.Property != nullptr) ? PropertyChangedEvent.Property->GetFName() : NAME_None;
@@ -115,9 +153,61 @@ void UPaperTileSet::PostEditChangeProperty(FPropertyChangedEvent& PropertyChange
 	TileWidth = FMath::Max<int32>(TileWidth, 1);
 	TileHeight = FMath::Max<int32>(TileHeight, 1);
 
-	//@TODO: Determine when these are really needed, as they're seriously expensive!
-	TComponentReregisterContext<UPaperTileMapComponent> ReregisterStaticComponents;
+	WidthInTiles = GetTileCountX();
+	HeightInTiles = GetTileCountY();
+	ReallocateAndCopyTileData();
+
+	// Rebuild any tile map components that may have been relying on us
+	if ((PropertyChangedEvent.ChangeType & EPropertyChangeType::Interactive) == 0)
+	{
+		//@TODO: Currently tile maps have no fast list of referenced tile sets, so we just rebuild all of them
+		TComponentReregisterContext<UPaperTileMapComponent> ReregisterAllTileMaps;
+	}
 
 	Super::PostEditChangeProperty(PropertyChangedEvent);
 }
 #endif
+
+void UPaperTileSet::DestructiveAllocateTileData(int32 NewWidth, int32 NewHeight)
+{
+	check((NewWidth > 0) && (NewHeight > 0));
+
+	const int32 NumCells = NewWidth * NewHeight;
+	ExperimentalPerTileData.Empty(NumCells);
+	for (int32 Index = 0; Index < NumCells; ++Index)
+	{
+		ExperimentalPerTileData.Add(FPaperTileMetadata());
+	}
+
+	AllocatedWidth = NewWidth;
+	AllocatedHeight = NewHeight;
+}
+
+void UPaperTileSet::ReallocateAndCopyTileData()
+{
+	if ((AllocatedWidth == WidthInTiles) && (AllocatedHeight == HeightInTiles))
+	{
+		return;
+	}
+
+	const int32 SavedWidth = AllocatedWidth;
+	const int32 SavedHeight = AllocatedHeight;
+	TArray<FPaperTileMetadata> SavedDesignedMap(ExperimentalPerTileData);
+
+	DestructiveAllocateTileData(WidthInTiles, HeightInTiles);
+
+	const int32 CopyWidth = FMath::Min<int32>(WidthInTiles, SavedWidth);
+	const int32 CopyHeight = FMath::Min<int32>(HeightInTiles, SavedHeight);
+
+	for (int32 Y = 0; Y < CopyHeight; ++Y)
+	{
+		for (int32 X = 0; X < CopyWidth; ++X)
+		{
+			const int32 SrcIndex = Y*SavedWidth + X;
+			const int32 DstIndex = Y*WidthInTiles + X;
+
+			ExperimentalPerTileData[DstIndex] = SavedDesignedMap[SrcIndex];
+		}
+	}
+
+}
diff --git a/Engine/Plugins/2D/Paper2D/Source/PaperTiledImporter/Classes/PaperTiledImporterFactory.h b/Engine/Plugins/2D/Paper2D/Source/PaperTiledImporter/Classes/PaperTiledImporterFactory.h
--- a/Engine/Plugins/2D/Paper2D/Source/PaperTiledImporter/Classes/PaperTiledImporterFactory.h
+++ b/Engine/Plugins/2D/Paper2D/Source/PaperTiledImporter/Classes/PaperTiledImporterFactory.h
@@ -163,6 +163,8 @@ public:
 	bool ParseFromJSON(TSharedPtr<FJsonObject> Tree, const FString& NameForErrors, bool bSilent = false);
 
 	static bool ParsePointArray(TArray<FVector2D>& OutPoints, const TArray<TSharedPtr<FJsonValue>>& InArray, const FString& NameForErrors, bool bSilent);
+
+	static void AddToSpriteGeometryCollection(const FVector2D& Offset, const TArray<FTiledObject>& InObjects, FSpriteGeometryCollection& InOutShapes);
 };
 
 //////////////////////////////////////////////////////////////////////////
diff --git a/Engine/Plugins/2D/Paper2D/Source/PaperTiledImporter/Private/PaperTiledImporterFactory.cpp b/Engine/Plugins/2D/Paper2D/Source/PaperTiledImporter/Private/PaperTiledImporterFactory.cpp
--- a/Engine/Plugins/2D/Paper2D/Source/PaperTiledImporter/Private/PaperTiledImporterFactory.cpp
+++ b/Engine/Plugins/2D/Paper2D/Source/PaperTiledImporter/Private/PaperTiledImporterFactory.cpp
@@ -238,7 +238,29 @@ UObject* UPaperTiledImporterFactory::FactoryCreateText(UClass* InClass, UObject*
 					bLoadedSuccessfully = false;
 				}
 
-				TileSetAsset->PostEditChange();
+				// Make the tile set allocate space for the per-tile data
+				FPropertyChangedEvent InteractiveRebuildTileSet(nullptr, EPropertyChangeType::Interactive);
+				TileSetAsset->PostEditChangeProperty(InteractiveRebuildTileSet);
+
+				// Copy across per-tile metadata
+				const int32 NumTilesCreated = TileSetAsset->GetTileCount();
+				for (const auto& KV : TileSetData.PerTileData)
+				{
+					const int32 TileIndex = KV.Key;
+					const FTiledTileInfo& SourceTileData = KV.Value;
+
+					if (FPaperTileMetadata* TargetTileData = TileSetAsset->GetMutableTileMetadata(TileIndex))
+					{
+						//@TODO: TileData.TerrainIndices
+						FTiledObject::AddToSpriteGeometryCollection(FVector2D::ZeroVector, SourceTileData.Objects, TargetTileData->CollisionData);
+					}
+				}
+				
+				// Update anyone who might be using the tile set (in case we're reimporting)
+				FPropertyChangedEvent FinalRebuildTileSet(nullptr, EPropertyChangeType::ValueSet);
+				TileSetAsset->PostEditChangeProperty(FinalRebuildTileSet);
+
+				// Save off that we created the asset
 				GlobalInfo.CreatedTileSetAssets.Add(TileSetAsset);
 			}
 			else
@@ -770,8 +792,8 @@ bool FTileLayerFromTiled::ParseFromJSON(TSharedPtr<FJsonObject> Tree, const FStr
 
 	// Parse all of the integer fields
 	FRequiredIntField IntFields[] = {
-		FRequiredIntField( Width, TEXT("width"), 1 ),
-		FRequiredIntField( Height, TEXT("height"), 1 ),
+		FRequiredIntField( Width, TEXT("width"), 0 ),
+		FRequiredIntField( Height, TEXT("height"), 0 ),
 		FRequiredIntField( OffsetX, TEXT("x"), 0 ),
 		FRequiredIntField( OffsetY, TEXT("y"), 0 )
 	};
@@ -798,6 +820,12 @@ bool FTileLayerFromTiled::ParseFromJSON(TSharedPtr<FJsonObject> Tree, const FStr
 	const FString LayerTypeStr = FPaperJSONHelpers::ReadString(Tree, TEXT("type"), FString());
 	if (LayerTypeStr == TEXT("tilelayer"))
 	{
+		if ((Width < 1) || (Height < 1))
+		{
+			TILED_IMPORT_ERROR(TEXT("Failed to parse '%s'.  Tile layers should be at least 1x1"), *NameForErrors);
+			bSuccessfullyParsed = false;
+		}
+
 		LayerType = ETiledLayerType::TileLayer;
 	}
 	else if (LayerTypeStr == TEXT("objectgroup"))
@@ -1122,6 +1150,53 @@ bool FTiledObject::ParsePointArray(TArray<FVector2D>& OutPoints, const TArray<TS
 	return bSuccessfullyParsed;
 }
 
+void FTiledObject::AddToSpriteGeometryCollection(const FVector2D& Offset, const TArray<FTiledObject>& InObjects, FSpriteGeometryCollection& InOutShapes)
+{
+	for (const FTiledObject& SourceObject : InObjects)
+	{
+		const FVector2D SourcePos = Offset + FVector2D(SourceObject.X, SourceObject.Y);
+		const float SmallerWidthOrHeight = FMath::Min(SourceObject.Width, SourceObject.Height);
+
+		bool bCreatedShape = false;
+		switch (SourceObject.TiledObjectType)
+		{
+		case ETiledObjectType::Box:
+			InOutShapes.AddRectangleShape(SourcePos, FVector2D(SourceObject.Width, SourceObject.Height));
+			bCreatedShape = true;
+			break;
+		case ETiledObjectType::Ellipse:
+			InOutShapes.AddRectangleShape(SourcePos, FVector2D(SmallerWidthOrHeight, SmallerWidthOrHeight));
+			bCreatedShape = true;
+			break;
+		case ETiledObjectType::Polygon:
+			{
+				FSpriteGeometryShape NewShape;
+				NewShape.ShapeType = ESpriteShapeType::Polygon;
+				NewShape.BoxPosition = SourcePos;
+				NewShape.Vertices = SourceObject.Points;
+				InOutShapes.Shapes.Add(NewShape);
+				bCreatedShape = true;
+		}
+			break;
+		case ETiledObjectType::PlacedTile:
+			UE_LOG(LogPaperTiledImporter, Warning, TEXT("Ignoring Tiled Object of type PlacedTile"));
+			break;
+		case ETiledObjectType::Polyline:
+			UE_LOG(LogPaperTiledImporter, Warning, TEXT("Ignoring Tiled Object of type Polyline"));
+			break;
+		default:
+			ensureMsg(false, TEXT("Unknown enumerant in ETiledObjectType"));
+			break;
+		}
+
+		if (bCreatedShape)
+		{
+			const float RotationUnwound = FMath::Fmod(SourceObject.RotationDegrees, 360.0f);
+			InOutShapes.Shapes.Last().Rotation = RotationUnwound;
+		}
+	}
+}
+
 //////////////////////////////////////////////////////////////////////////
 
 #undef TILED_IMPORT_ERROR
