Paper2D: Don't lose reference to object during GC. [UE-13172]

diff --git a/Engine/Plugins/2D/Paper2D/Source/Paper2DEditor/Private/ExtractSprites/SPaperExtractSpritesDialog.cpp b/Engine/Plugins/2D/Paper2D/Source/Paper2DEditor/Private/ExtractSprites/SPaperExtractSpritesDialog.cpp
--- a/Engine/Plugins/2D/Paper2D/Source/Paper2DEditor/Private/ExtractSprites/SPaperExtractSpritesDialog.cpp
+++ b/Engine/Plugins/2D/Paper2D/Source/Paper2DEditor/Private/ExtractSprites/SPaperExtractSpritesDialog.cpp
@@ -127,6 +127,7 @@ void SPaperExtractSpritesDialog::Construct(const FArguments& InArgs, UTexture2D*
 	SourceTexture = InSourceTexture;
 
 	ExtractSpriteSettings = NewObject<UPaperExtractSpritesSettings>();
+	ExtractSpriteSettings->AddToRoot();
 	ExtractSpriteSettings->GridWidth = InSourceTexture->GetSizeX();
 	ExtractSpriteSettings->GridHeight = InSourceTexture->GetSizeY();
 	PreviewExtractedSprites();
@@ -187,6 +188,14 @@ void SPaperExtractSpritesDialog::Construct(const FArguments& InArgs, UTexture2D*
 	];
 }
 
+SPaperExtractSpritesDialog::~SPaperExtractSpritesDialog()
+{
+	if (ExtractSpriteSettings && ExtractSpriteSettings->IsValidLowLevel())
+	{
+		ExtractSpriteSettings->RemoveFromRoot();
+	}
+}
+
 bool SPaperExtractSpritesDialog::ShowWindow(const FText& TitleText, UTexture2D* SourceTexture)
 {
 	// Create the window to pick the class
diff --git a/Engine/Plugins/2D/Paper2D/Source/Paper2DEditor/Private/ExtractSprites/SPaperExtractSpritesDialog.h b/Engine/Plugins/2D/Paper2D/Source/Paper2DEditor/Private/ExtractSprites/SPaperExtractSpritesDialog.h
--- a/Engine/Plugins/2D/Paper2D/Source/Paper2DEditor/Private/ExtractSprites/SPaperExtractSpritesDialog.h
+++ b/Engine/Plugins/2D/Paper2D/Source/Paper2DEditor/Private/ExtractSprites/SPaperExtractSpritesDialog.h
@@ -22,6 +22,8 @@ public:
 	// Constructs this widget with InArgs
 	void Construct(const FArguments& InArgs, UTexture2D* Texture);
 
+	virtual ~SPaperExtractSpritesDialog();
+
 	// Show the dialog, returns true if successfully extracted sprites
 	static bool ShowWindow(const FText& TitleText, UTexture2D* SourceTexture);
 
@@ -44,6 +46,7 @@ private:
 	// Source texture to extract from
 	UTexture2D* SourceTexture;
 
+
 	class UPaperExtractSpritesSettings* ExtractSpriteSettings;
 
 	TSharedPtr<class IDetailsView> PropertyView;
