Paper2D: Improvements to Extract Sprites dialog [UE-12292] [UE-13168] [UE-13171]
- Hide properties that are not relevant to the current mode
- Better tooltips
- Can customize background color
- Added numcells X & Y to limit extraction
- Split up Margin and Padding to have independent X and Y values
- Added basic naming template ({0}

diff --git a/Engine/Plugins/2D/Paper2D/Source/Paper2DEditor/Private/ExtractSprites/PaperExtractSpritesSettings.cpp b/Engine/Plugins/2D/Paper2D/Source/Paper2DEditor/Private/ExtractSprites/PaperExtractSpritesSettings.cpp
--- a/Engine/Plugins/2D/Paper2D/Source/Paper2DEditor/Private/ExtractSprites/PaperExtractSpritesSettings.cpp
+++ b/Engine/Plugins/2D/Paper2D/Source/Paper2DEditor/Private/ExtractSprites/PaperExtractSpritesSettings.cpp
@@ -8,15 +8,22 @@ UPaperExtractSpritesSettings::UPaperExtractSpritesSettings(const FObjectInitiali
 {
 	OutlineColor = FLinearColor::Yellow;
 	TextureTint = FLinearColor::Gray;
+	BackgroundColor = FLinearColor(0.1f, 0.1f, 0.1f);
 }
 
 void UPaperExtractSpritesSettings::PostEditChangeProperty(FPropertyChangedEvent& PropertyChangedEvent)
 {
 	const FName PropertyName = (PropertyChangedEvent.Property != nullptr) ? PropertyChangedEvent.Property->GetFName() : NAME_None;
-	if (PropertyName == GET_MEMBER_NAME_CHECKED(UPaperExtractSpritesSettings, Mode))
+	if (PropertyName == GET_MEMBER_NAME_CHECKED(UPaperExtractSpritesSettings, SpriteExtractMode))
 	{
 		// Update bools
-		bAutoMode = (Mode == ESpriteExtractMode::Auto);
-		bGridMode = (Mode == ESpriteExtractMode::Grid);
+		bAutoMode = (SpriteExtractMode == ESpriteExtractMode::Auto);
+		bGridMode = (SpriteExtractMode == ESpriteExtractMode::Grid);
 	}
 }
+
+
+UPaperExtractSpriteGridSettings::UPaperExtractSpriteGridSettings(const FObjectInitializer& ObjectInitializer)
+	: Super(ObjectInitializer)
+{
+}
diff --git a/Engine/Plugins/2D/Paper2D/Source/Paper2DEditor/Private/ExtractSprites/PaperExtractSpritesSettings.h b/Engine/Plugins/2D/Paper2D/Source/Paper2DEditor/Private/ExtractSprites/PaperExtractSpritesSettings.h
--- a/Engine/Plugins/2D/Paper2D/Source/Paper2DEditor/Private/ExtractSprites/PaperExtractSpritesSettings.h
+++ b/Engine/Plugins/2D/Paper2D/Source/Paper2DEditor/Private/ExtractSprites/PaperExtractSpritesSettings.h
@@ -7,10 +7,15 @@
 UENUM()
 enum class ESpriteExtractMode : uint8
 {
-	Auto,	// Automatically extract sprites
-	Grid,	// Extract sprites by grid
+	// Automatically extract sprites by detecting using alpha
+	Auto,
+
+	// Extract sprites in a grid defined in the properties below
+	Grid,	
 };
 
+
+
 /**
 *
 */
@@ -20,15 +25,22 @@ class UPaperExtractSpritesSettings : public UObject
 	GENERATED_BODY()
 
 public:
+	// Sprite extract mode
 	UPROPERTY(Category = Settings, EditAnywhere)
-	ESpriteExtractMode Mode;
+	ESpriteExtractMode SpriteExtractMode;
 
+	// The color of the sprite boundary outlines
 	UPROPERTY(Category = Settings, EditAnywhere)
 	FLinearColor OutlineColor;
 
+	// Tint the texture to help increase outline visibility
 	UPROPERTY(Category = Settings, EditAnywhere)
 	FLinearColor TextureTint;
 
+	// The viewport background color
+	UPROPERTY(Category = Settings, EditAnywhere)
+	FLinearColor BackgroundColor;
+
 	UPROPERTY()
 	bool bAutoMode;
 
@@ -36,26 +48,60 @@ public:
 	bool bGridMode;
 
 
+	UPaperExtractSpritesSettings(const FObjectInitializer& ObjectInitializer);
+
+	// UObject interface
+	virtual void PostEditChangeProperty(FPropertyChangedEvent& PropertyChangedEvent) override;
+	// End of UObject interface
+};
+
+UCLASS()
+class UPaperExtractSpriteGridSettings : public UObject
+{
+	GENERATED_BODY()
+
 	//////////////////////////////////////////////////////////////////////////
 	// Grid mode
 
-	UPROPERTY(Category = Grid, EditAnywhere, meta = (EditCondition = "bGridMode", UIMin = 1, ClampMin = 1))
-	int32 GridWidth;
+	// The name of the sprite that will be created. {0} will get replaced by the sprite number.
+	UPROPERTY(Category = Naming, EditAnywhere)
+	FString NamingTemplate;
 
-	UPROPERTY(Category = Grid, EditAnywhere, meta = (EditCondition = "bGridMode", UIMin = 1, ClampMin = 1))
-	int32 GridHeight;
+	// The number to start naming with
+	UPROPERTY(Category = Naming, EditAnywhere)
+	int32 NamingStartIndex;
 
-	UPROPERTY(Category = Grid, EditAnywhere, meta = (EditCondition = "bGridMode", UIMin = 0, ClampMin = 0))
-	int32 Margin;
+	// The width of each sprite in grid mode
+	UPROPERTY(Category = Grid, EditAnywhere, meta = (UIMin = 1, ClampMin = 1))
+	int32 CellWidth;
 
-	UPROPERTY(Category = Grid, EditAnywhere, meta = (EditCondition = "bGridMode", UIMin = 0, ClampMin = 0))
-	int32 Spacing;
-	
+	// The height of each sprite in grid mode
+	UPROPERTY(Category = Grid, EditAnywhere, meta = (UIMin = 1, ClampMin = 1))
+	int32 CellHeight;
 
+	// Number of cells extracted horizontally. Can be used to limit the number of sprites extracted. Set to 0 to extract all sprites
+	UPROPERTY(Category = Grid, EditAnywhere, meta = (UIMin = 0, ClampMin = 0))
+	int32 NumCellsX;
 
-	UPaperExtractSpritesSettings(const FObjectInitializer& ObjectInitializer);
+	// Number of cells extracted vertically. Can be used to limit the number of sprites extracted. Set to 0 to extract all sprites
+	UPROPERTY(Category = Grid, EditAnywhere, meta = (UIMin = 0, ClampMin = 0))
+	int32 NumCellsY;
 
-	// UObject interface
-	virtual void PostEditChangeProperty(FPropertyChangedEvent& PropertyChangedEvent) override;
-	// End of UObject interface
+	// Margin from the left of the texture to the first sprite
+	UPROPERTY(Category = Grid, EditAnywhere, meta = (UIMin = 0, ClampMin = 0))
+	int32 MarginX;
+
+	// Margin from the top of the texture to the first sprite
+	UPROPERTY(Category = Grid, EditAnywhere, meta = (UIMin = 0, ClampMin = 0))
+	int32 MarginY;
+
+	// Horizontal spacing between sprites
+	UPROPERTY(Category = Grid, EditAnywhere, meta = (UIMin = 0, ClampMin = 0))
+	int32 SpacingX;
+
+	// Vertical spacing between sprites
+	UPROPERTY(Category = Grid, EditAnywhere, meta = (UIMin = 0, ClampMin = 0))
+	int32 SpacingY;
+
+	UPaperExtractSpriteGridSettings(const FObjectInitializer& ObjectInitializer);
 };
diff --git a/Engine/Plugins/2D/Paper2D/Source/Paper2DEditor/Private/ExtractSprites/SPaperExtractSpritesDialog.cpp b/Engine/Plugins/2D/Paper2D/Source/Paper2DEditor/Private/ExtractSprites/SPaperExtractSpritesDialog.cpp
--- a/Engine/Plugins/2D/Paper2D/Source/Paper2DEditor/Private/ExtractSprites/SPaperExtractSpritesDialog.cpp
+++ b/Engine/Plugins/2D/Paper2D/Source/Paper2DEditor/Private/ExtractSprites/SPaperExtractSpritesDialog.cpp
@@ -43,8 +43,12 @@ void FPaperExtractSpritesViewportClient::Draw(FViewport* Viewport, FCanvas* Canv
 	// Super will clear the viewport 
 	FPaperEditorViewportClient::Draw(Viewport, Canvas);
 
-	UTexture2D* Texture = TextureBeingExtracted.Get();
+	if (Settings != nullptr)
+	{
+		Canvas->Clear(Settings->BackgroundColor);
+	}
 
+	UTexture2D* Texture = TextureBeingExtracted.Get();
 	if (Texture != nullptr)
 	{
 		const bool bUseTranslucentBlend = Texture->HasAlphaChannel();
@@ -128,16 +132,26 @@ void SPaperExtractSpritesDialog::Construct(const FArguments& InArgs, UTexture2D*
 
 	ExtractSpriteSettings = NewObject<UPaperExtractSpritesSettings>();
 	ExtractSpriteSettings->AddToRoot();
-	ExtractSpriteSettings->GridWidth = InSourceTexture->GetSizeX();
-	ExtractSpriteSettings->GridHeight = InSourceTexture->GetSizeY();
+
+	ExtractSpriteGridSettings = NewObject<UPaperExtractSpriteGridSettings>();
+	ExtractSpriteGridSettings->AddToRoot();
+	ExtractSpriteGridSettings->NamingTemplate = "Sprite_{0}";
+	ExtractSpriteGridSettings->CellWidth = InSourceTexture->GetSizeX();
+	ExtractSpriteGridSettings->CellHeight = InSourceTexture->GetSizeY();
+
 	PreviewExtractedSprites();
 
 	FPropertyEditorModule& EditModule = FModuleManager::Get().GetModuleChecked<FPropertyEditorModule>("PropertyEditor");
 	FDetailsViewArgs DetailsViewArgs(/*bUpdateFromSelection=*/ false, /*bLockable=*/ false, /*bAllowSearch=*/ false, /*InNameAreaSettings=*/ FDetailsViewArgs::HideNameArea, /*bHideSelectionTip=*/ true);
-	PropertyView = EditModule.CreateDetailView(DetailsViewArgs);
-	PropertyView->SetObject(ExtractSpriteSettings);
-	PropertyView->OnFinishedChangingProperties().AddSP(this, &SPaperExtractSpritesDialog::OnFinishedChangingProperties);
+	MainPropertyView = EditModule.CreateDetailView(DetailsViewArgs);
+	MainPropertyView->SetObject(ExtractSpriteSettings);
+	MainPropertyView->OnFinishedChangingProperties().AddSP(this, &SPaperExtractSpritesDialog::OnFinishedChangingProperties);
 	
+	DetailsPropertyView = EditModule.CreateDetailView(DetailsViewArgs);
+	DetailsPropertyView->OnFinishedChangingProperties().AddSP(this, &SPaperExtractSpritesDialog::OnFinishedChangingProperties);
+
+	SetDetailsViewForActiveMode();
+
 	TSharedRef<SPaperExtractSpritesViewport> Viewport = SNew(SPaperExtractSpritesViewport, SourceTexture, ExtractedSprites, ExtractSpriteSettings, this);
 
 	ChildSlot
@@ -158,11 +172,16 @@ void SPaperExtractSpritesDialog::Construct(const FArguments& InArgs, UTexture2D*
 				SNew(SVerticalBox)
 				+ SVerticalBox::Slot()
 				.Padding(2.0f)
+				.AutoHeight()
 				[
-					PropertyView.ToSharedRef()
+					MainPropertyView.ToSharedRef()
 				]
 				+ SVerticalBox::Slot()
+				.Padding(2.0f)
 				.FillHeight(1.0f)
+				[
+					DetailsPropertyView.ToSharedRef()
+				]
 				+ SVerticalBox::Slot()
 				.Padding(2.0f)
 				.AutoHeight()
@@ -194,6 +213,11 @@ SPaperExtractSpritesDialog::~SPaperExtractSpritesDialog()
 	{
 		ExtractSpriteSettings->RemoveFromRoot();
 	}
+
+	if (ExtractSpriteGridSettings && ExtractSpriteGridSettings->IsValidLowLevel())
+	{
+		ExtractSpriteGridSettings->RemoveFromRoot();
+	}
 }
 
 bool SPaperExtractSpritesDialog::ShowWindow(const FText& TitleText, UTexture2D* SourceTexture)
@@ -230,7 +254,7 @@ void SPaperExtractSpritesDialog::PreviewExtractedSprites()
 	SlowTask.EnterProgressFrame();
 
 	const FString DefaultSuffix = TEXT("Sprite");
-	if (ExtractSpriteSettings->Mode == ESpriteExtractMode::Auto)
+	if (ExtractSpriteSettings->SpriteExtractMode == ESpriteExtractMode::Auto)
 	{
 		ExtractedSprites.Empty();
 
@@ -291,21 +315,44 @@ void SPaperExtractSpritesDialog::PreviewExtractedSprites()
 	}
 	else
 	{
+		FString NamingTemplate = ExtractSpriteGridSettings->NamingTemplate;
+		if (NamingTemplate.Find(TEXT("{0}")) == INDEX_NONE)
+		{
+			NamingTemplate.Append(TEXT("_{0}"));
+		}
+
 		// Calculate rects
 		ExtractedSprites.Empty();
 		if (SourceTexture != nullptr)
 		{
-			int32 ExtractedRectIndex = 0;
+			int32 ExtractedRectIndex = ExtractSpriteGridSettings->NamingStartIndex;
 			int32 TextureWidth = SourceTexture->GetSizeX();
 			int32 TextureHeight = SourceTexture->GetSizeY();
-			for (int32 Y = ExtractSpriteSettings->Margin; Y + ExtractSpriteSettings->GridHeight <= TextureHeight; Y += ExtractSpriteSettings->GridHeight + ExtractSpriteSettings->Spacing)
+			int NumExtractedCellsY = 0;
+			for (int32 Y = ExtractSpriteGridSettings->MarginY; Y + ExtractSpriteGridSettings->CellHeight <= TextureHeight; Y += ExtractSpriteGridSettings->CellHeight + ExtractSpriteGridSettings->SpacingY)
 			{
-				for (int32 X = ExtractSpriteSettings->Margin; X + ExtractSpriteSettings->GridWidth <= TextureWidth; X += ExtractSpriteSettings->GridWidth + ExtractSpriteSettings->Spacing)
+				int NumExtractedCellsX = 0;
+				for (int32 X = ExtractSpriteGridSettings->MarginX; X + ExtractSpriteGridSettings->CellWidth <= TextureWidth; X += ExtractSpriteGridSettings->CellWidth + ExtractSpriteGridSettings->SpacingX)
 				{
 					FPaperExtractedSprite* Sprite = new(ExtractedSprites)FPaperExtractedSprite();
-					Sprite->Rect = FIntRect(X, Y, X + ExtractSpriteSettings->GridWidth, Y + ExtractSpriteSettings->GridHeight);
-					Sprite->Name = FString::Printf(TEXT("%s_%d"), *DefaultSuffix, ExtractedRectIndex);
+					Sprite->Rect = FIntRect(X, Y, X + ExtractSpriteGridSettings->CellWidth, Y + ExtractSpriteGridSettings->CellHeight);
+					
+					Sprite->Name = NamingTemplate;
+					Sprite->Name.ReplaceInline(TEXT("{0}"), *FString::Printf(TEXT("%d"), ExtractedRectIndex));
+
 					ExtractedRectIndex++;
+
+					++NumExtractedCellsX;
+					if (ExtractSpriteGridSettings->NumCellsX > 0 && NumExtractedCellsX >= ExtractSpriteGridSettings->NumCellsX)
+					{
+						break;
+					}
+				}
+
+				++NumExtractedCellsY;
+				if (ExtractSpriteGridSettings->NumCellsY > 0 && NumExtractedCellsY >= ExtractSpriteGridSettings->NumCellsY)
+				{
+					break;
 				}
 			}
 		}
@@ -314,7 +361,22 @@ void SPaperExtractSpritesDialog::PreviewExtractedSprites()
 
 void SPaperExtractSpritesDialog::OnFinishedChangingProperties(const FPropertyChangedEvent& PropertyChangedEvent)
 {
-	PreviewExtractedSprites();
+	if (PropertyChangedEvent.Property != nullptr)
+	{
+		const FName PropertyName(PropertyChangedEvent.Property->GetFName());
+		
+		if (PropertyName != GET_MEMBER_NAME_CHECKED(UPaperExtractSpritesSettings, OutlineColor) &&
+			PropertyName != GET_MEMBER_NAME_CHECKED(UPaperExtractSpritesSettings, BackgroundColor) &&
+			PropertyName != GET_MEMBER_NAME_CHECKED(UPaperExtractSpritesSettings, TextureTint))
+		{
+			PreviewExtractedSprites();
+		}
+
+		if (PropertyName == GET_MEMBER_NAME_CHECKED(UPaperExtractSpritesSettings, SpriteExtractMode))
+		{
+			SetDetailsViewForActiveMode();
+		}
+	}
 }
 
 FReply SPaperExtractSpritesDialog::ExtractClicked()
@@ -394,5 +456,17 @@ void SPaperExtractSpritesDialog::CreateExtractedSprites()
 	}
 }
 
+void SPaperExtractSpritesDialog::SetDetailsViewForActiveMode()
+{
+	if (ExtractSpriteSettings->SpriteExtractMode == ESpriteExtractMode::Grid)
+	{
+		DetailsPropertyView->SetObject(ExtractSpriteGridSettings);
+	}
+	else
+	{
+		DetailsPropertyView->SetObject(nullptr);
+	}
+}
+
 #undef LOCTEXT_NAMESPACE
 
diff --git a/Engine/Plugins/2D/Paper2D/Source/Paper2DEditor/Private/ExtractSprites/SPaperExtractSpritesDialog.h b/Engine/Plugins/2D/Paper2D/Source/Paper2DEditor/Private/ExtractSprites/SPaperExtractSpritesDialog.h
--- a/Engine/Plugins/2D/Paper2D/Source/Paper2DEditor/Private/ExtractSprites/SPaperExtractSpritesDialog.h
+++ b/Engine/Plugins/2D/Paper2D/Source/Paper2DEditor/Private/ExtractSprites/SPaperExtractSpritesDialog.h
@@ -42,14 +42,19 @@ private:
 	// Actually create extracted sprites
 	void CreateExtractedSprites();
 
+	// Sets the details panel appropriately for the currently selected panel
+	void SetDetailsViewForActiveMode();
+
 private:
 	// Source texture to extract from
 	UTexture2D* SourceTexture;
 
 
 	class UPaperExtractSpritesSettings* ExtractSpriteSettings;
+	class UPaperExtractSpriteGridSettings* ExtractSpriteGridSettings;
 
-	TSharedPtr<class IDetailsView> PropertyView;
+	TSharedPtr<class IDetailsView> MainPropertyView;
+	TSharedPtr<class IDetailsView> DetailsPropertyView;
 	TArray<FPaperExtractedSprite> ExtractedSprites;
 };
 
