Paper2D: Configurable background color in extract sprites dialog [UE-13170]
- Changed to use GetBackgroundColor instead of clearing the canvas

diff --git a/Engine/Plugins/2D/Paper2D/Source/Paper2DEditor/Private/ExtractSprites/SPaperExtractSpritesDialog.cpp b/Engine/Plugins/2D/Paper2D/Source/Paper2DEditor/Private/ExtractSprites/SPaperExtractSpritesDialog.cpp
--- a/Engine/Plugins/2D/Paper2D/Source/Paper2DEditor/Private/ExtractSprites/SPaperExtractSpritesDialog.cpp
+++ b/Engine/Plugins/2D/Paper2D/Source/Paper2DEditor/Private/ExtractSprites/SPaperExtractSpritesDialog.cpp
@@ -27,6 +27,7 @@ public:
 
 	// FViewportClient interface
 	virtual void Draw(FViewport* Viewport, FCanvas* Canvas) override;
+	virtual FLinearColor GetBackgroundColor() const override;
 	// End of FViewportClient interface
 
 private:
@@ -43,11 +44,6 @@ void FPaperExtractSpritesViewportClient::Draw(FViewport* Viewport, FCanvas* Canv
 	// Super will clear the viewport 
 	FPaperEditorViewportClient::Draw(Viewport, Canvas);
 
-	if (Settings != nullptr)
-	{
-		Canvas->Clear(Settings->BackgroundColor);
-	}
-
 	UTexture2D* Texture = TextureBeingExtracted.Get();
 	if (Texture != nullptr)
 	{
@@ -75,6 +71,19 @@ void FPaperExtractSpritesViewportClient::Draw(FViewport* Viewport, FCanvas* Canv
 	}
 }
 
+FLinearColor FPaperExtractSpritesViewportClient::GetBackgroundColor() const
+{
+	if (Settings != nullptr)
+	{
+		return Settings->BackgroundColor;
+	}
+	else
+	{
+		return FEditorViewportClient::GetBackgroundColor();
+	}
+}
+
+
 void FPaperExtractSpritesViewportClient::DrawRectangle(FCanvas* Canvas, const FLinearColor& Color, const FIntRect& Rect)
 {
 	FVector2D TopLeft(-ZoomPos.X * ZoomAmount + Rect.Min.X * ZoomAmount, -ZoomPos.Y * ZoomAmount + Rect.Min.Y * ZoomAmount);
