Paper2D: Fixed button alignment & style [UE-12292]

diff --git a/Engine/Plugins/2D/Paper2D/Source/Paper2DEditor/Private/ExtractSprites/SPaperExtractSpritesDialog.cpp b/Engine/Plugins/2D/Paper2D/Source/Paper2DEditor/Private/ExtractSprites/SPaperExtractSpritesDialog.cpp
--- a/Engine/Plugins/2D/Paper2D/Source/Paper2DEditor/Private/ExtractSprites/SPaperExtractSpritesDialog.cpp
+++ b/Engine/Plugins/2D/Paper2D/Source/Paper2DEditor/Private/ExtractSprites/SPaperExtractSpritesDialog.cpp
@@ -193,20 +193,24 @@ void SPaperExtractSpritesDialog::Construct(const FArguments& InArgs, UTexture2D*
 				]
 				+ SVerticalBox::Slot()
 				.Padding(2.0f)
+				.HAlign(HAlign_Right)
 				.AutoHeight()
 				[
-					SNew(SHorizontalBox)
-					+ SHorizontalBox::Slot()
-					.AutoWidth()
+					SNew(SUniformGridPanel)
+					.SlotPadding(2)
+					+ SUniformGridPanel::Slot(0, 0)
 					[
 						SNew(SButton)
+						.ButtonStyle(FEditorStyle::Get(), "FlatButton.Success")
+						.ForegroundColor(FLinearColor::White)
 						.Text(LOCTEXT("PaperExtractSpritesExtractButton", "Extract..."))
 						.OnClicked(this, &SPaperExtractSpritesDialog::ExtractClicked)
 					]
-					+ SHorizontalBox::Slot()
-					.AutoWidth()
+					+ SUniformGridPanel::Slot(1, 0)
 					[
 						SNew(SButton)
+						.ButtonStyle(FEditorStyle::Get(), "FlatButton")
+						.ForegroundColor(FLinearColor::White)
 						.Text(LOCTEXT("PaperExtractSpritesCancelButton", "Cancel"))
 						.OnClicked(this, &SPaperExtractSpritesDialog::CancelClicked)
 					]
