Paper2D: Changed tile maps imported from Tiled to always use positive separation between layers, reversing the layer order from Tiled since it orders back to front in the .JSON

diff --git a/Engine/Plugins/2D/Paper2D/Source/PaperTiledImporter/Private/PaperTiledImporterFactory.cpp b/Engine/Plugins/2D/Paper2D/Source/PaperTiledImporter/Private/PaperTiledImporterFactory.cpp
--- a/Engine/Plugins/2D/Paper2D/Source/PaperTiledImporter/Private/PaperTiledImporterFactory.cpp
+++ b/Engine/Plugins/2D/Paper2D/Source/PaperTiledImporter/Private/PaperTiledImporterFactory.cpp
@@ -206,7 +206,7 @@ UObject* UPaperTiledImporterFactory::FactoryCreateText(UClass* InClass, UObject*
 		Result->TileHeight = GlobalInfo.TileHeight;
 		Result->SeparationPerTileX = 0.0f;
 		Result->SeparationPerTileY = 0.0f;
-		Result->SeparationPerLayer = -1.0f;
+		Result->SeparationPerLayer = 1.0f;
 		Result->ProjectionMode = GlobalInfo.GetOrientationType();
 		Result->PixelsPerUnrealUnit = GetDefault<UPaperRuntimeSettings>()->DefaultPixelsPerUnrealUnit;
 		Result->BackgroundColor = GlobalInfo.BackgroundColor;
@@ -276,8 +276,10 @@ UObject* UPaperTiledImporterFactory::FactoryCreateText(UClass* InClass, UObject*
 		}
 
 		// Create the layers
-		for (const FTileLayerFromTiled& LayerData : GlobalInfo.Layers)
+		for (int32 LayerIndex = GlobalInfo.Layers.Num() - 1; LayerIndex >= 0; --LayerIndex)
 		{
+			const FTileLayerFromTiled& LayerData = GlobalInfo.Layers[LayerIndex];
+
 			if (LayerData.IsValid())
 			{
 				UPaperTileLayer* NewLayer = NewObject<UPaperTileLayer>(Result);
