Paper2D: Removed an obsolete comment

diff --git a/Engine/Plugins/2D/Paper2D/Source/Paper2DEditor/Private/SpriteEditor/SpriteEditor.cpp b/Engine/Plugins/2D/Paper2D/Source/Paper2DEditor/Private/SpriteEditor/SpriteEditor.cpp
--- a/Engine/Plugins/2D/Paper2D/Source/Paper2DEditor/Private/SpriteEditor/SpriteEditor.cpp
+++ b/Engine/Plugins/2D/Paper2D/Source/Paper2DEditor/Private/SpriteEditor/SpriteEditor.cpp
@@ -415,7 +415,7 @@ void FSpriteEditor::InitSpriteEditor(const EToolkitMode::Type Mode, const TShare
 			)
 		);
 
-	// Initialize the asset editor and spawn nothing (dummy layout)
+	// Initialize the asset editor
 	InitAssetEditor(Mode, InitToolkitHost, SpriteEditorAppName, StandaloneDefaultLayout, /*bCreateDefaultStandaloneMenu=*/ true, /*bCreateDefaultToolbar=*/ true, InitSprite);
 
 	// Extend things
