Tile Set Editor: Added an inset 'tile being edited' view
- Added a toolbar and 'single tile editor' viewport
- Refactored FTileSetEditor to no longer depend on FWorkflowCentricApplication

diff --git a/Engine/Plugins/2D/Paper2D/Source/Paper2DEditor/Private/TileSetEditor.cpp b/Engine/Plugins/2D/Paper2D/Source/Paper2DEditor/Private/TileSetEditor.cpp
--- a/Engine/Plugins/2D/Paper2D/Source/Paper2DEditor/Private/TileSetEditor.cpp
+++ b/Engine/Plugins/2D/Paper2D/Source/Paper2DEditor/Private/TileSetEditor.cpp
@@ -10,6 +10,7 @@
 #include "SPaperEditorViewport.h"
 #include "CanvasTypes.h"
 #include "CanvasItem.h"
+#include "SDockTab.h"
 
 #define LOCTEXT_NAMESPACE "TileSetEditor"
 
@@ -19,24 +20,19 @@ const FName TileSetEditorAppName = FName(TEXT("TileSetEditorApp"));
 
 //////////////////////////////////////////////////////////////////////////
 
-struct FTileSetEditorModes
-{
-	// Mode identifiers
-	static const FName StandardMode;
-};
-
 struct FTileSetEditorTabs
 {
 	// Tab identifiers
 	static const FName DetailsID;
 	static const FName TextureViewID;
+	static const FName SingleTileEditorID;
 };
 
 //////////////////////////////////////////////////////////////////////////
 
-const FName FTileSetEditorModes::StandardMode(TEXT("StandardMode"));
 const FName FTileSetEditorTabs::DetailsID(TEXT("Details"));
 const FName FTileSetEditorTabs::TextureViewID(TEXT("TextureCanvas"));
+const FName FTileSetEditorTabs::SingleTileEditorID(TEXT("SingleTileEditor"));
 
 //////////////////////////////////////////////////////////////////////////
 // FTileSetEditorViewportClient 
@@ -90,19 +86,6 @@ void FTileSetEditorViewportClient::Draw(FViewport* Viewport, FCanvas* Canvas)
 		Texture->SetForceMipLevelsToBeResident(30.0f);
 		Texture->WaitForStreaming();
 
-		// Draw the background checkerboard pattern in the same size/position as the render texture so it will show up anywhere
-		// the texture has transparency
-		//if (TextureEditorPtr.Pin()->GetIsCheckeredBackground())
-			// Handle case of using the checkerboard as a background
-// 			if (TextureEditorPtr.Pin()->GetCheckeredBackground_Fill())
-// 			{
-// 				DrawTile(Canvas, 0.0f, 0.0f, Viewport->GetSizeXY().X, Viewport->GetSizeXY().Y, 0.0f, 0.0f, (Viewport->GetSizeXY().X / CheckerboardTexture->GetSizeX()), (Viewport->GetSizeXY().Y / CheckerboardTexture->GetSizeY()), FLinearColor::White, CheckerboardTexture->Resource);
-// 			}
-// 			else
-// 			{
-// 				DrawTile(Canvas, XPos, YPos, Width, Height, 0.0f, 0.0f, (Width / CheckerboardTexture->GetSizeX()), (Height / CheckerboardTexture->GetSizeY()), FLinearColor::White, CheckerboardTexture->Resource);
-// 			}
-
 		FLinearColor TextureDrawColor = FLinearColor::White;
 
 		const float XPos = -ZoomPos.X * ZoomAmount;
@@ -128,7 +111,6 @@ void FTileSetEditorViewportClient::Draw(FViewport* Viewport, FCanvas* Canvas)
 		FCanvasBoxItem BoxItem(FVector2D(X, Y), FVector2D(W, H));
 		BoxItem.SetColor(Rect.Color);
 		Canvas->DrawItem(BoxItem);
-
 	}
 
 	if (TileIndex != INDEX_NONE)
@@ -238,6 +220,8 @@ void STileSetSelectorViewport::OnSelectionChanged(FMarqueeOperation Marquee, boo
 		SelectionDimensions.Y = 0;
 	}
 
+	OnTileSelectionChanged.Broadcast(SelectionTopLeft, SelectionDimensions);
+
 	const bool bHasSelection = (SelectionDimensions.X * SelectionDimensions.Y) > 0;
 	if (bIsPreview && bHasSelection)
 	{
@@ -280,32 +264,6 @@ void STileSetSelectorViewport::RefreshSelectionRectangle()
 	}
 }
 
-//////////////////////////////////////////////////////////////////////////
-// FTileSetViewportSummoner
-
-struct FTileSetViewportSummoner : public FWorkflowTabFactory
-{
-public:
-	FTileSetViewportSummoner(TSharedPtr<class FAssetEditorToolkit> InHostingApp)
-		: FWorkflowTabFactory(FTileSetEditorTabs::TextureViewID, InHostingApp)
-	{
-		TabLabel = LOCTEXT("TextureViewTabLabel", "Viewport");
-		//TabIcon = FEditorStyle::GetBrush("Kismet.Tabs.Variables");
-
-		bIsSingleton = true;
-
-		ViewMenuDescription = LOCTEXT("TextureViewTabMenu_Description", "Viewport");
-		ViewMenuTooltip = LOCTEXT("TextureViewTabMenu_ToolTip", "Shows the viewport");
-	}
-
-	virtual TSharedRef<SWidget> CreateTabBody(const FWorkflowTabSpawnInfo& Info) const override
-	{
-		TSharedPtr<FTileSetEditor> TileSetEditorPtr = StaticCastSharedPtr<FTileSetEditor>(HostingApp.Pin());
-
-		return SNew(STileSetSelectorViewport, TileSetEditorPtr->GetTileSetBeingEdited(), /*EdMode=*/ nullptr);
-	}
-};
-
 /////////////////////////////////////////////////////
 // STileSetPropertiesTabBody
 
@@ -344,57 +302,225 @@ public:
 	// End of SSingleObjectDetailsPanel interface
 };
 
+
 //////////////////////////////////////////////////////////////////////////
-// FTileSetDetailsSummoner
+// FSingleTileEditorViewportClient
 
-struct FTileSetDetailsSummoner : public FWorkflowTabFactory
+class FSingleTileEditorViewportClient : public FPaperEditorViewportClient
 {
 public:
-	FTileSetDetailsSummoner(TSharedPtr<class FAssetEditorToolkit> InHostingApp)
-		: FWorkflowTabFactory(FTileSetEditorTabs::DetailsID, InHostingApp)
+	FSingleTileEditorViewportClient(UPaperTileSet* InTileSet)
+		: TileSet(InTileSet)
+		, TileBeingEditedIndex(INDEX_NONE)
+		, PixelSize(4.0f)
 	{
-		TabLabel = LOCTEXT("DetailsTabLabel", "Details");
-		//TabIcon = FEditorStyle::GetBrush("Kismet.Tabs.Variables");
+		ZoomPos = FVector2D(-TileSet->TileWidth * ZoomAmount * PixelSize * 0.5f, -TileSet->TileHeight * ZoomAmount * PixelSize * 0.5f);
+	}
+
+	// FViewportClient interface
+	virtual void Draw(FViewport* Viewport, FCanvas* Canvas) override;
+	// End of FViewportClient interface
+
+	// FEditorViewportClient interface
+	virtual FLinearColor GetBackgroundColor() const override;
+	// End of FEditorViewportClient interface
+
+	void SetTileIndex(int32 InTileIndex);
+	int32 GetTileIndex() const;
+	void OnActiveTileIndexChanged(const FIntPoint& TopLeft, const FIntPoint& Dimensions);
+public:
+	// Tile set
+	UPaperTileSet* TileSet;
 
-		bIsSingleton = true;
+	int32 TileBeingEditedIndex;
 
-		ViewMenuDescription = LOCTEXT("DetailsTabMenu_Description", "Details");
-		ViewMenuTooltip = LOCTEXT("DetailsTabMenu_ToolTip", "Shows the details panel");
+	float PixelSize;
+};
+
+void FSingleTileEditorViewportClient::Draw(FViewport* Viewport, FCanvas* Canvas)
+{
+	// Super will clear the viewport
+	FPaperEditorViewportClient::Draw(Viewport, Canvas);
+
+	if (TileBeingEditedIndex == INDEX_NONE)
+	{
+		return;
 	}
 
-	virtual TSharedRef<SWidget> CreateTabBody(const FWorkflowTabSpawnInfo& Info) const override
+	if (UTexture2D* Texture = TileSet->TileSheet)
 	{
-		TSharedPtr<FTileSetEditor> TileSetEditorPtr = StaticCastSharedPtr<FTileSetEditor>(HostingApp.Pin());
-		return SNew(STileSetPropertiesTabBody, TileSetEditorPtr);
+		const bool bUseTranslucentBlend = Texture->HasAlphaChannel();
+
+		// Fully stream in the texture before drawing it.
+		Texture->SetForceMipLevelsToBeResident(30.0f);
+		Texture->WaitForStreaming();
+
+		const float XPos = -ZoomPos.X * ZoomAmount;
+		const float YPos = -ZoomPos.Y * ZoomAmount;
+		const float Width = TileSet->TileWidth * ZoomAmount * PixelSize;
+		const float Height = TileSet->TileHeight * ZoomAmount * PixelSize;
+
+		const float InvTextureWidth = 1.0f / Texture->GetSurfaceWidth();
+		const float InvTextureHeight = 1.0f / Texture->GetSurfaceHeight();
+
+		FVector2D TopLeft;
+		TileSet->GetTileUV(TileBeingEditedIndex, /*out*/ TopLeft);
+		const FVector2D UV0(TopLeft.X * InvTextureWidth, TopLeft.Y * InvTextureHeight);
+
+		const FVector2D UVSize(TileSet->TileWidth * InvTextureWidth, TileSet->TileHeight * InvTextureHeight);
+
+		const FLinearColor TextureDrawColor = FLinearColor::White;
+		Canvas->DrawTile(XPos, YPos, Width, Height, UV0.X, UV0.Y, UVSize.X, UVSize.Y, TextureDrawColor, Texture->Resource, bUseTranslucentBlend);
 	}
-};
+}
+
+
+FLinearColor FSingleTileEditorViewportClient::GetBackgroundColor() const
+{
+	return TileSet->BackgroundColor;
+}
+
+void FSingleTileEditorViewportClient::SetTileIndex(int32 InTileIndex)
+{
+	const bool bNewIndexValid = (InTileIndex >= 0) && (InTileIndex < TileSet->GetTileCount());
+	TileBeingEditedIndex = bNewIndexValid ? InTileIndex : INDEX_NONE;
+	Invalidate();
+}
+
+void FSingleTileEditorViewportClient::OnActiveTileIndexChanged(const FIntPoint& TopLeft, const FIntPoint& Dimensions)
+{
+	const int32 NewIndex = (TileSet->GetTileCountX() * TopLeft.Y) + TopLeft.X;
+	SetTileIndex(NewIndex);
+}
+
+int32 FSingleTileEditorViewportClient::GetTileIndex() const
+{
+	return TileBeingEditedIndex;
+}
 
 //////////////////////////////////////////////////////////////////////////
-// FBasicTileSetEditorMode
+// SSingleTileEditorViewport
 
-class FBasicTileSetEditorMode : public FApplicationMode
+class SSingleTileEditorViewport : public SPaperEditorViewport
 {
 public:
-	FBasicTileSetEditorMode(TSharedPtr<class FTileSetEditor> InTileSetEditor, FName InModeName);
+	SLATE_BEGIN_ARGS(SSingleTileEditorViewport) {}
+	SLATE_END_ARGS()
+
+	~SSingleTileEditorViewport();
 
-	// FApplicationMode interface
-	virtual void RegisterTabFactories(TSharedPtr<FTabManager> InTabManager) override;
-	// End of FApplicationMode interface
+	void Construct(const FArguments& InArgs, TSharedPtr<class FSingleTileEditorViewportClient> InViewportClient);
 
 protected:
-	TWeakPtr<FTileSetEditor> MyTileSetEditor;
-	FWorkflowAllowedTabSet TabFactories;
+	// SPaperEditorViewport interface
+	virtual FText GetTitleText() const override;
+	// End of SPaperEditorViewport interface
+
+private:
+	TSharedPtr<class FSingleTileEditorViewportClient> TypedViewportClient;
+	class FEdModeTileMap* TileMapEditor;
 };
 
-FBasicTileSetEditorMode::FBasicTileSetEditorMode(TSharedPtr<class FTileSetEditor> InTileSetEditor, FName InModeName)
-	: FApplicationMode(InModeName)
+//////////////////////////////////////////////////////////////////////////
+// SSingleTileEditorViewport
+
+SSingleTileEditorViewport::~SSingleTileEditorViewport()
+{
+	TypedViewportClient.Reset();
+}
+
+void SSingleTileEditorViewport::Construct(const FArguments& InArgs, TSharedPtr<class FSingleTileEditorViewportClient> InViewportClient)
+{
+	TypedViewportClient = InViewportClient;
+
+	SPaperEditorViewport::Construct(
+		SPaperEditorViewport::FArguments(),
+		TypedViewportClient.ToSharedRef());
+
+	// Make sure we get input instead of the viewport stealing it
+	ViewportWidget->SetVisibility(EVisibility::HitTestInvisible);
+}
+
+FText SSingleTileEditorViewport::GetTitleText() const
+{
+	const int32 CurrentTileIndex = TypedViewportClient->GetTileIndex();
+	if (CurrentTileIndex != INDEX_NONE)
+	{
+		FNumberFormattingOptions NoGroupingFormat;
+		NoGroupingFormat.SetUseGrouping(false);
+
+		return FText::Format(LOCTEXT("SingleTileEditorViewportTitle", "Editing tile #{0}"), FText::AsNumber(CurrentTileIndex, &NoGroupingFormat));
+	}
+	else
+	{
+		return LOCTEXT("SingleTileEditorViewportTitle_NoTile", "Tile Editor - Select a tile");
+	}
+}
+
+//////////////////////////////////////////////////////////////////////////
+// FTileSetEditor
+
+FTileSetEditor::FTileSetEditor()
+{
+	// Register to be notified when properties are edited
+	FCoreUObjectDelegates::FOnObjectPropertyChanged::FDelegate OnPropertyChangedDelegate = FCoreUObjectDelegates::FOnObjectPropertyChanged::FDelegate::CreateRaw(this, &FTileSetEditor::OnPropertyChanged);
+	OnPropertyChangedHandle = FCoreUObjectDelegates::OnObjectPropertyChanged.Add(OnPropertyChangedDelegate);
+}
+
+FTileSetEditor::~FTileSetEditor()
+{
+	// Unregister the property modification handler
+	FCoreUObjectDelegates::OnObjectPropertyChanged.Remove(OnPropertyChangedHandle);
+}
+
+void FTileSetEditor::RegisterTabSpawners(const TSharedRef<class FTabManager>& TabManager)
+{
+	WorkspaceMenuCategory = TabManager->AddLocalWorkspaceMenuCategory(LOCTEXT("WorkspaceMenu_TileSetEditor", "Tile Set Editor"));
+	TSharedRef<FWorkspaceItem> WorkspaceMenuCategoryRef = WorkspaceMenuCategory.ToSharedRef();
+
+	FAssetEditorToolkit::RegisterTabSpawners(TabManager);
+
+	TabManager->RegisterTabSpawner(FTileSetEditorTabs::TextureViewID, FOnSpawnTab::CreateSP(this, &FTileSetEditor::SpawnTab_TextureView))
+		.SetDisplayName(LOCTEXT("TextureViewTabMenu_Description", "Viewport"))
+		.SetTooltipText(LOCTEXT("TextureViewTabMenu_ToolTip", "Shows the viewport"))
+		.SetGroup(WorkspaceMenuCategoryRef)
+		.SetIcon(FSlateIcon(FEditorStyle::GetStyleSetName(), "LevelEditor.Tabs.Viewports"));
+
+	TabManager->RegisterTabSpawner(FTileSetEditorTabs::DetailsID, FOnSpawnTab::CreateSP(this, &FTileSetEditor::SpawnTab_Details))
+		.SetDisplayName(LOCTEXT("DetailsTabLabel", "Details"))
+		.SetGroup(WorkspaceMenuCategoryRef)
+		.SetIcon(FSlateIcon(FEditorStyle::GetStyleSetName(), "LevelEditor.Tabs.Details"));
+
+	TabManager->RegisterTabSpawner(FTileSetEditorTabs::SingleTileEditorID, FOnSpawnTab::CreateSP(this, &FTileSetEditor::SpawnTab_SingleTileEditor))
+		.SetDisplayName(LOCTEXT("SingleTileEditTabMenu_Description", "Tile Editor"))
+		.SetTooltipText(LOCTEXT("SingleTileEditTabMenu_ToolTip", "Shows the tile editor"))
+		.SetGroup(WorkspaceMenuCategoryRef)
+		.SetIcon(FSlateIcon(FEditorStyle::GetStyleSetName(), "LevelEditor.Tabs.Viewports"));
+}
+
+void FTileSetEditor::UnregisterTabSpawners(const TSharedRef<class FTabManager>& TabManager)
+{
+	FAssetEditorToolkit::UnregisterTabSpawners(TabManager);
+
+	TabManager->UnregisterTabSpawner(FTileSetEditorTabs::TextureViewID);
+	TabManager->UnregisterTabSpawner(FTileSetEditorTabs::DetailsID);
+	TabManager->UnregisterTabSpawner(FTileSetEditorTabs::SingleTileEditorID);
+}
+
+void FTileSetEditor::InitTileSetEditor(const EToolkitMode::Type Mode, const TSharedPtr< class IToolkitHost >& InitToolkitHost, class UPaperTileSet* InitTileSet)
 {
-	MyTileSetEditor = InTileSetEditor;
+	FAssetEditorManager::Get().CloseOtherEditors(InitTileSet, this);
+	TileSetBeingEdited = InitTileSet;
 
-	TabFactories.RegisterFactory(MakeShareable(new FTileSetViewportSummoner(InTileSetEditor)));
-	TabFactories.RegisterFactory(MakeShareable(new FTileSetDetailsSummoner(InTileSetEditor)));
+	TileSetViewport = SNew(STileSetSelectorViewport, InitTileSet, /*EdMode=*/ nullptr);
+	TileEditorViewportClient = MakeShareable(new FSingleTileEditorViewportClient(InitTileSet));
+	TileSetViewport->GetTileSelectionChanged().AddRaw(TileEditorViewportClient.Get(), &FSingleTileEditorViewportClient::OnActiveTileIndexChanged);
 
-	TabLayout = FTabManager::NewLayout("Standalone_TileSetEditor_Layout_v2")
+	//@TODO: FTileSetEditorCommands::Register();
+	BindCommands();
+
+	// Default layout
+	const TSharedRef<FTabManager::FLayout> StandaloneDefaultLayout = FTabManager::NewLayout("Standalone_TileSetEditor_Layout_v3")
 		->AddArea
 		(
 			FTabManager::NewPrimaryArea()
@@ -402,49 +528,126 @@ FBasicTileSetEditorMode::FBasicTileSetEditorMode(TSharedPtr<class FTileSetEditor
 			->Split
 			(
 				FTabManager::NewStack()
-				->SetSizeCoefficient(0.8f)
+				->SetSizeCoefficient(0.1f)
 				->SetHideTabWell(true)
-				->AddTab(FTileSetEditorTabs::TextureViewID, ETabState::OpenedTab)
+				->AddTab(GetToolbarTabId(), ETabState::OpenedTab)
 			)
 			->Split
 			(
-				FTabManager::NewStack()
-				->SetSizeCoefficient(0.2f)
-				->SetHideTabWell(true)
-				->AddTab(FTileSetEditorTabs::DetailsID, ETabState::OpenedTab)
+				FTabManager::NewSplitter()
+				->SetOrientation(Orient_Horizontal)
+				->SetSizeCoefficient(0.9f)
+				->Split
+				(
+					FTabManager::NewStack()
+					->SetSizeCoefficient(0.6f)
+					->SetHideTabWell(true)
+					->AddTab(FTileSetEditorTabs::TextureViewID, ETabState::OpenedTab)
+				)
+				->Split
+				(
+					FTabManager::NewSplitter()
+					->SetOrientation(Orient_Vertical)
+					->SetSizeCoefficient(0.4f)
+					->Split
+					(
+						FTabManager::NewStack()
+						->SetSizeCoefficient(0.5f)
+						->AddTab(FTileSetEditorTabs::SingleTileEditorID, ETabState::OpenedTab)
+					)
+					->Split
+					(
+						FTabManager::NewStack()
+						->SetSizeCoefficient(0.5f)
+						->AddTab(FTileSetEditorTabs::DetailsID, ETabState::OpenedTab)
+					)
+				)
 			)
 		);
+
+
+	// Initialize the asset editor
+	InitAssetEditor(Mode, InitToolkitHost, TileSetEditorAppName, StandaloneDefaultLayout, /*bCreateDefaultStandaloneMenu=*/ true, /*bCreateDefaultToolbar=*/ true, InitTileSet);
+
+	// Extend things
+	ExtendMenu();
+	ExtendToolbar();
+	RegenerateMenusAndToolbars();
 }
 
-void FBasicTileSetEditorMode::RegisterTabFactories(TSharedPtr<FTabManager> InTabManager)
+TSharedRef<SDockTab> FTileSetEditor::SpawnTab_TextureView(const FSpawnTabArgs& Args)
 {
-	TSharedPtr<FTileSetEditor> Editor = MyTileSetEditor.Pin();
+	return SNew(SDockTab)
+		.Label(LOCTEXT("TextureViewTabLabel", "Viewport"))
+		[
+			TileSetViewport.ToSharedRef()
+		];
+}
 
-	Editor->PushTabFactories(TabFactories);
+TSharedRef<SDockTab> FTileSetEditor::SpawnTab_Details(const FSpawnTabArgs& Args)
+{
+	TSharedPtr<FTileSetEditor> TileSetEditorPtr = SharedThis(this);
+
+	return SNew(SDockTab)
+		.Label(LOCTEXT("DetailsTabLabel", "Details"))
+		[
+			SNew(STileSetPropertiesTabBody, TileSetEditorPtr)
+		];
 }
 
-//////////////////////////////////////////////////////////////////////////
-// FTileSetEditor
+TSharedRef<SDockTab> FTileSetEditor::SpawnTab_SingleTileEditor(const FSpawnTabArgs& Args)
+{
+	TSharedPtr<FTileSetEditor> TileSetEditorPtr = SharedThis(this);
 
-void FTileSetEditor::RegisterTabSpawners(const TSharedRef<class FTabManager>& TabManager)
+	return SNew(SDockTab)
+		.Label(LOCTEXT("SingleTileEditTabLabel", "Tile Editor"))
+		[
+			SNew(SSingleTileEditorViewport, TileEditorViewportClient)
+		];
+}
+
+void FTileSetEditor::BindCommands()
 {
-	FWorkflowCentricApplication::RegisterTabSpawners(TabManager);
 }
 
-void FTileSetEditor::InitTileSetEditor(const EToolkitMode::Type Mode, const TSharedPtr< class IToolkitHost >& InitToolkitHost, class UPaperTileSet* InitTileSet)
+void FTileSetEditor::ExtendMenu()
 {
-	FAssetEditorManager::Get().CloseOtherEditors(InitTileSet, this);
-	TileSetBeingEdited = InitTileSet;
+}
 
-	// Initialize the asset editor and spawn nothing (dummy layout)
-	const TSharedRef<FTabManager::FLayout> DummyLayout = FTabManager::NewLayout("NullLayout")->AddArea(FTabManager::NewPrimaryArea());
-	InitAssetEditor(Mode, InitToolkitHost, TileSetEditorAppName, DummyLayout, /*bCreateDefaultStandaloneMenu=*/ true, /*bCreateDefaultToolbar=*/ false, InitTileSet);
+void FTileSetEditor::ExtendToolbar()
+{
+//@TODO: Add geometry creation commands
+// 	struct Local
+// 	{
+// 		static void FillToolbar(FToolBarBuilder& ToolbarBuilder)
+// 		{
+// 			ToolbarBuilder.BeginSection("Tools");
+// 			{
+// 				ToolbarBuilder.AddToolBarButton(FSpriteEditorCommands::Get().AddBoxShape);
+// 				ToolbarBuilder.AddToolBarButton(FSpriteEditorCommands::Get().ToggleAddPolygonMode);
+// 				ToolbarBuilder.AddToolBarButton(FSpriteEditorCommands::Get().AddCircleShape);
+// 				ToolbarBuilder.AddToolBarButton(FSpriteEditorCommands::Get().SnapAllVertices);
+// 			}
+// 			ToolbarBuilder.EndSection();
+// 		}
+// 	};
+// 
+// 	TSharedPtr<FExtender> ToolbarExtender = MakeShareable(new FExtender);
+// 
+// 	ToolbarExtender->AddToolBarExtension(
+// 		"Asset",
+// 		EExtensionHook::After,
+// 		ViewportPtr->GetCommandList(),
+// 		FToolBarExtensionDelegate::CreateStatic(&Local::FillToolbar)
+// 		);
+}
 
-	// Create the modes and activate one (which will populate with a real layout)
-	AddApplicationMode(
-		FTileSetEditorModes::StandardMode, 
-		MakeShareable(new FBasicTileSetEditorMode(SharedThis(this), FTileSetEditorModes::StandardMode)));
-	SetCurrentMode(FTileSetEditorModes::StandardMode);
+void FTileSetEditor::OnPropertyChanged(UObject* ObjectBeingModified, FPropertyChangedEvent& PropertyChangedEvent)
+{
+	if (ObjectBeingModified == TileSetBeingEdited)
+	{
+		TileEditorViewportClient->SetTileIndex(TileEditorViewportClient->GetTileIndex());
+	}
 }
 
 FName FTileSetEditor::GetToolkitFName() const
@@ -477,6 +680,12 @@ FLinearColor FTileSetEditor::GetWorldCentricTabColorScale() const
 	return FLinearColor::White;
 }
 
+FString FTileSetEditor::GetDocumentationLink() const
+{
+	//@TODO: Need to make a page for this
+	return TEXT("Engine/Paper2D/TileSetEditor");
+}
+
 void FTileSetEditor::AddReferencedObjects(FReferenceCollector& Collector)
 {
 	Collector.AddReferencedObject(TileSetBeingEdited);
diff --git a/Engine/Plugins/2D/Paper2D/Source/Paper2DEditor/Private/TileSetEditor.h b/Engine/Plugins/2D/Paper2D/Source/Paper2DEditor/Private/TileSetEditor.h
--- a/Engine/Plugins/2D/Paper2D/Source/Paper2DEditor/Private/TileSetEditor.h
+++ b/Engine/Plugins/2D/Paper2D/Source/Paper2DEditor/Private/TileSetEditor.h
@@ -2,17 +2,22 @@
 
 #pragma once
 
-#include "WorkflowCentricApplication.h"
+#include "Toolkits/AssetEditorToolkit.h"
+#include "Toolkits/AssetEditorManager.h"
 #include "SPaperEditorViewport.h"
 
 //////////////////////////////////////////////////////////////////////////
 // FTileSetEditor
 
-class FTileSetEditor : public FWorkflowCentricApplication, public FGCObject
+class FTileSetEditor : public FAssetEditorToolkit, public FGCObject
 {
 public:
+	FTileSetEditor();
+	~FTileSetEditor();
+
 	// IToolkit interface
 	virtual void RegisterTabSpawners(const TSharedRef<class FTabManager>& TabManager) override;
+	virtual void UnregisterTabSpawners(const TSharedRef<class FTabManager>& TabManager) override;
 	// End of IToolkit interface
 
 	// FAssetEditorToolkit
@@ -21,6 +26,7 @@ public:
 	virtual FText GetToolkitName() const override;
 	virtual FLinearColor GetWorldCentricTabColorScale() const override;
 	virtual FString GetWorldCentricTabPrefix() const override;
+	virtual FString GetDocumentationLink() const override;
 	// End of FAssetEditorToolkit
 
 	// FSerializableObject interface
@@ -31,8 +37,25 @@ public:
 	void InitTileSetEditor(const EToolkitMode::Type Mode, const TSharedPtr< class IToolkitHost >& InitToolkitHost, class UPaperTileSet* InitTileSet);
 
 	UPaperTileSet* GetTileSetBeingEdited() const { return TileSetBeingEdited; }
+
+protected:
+	TSharedRef<class SDockTab> SpawnTab_TextureView(const FSpawnTabArgs& Args);
+	TSharedRef<class SDockTab> SpawnTab_Details(const FSpawnTabArgs& Args);
+	TSharedRef<class SDockTab> SpawnTab_SingleTileEditor(const FSpawnTabArgs& Args);
+
+	void BindCommands();
+	void ExtendMenu();
+	void ExtendToolbar();
+
+	void OnPropertyChanged(UObject* ObjectBeingModified, FPropertyChangedEvent& PropertyChangedEvent);
+
 protected:
 	UPaperTileSet* TileSetBeingEdited;
+
+	TSharedPtr<class STileSetSelectorViewport> TileSetViewport;
+	TSharedPtr<class FSingleTileEditorViewportClient> TileEditorViewportClient;
+
+	FDelegateHandle OnPropertyChangedHandle;
 };
 
 //////////////////////////////////////////////////////////////////////////
@@ -41,6 +64,8 @@ protected:
 class STileSetSelectorViewport : public SPaperEditorViewport
 {
 public:
+	DECLARE_MULTICAST_DELEGATE_TwoParams(FOnTileViewportSelectionChanged, const FIntPoint& /*TopLeft*/, const FIntPoint& /*Dimensions*/);
+
 	SLATE_BEGIN_ARGS(STileSetSelectorViewport) {}
 	SLATE_END_ARGS()
 
@@ -49,6 +74,12 @@ public:
 	void Construct(const FArguments& InArgs, UPaperTileSet* InTileSet, class FEdModeTileMap* InTileMapEditor);
 
 	void ChangeTileSet(UPaperTileSet* InTileSet);
+
+	FOnTileViewportSelectionChanged& GetTileSelectionChanged()
+	{
+		return OnTileSelectionChanged;
+	}
+
 protected:
 	// SPaperEditorViewport interface
 	virtual FText GetTitleText() const override;
@@ -64,4 +95,6 @@ private:
 	class FEdModeTileMap* TileMapEditor;
 	FIntPoint SelectionTopLeft;
 	FIntPoint SelectionDimensions;
+
+	FOnTileViewportSelectionChanged OnTileSelectionChanged;
 };
\ No newline at end of file
