Tile Maps: Added a call to ConditionGeometry when importing geometry from Tiled to extract box primitives from polygons that meet the criteria

diff --git a/Engine/Plugins/2D/Paper2D/Source/PaperTiledImporter/Private/PaperTiledImporterFactory.cpp b/Engine/Plugins/2D/Paper2D/Source/PaperTiledImporter/Private/PaperTiledImporterFactory.cpp
--- a/Engine/Plugins/2D/Paper2D/Source/PaperTiledImporter/Private/PaperTiledImporterFactory.cpp
+++ b/Engine/Plugins/2D/Paper2D/Source/PaperTiledImporter/Private/PaperTiledImporterFactory.cpp
@@ -1197,6 +1197,8 @@ void FTiledObject::AddToSpriteGeometryCollection(const FVector2D& Offset, const
 			InOutShapes.Shapes.Last().Rotation = RotationUnwound;
 		}
 	}
+
+	InOutShapes.ConditionGeometry();
 }
 
 //////////////////////////////////////////////////////////////////////////
