Tile Maps: Added a CollisionThickness property and made CollisionThickness/SpriteCollisionDomain protected with getters

diff --git a/Engine/Plugins/2D/Paper2D/Source/Paper2D/Classes/PaperTileMap.h b/Engine/Plugins/2D/Paper2D/Source/Paper2D/Classes/PaperTileMap.h
--- a/Engine/Plugins/2D/Paper2D/Source/Paper2D/Classes/PaperTileMap.h
+++ b/Engine/Plugins/2D/Paper2D/Source/Paper2D/Classes/PaperTileMap.h
@@ -77,10 +77,16 @@ class PAPER2D_API UPaperTileMap : public UObject
 	UPROPERTY(Instanced, Category=Sprite, BlueprintReadOnly)
 	TArray<class UPaperTileLayer*> TileLayers;
 
+protected:
+	// The extrusion thickness of collision geometry when using a 3D collision domain
+	UPROPERTY(Category=Collision, EditAnywhere, BlueprintReadOnly)
+	float CollisionThickness;
+
 	// Collision domain (no collision, 2D, or 3D)
 	UPROPERTY(Category=Collision, EditAnywhere, BlueprintReadOnly)
 	TEnumAsByte<ESpriteCollisionMode::Type> SpriteCollisionDomain;
 
+public:
 	// Tile map type
 	UPROPERTY(Category=Setup, EditAnywhere, BlueprintReadOnly)
 	TEnumAsByte<ETileMapProjectionMode::Type> ProjectionMode;
@@ -133,6 +139,17 @@ public:
 	void GetTileToLocalParameters(FVector& OutCornerPosition, FVector& OutStepX, FVector& OutStepY, FVector& OutOffsetYFactor) const;
 	void GetLocalToTileParameters(FVector& OutCornerPosition, FVector& OutStepX, FVector& OutStepY, FVector& OutOffsetYFactor) const;
 
+	// Returns the extrusion thickness of collision geometry when using a 3D collision domain
+	float GetCollisionThickness() const
+	{
+		return CollisionThickness;
+	}
+
+	// Returns the collision domain (no collision, 2D, or 3D)
+	ESpriteCollisionMode::Type GetSpriteCollisionDomain() const
+	{
+		return SpriteCollisionDomain;
+	}
 
 	FBoxSphereBounds GetRenderBounds() const;
 
diff --git a/Engine/Plugins/2D/Paper2D/Source/Paper2D/Private/PaperTileMap.cpp b/Engine/Plugins/2D/Paper2D/Source/Paper2D/Private/PaperTileMap.cpp
--- a/Engine/Plugins/2D/Paper2D/Source/Paper2D/Private/PaperTileMap.cpp
+++ b/Engine/Plugins/2D/Paper2D/Source/Paper2D/Private/PaperTileMap.cpp
@@ -20,6 +20,7 @@ UPaperTileMap::UPaperTileMap(const FObjectInitializer& ObjectInitializer)
 	SeparationPerTileX = 0.0f;
 	SeparationPerTileY = 0.0f;
 	SeparationPerLayer = 4.0f;
+	CollisionThickness = 50.0f;
 	SpriteCollisionDomain = ESpriteCollisionMode::None;
 
 #if WITH_EDITORONLY_DATA
@@ -149,18 +150,10 @@ void UPaperTileMap::UpdateBodySetup()
 	switch (SpriteCollisionDomain)
 	{
 	case ESpriteCollisionMode::Use3DPhysics:
-		BodySetup = nullptr;
-		if (BodySetup == nullptr)
-		{
-			BodySetup = NewObject<UBodySetup>(this);
-		}
+		BodySetup = NewObject<UBodySetup>(this);
 		break;
 	case ESpriteCollisionMode::Use2DPhysics:
-		BodySetup = nullptr;
-		if (BodySetup == nullptr)
-		{
-			BodySetup = NewObject<UBodySetup2D>(this);
-		}
+		BodySetup = NewObject<UBodySetup2D>(this);
 		break;
 	case ESpriteCollisionMode::None:
 		BodySetup = nullptr;
@@ -169,18 +162,16 @@ void UPaperTileMap::UpdateBodySetup()
 
 	if (SpriteCollisionDomain != ESpriteCollisionMode::None)
 	{
-		if (SpriteCollisionDomain == ESpriteCollisionMode::Use3DPhysics)
-		{
-			BodySetup->CollisionTraceFlag = CTF_UseSimpleAsComplex;
+		BodySetup->CollisionTraceFlag = CTF_UseSimpleAsComplex;
 
-			BodySetup->AggGeom.BoxElems.Empty();
-			for (int32 LayerIndex = 0; LayerIndex < TileLayers.Num(); ++LayerIndex)
-			{
-				TileLayers[LayerIndex]->AugmentBodySetup(BodySetup);
-			}
+		for (int32 LayerIndex = 0; LayerIndex < TileLayers.Num(); ++LayerIndex)
+		{
+			TileLayers[LayerIndex]->AugmentBodySetup(BodySetup);
 		}
 
-		//@TODO: BOX2D: Add support for 2D physics on tile maps
+		// Finalize the BodySetup
+		BodySetup->InvalidatePhysicsData();
+		BodySetup->CreatePhysicsMeshes();
 	}
 }
 
