Tile Map Editor: Removed some commented out code and unused methods

diff --git a/Engine/Plugins/2D/Paper2D/Source/Paper2DEditor/Private/TileMapEditing/TileMapEditor.cpp b/Engine/Plugins/2D/Paper2D/Source/Paper2DEditor/Private/TileMapEditing/TileMapEditor.cpp
--- a/Engine/Plugins/2D/Paper2D/Source/Paper2DEditor/Private/TileMapEditing/TileMapEditor.cpp
+++ b/Engine/Plugins/2D/Paper2D/Source/Paper2DEditor/Private/TileMapEditing/TileMapEditor.cpp
@@ -359,24 +359,7 @@ void FTileMapEditor::InitTileMapEditor(const EToolkitMode::Type Mode, const TSha
 
 void FTileMapEditor::BindCommands()
 {
-	// 	const FTileMapEditorCommands& Commands = FTileMapEditorCommands::Get();
-	// 
-	// 	const TSharedRef<FUICommandList>& UICommandList = GetToolkitCommands();
-	// 
-	// 	UICommandList->MapAction( FGenericCommands::Get().Delete,
-	// 		FExecuteAction::CreateSP( this, &FStaticMeshEditor::DeleteSelectedSockets ),
-	// 		FCanExecuteAction::CreateSP( this, &FStaticMeshEditor::HasSelectedSockets ));
-	// 
-	// 	UICommandList->MapAction( FGenericCommands::Get().Undo, 
-	// 		FExecuteAction::CreateSP( this, &FStaticMeshEditor::UndoAction ) );
-	// 
-	// 	UICommandList->MapAction( FGenericCommands::Get().Redo, 
-	// 		FExecuteAction::CreateSP( this, &FStaticMeshEditor::RedoAction ) );
-	// 
-	// 	UICommandList->MapAction(
-	// 		FGenericCommands::Get().Duplicate,
-	// 		FExecuteAction::CreateSP(this, &FStaticMeshEditor::DuplicateSelectedSocket),
-	// 		FCanExecuteAction::CreateSP(this, &FStaticMeshEditor::HasSelectedSockets));
+	// Commands would go here
 }
 
 FName FTileMapEditor::GetToolkitFName() const
@@ -464,24 +447,7 @@ void FTileMapEditor::ExtendToolbar()
 	{
 		static void FillToolbar(FToolBarBuilder& ToolbarBuilder)
 		{
-			// 			ToolbarBuilder.BeginSection("Realtime");
-			// 			{
-			// 				ToolbarBuilder.AddToolBarButton(FEditorViewportCommands::Get().ToggleRealTime);
-			// 			}
-			// 			ToolbarBuilder.EndSection();
-
-// 			ToolbarBuilder.BeginSection("Command");
-// 			{
-// 				ToolbarBuilder.AddToolBarButton(FTileMapEditorCommands::Get().SetShowSourceTexture);
-// 			}
-// 			ToolbarBuilder.EndSection();
-// 
-// 			ToolbarBuilder.BeginSection("Tools");
-// 			{
-// 				ToolbarBuilder.AddToolBarButton(FTileMapEditorCommands::Get().AddPolygon);
-// 				ToolbarBuilder.AddToolBarButton(FTileMapEditorCommands::Get().SnapAllVertices);
-// 			}
-// 			ToolbarBuilder.EndSection();
+			// Toolbar buttons would go here
 		}
 	};
 
@@ -494,16 +460,7 @@ void FTileMapEditor::ExtendToolbar()
 		FToolBarExtensionDelegate::CreateStatic(&Local::FillToolbar)
 		);
 
-// 	ToolbarExtender->AddToolBarExtension(
-// 		"Asset",
-// 		EExtensionHook::After,
-// 		ViewportPtr->GetCommandList(),
-// 		FToolBarExtensionDelegate::CreateSP(this, &FTileMapEditor::CreateModeToolbarWidgets));
-
 	AddToolbarExtender(ToolbarExtender);
-// 
-// 	IPaper2DEditorModule* Paper2DEditorModule = &FModuleManager::LoadModuleChecked<IPaper2DEditorModule>("Paper2DEditor");
-// 	AddToolbarExtender(Paper2DEditorModule->GetTileMapEditorToolBarExtensibilityManager()->GetAllExtenders());
 }
 
 void FTileMapEditor::SetTileMapBeingEdited(UPaperTileMap* NewTileMap)
@@ -522,17 +479,6 @@ void FTileMapEditor::SetTileMapBeingEdited(UPaperTileMap* NewTileMap)
 	}
 }
 
-
-void FTileMapEditor::CreateModeToolbarWidgets(FToolBarBuilder& IgnoredBuilder)
-{
-// 	FToolBarBuilder ToolbarBuilder(ViewportPtr->GetCommandList(), FMultiBoxCustomization::None);
-// 	ToolbarBuilder.AddToolBarButton(FTileMapEditorCommands::Get().EnterViewMode);
-// 	ToolbarBuilder.AddToolBarButton(FTileMapEditorCommands::Get().EnterSourceRegionEditMode);
-// 	ToolbarBuilder.AddToolBarButton(FTileMapEditorCommands::Get().EnterCollisionEditMode);
-// 	ToolbarBuilder.AddToolBarButton(FTileMapEditorCommands::Get().EnterRenderingEditMode);
-// 	AddToolbarWidget(ToolbarBuilder.MakeWidget());
-}
-
 //////////////////////////////////////////////////////////////////////////
 
 #undef LOCTEXT_NAMESPACE
diff --git a/Engine/Plugins/2D/Paper2D/Source/Paper2DEditor/Private/TileMapEditing/TileMapEditor.h b/Engine/Plugins/2D/Paper2D/Source/Paper2DEditor/Private/TileMapEditing/TileMapEditor.h
--- a/Engine/Plugins/2D/Paper2D/Source/Paper2DEditor/Private/TileMapEditing/TileMapEditor.h
+++ b/Engine/Plugins/2D/Paper2D/Source/Paper2DEditor/Private/TileMapEditing/TileMapEditor.h
@@ -48,6 +48,4 @@ protected:
 	TSharedRef<SDockTab> SpawnTab_Viewport(const FSpawnTabArgs& Args);
 	TSharedRef<SDockTab> SpawnTab_ToolboxHost(const FSpawnTabArgs& Args);
 	TSharedRef<SDockTab> SpawnTab_Details(const FSpawnTabArgs& Args);
-
-	void CreateModeToolbarWidgets(FToolBarBuilder& ToolbarBuilder);
 };
