Paper2D: Broke out the code that takes a FSpriteGeometryCollection and converts it into collision geometry in a UBodySetup into a new helper class FSpriteGeometryCollisionBuilderBase for reuse elsewhere
- Added support for an arbitrary PaperAxisZ offset when generating collision geometry
- Added getters for CollisionThickness and SpriteCollisionDomain to UPaperSprite

diff --git a/Engine/Plugins/2D/Paper2D/Source/Paper2D/Classes/PaperSprite.h b/Engine/Plugins/2D/Paper2D/Source/Paper2D/Classes/PaperSprite.h
--- a/Engine/Plugins/2D/Paper2D/Source/Paper2D/Classes/PaperSprite.h
+++ b/Engine/Plugins/2D/Paper2D/Source/Paper2D/Classes/PaperSprite.h
@@ -7,22 +7,6 @@
 #include "Interfaces/Interface_CollisionDataProvider.h"
 #include "PaperSprite.generated.h"
 
-UENUM()
-namespace ESpriteCollisionMode
-{
-	enum Type
-	{
-		// Should this have no collison and not participate in physics?
-		None,
-
-		// EXPERIMENTAL: Should this have 2D collision geometry and participate in the 2D physics world?
-		Use2DPhysics UMETA(DisplayName = "Use 2D Physics"),
-
-		// Should this have 3D collision geometry and participate in the 3D physics world?
-		Use3DPhysics UMETA(DisplayName = "Use 3D Physics")
-	};
-}
-
 //@TODO: Should have some nice UI and enforce unique names, etc...
 USTRUCT()
 struct FPaperSpriteSocket
@@ -203,6 +187,18 @@ public:
 	// Returns the current pivot position in texture space
 	FVector2D GetPivotPosition() const;
 
+	// Returns the extrusion thickness of collision geometry when using a 3D collision domain
+	float GetCollisionThickness() const
+	{
+		return CollisionThickness;
+	}
+
+	// Returns the collision domain (no collision, 2D, or 3D)
+	ESpriteCollisionMode::Type GetSpriteCollisionDomain() const
+	{
+		return SpriteCollisionDomain;
+	}
+
 	// Rescale properties to handle source texture size change
 	void RescaleSpriteData(UTexture2D* Texture);
 	bool NeedRescaleSpriteData();
@@ -220,13 +216,6 @@ public:
 	static void ExtractRectsFromTexture(UTexture2D* Texture, TArray<FIntRect>& OutRects);
 	void BuildGeometryFromContours(FSpriteGeometryCollection& GeomOwner);
 
-	void ConvertGeometryToCollisionData();
-	void AddPolygonCollisionShapesToBodySetup();
-	void AddBoxCollisionShapesToBodySetup();
-	void AddCircleCollisionShapesToBodySetup();
-
-
-
 	void CreatePolygonFromBoundingBox(FSpriteGeometryCollection& GeomOwner, bool bUseTightBounds);
 
 	// Reinitializes this sprite (NOTE: Does not register existing components in the world)
@@ -322,5 +311,6 @@ public:
 	friend class FSpriteSelectedEdge;
 	friend class FSpriteSelectedSourceRegion;
 	friend struct FPaperAtlasGenerator;
+	friend class FSingleTileEditorViewportClient;
 };
 
diff --git a/Engine/Plugins/2D/Paper2D/Source/Paper2D/Classes/SpriteEditorOnlyTypes.h b/Engine/Plugins/2D/Paper2D/Source/Paper2D/Classes/SpriteEditorOnlyTypes.h
--- a/Engine/Plugins/2D/Paper2D/Source/Paper2D/Classes/SpriteEditorOnlyTypes.h
+++ b/Engine/Plugins/2D/Paper2D/Source/Paper2D/Classes/SpriteEditorOnlyTypes.h
@@ -4,6 +4,23 @@
 
 #include "SpriteEditorOnlyTypes.generated.h"
 
+// The kind of collision that a Paper2D asset or component might participate in
+UENUM()
+namespace ESpriteCollisionMode
+{
+	enum Type
+	{
+		// Should this have no collison and not participate in physics?
+		None,
+
+		// EXPERIMENTAL: Should this have 2D collision geometry and participate in the 2D physics world?
+		Use2DPhysics UMETA(DisplayName = "Use 2D Physics"),
+
+		// Should this have 3D collision geometry and participate in the 3D physics world?
+		Use3DPhysics UMETA(DisplayName = "Use 3D Physics")
+	};
+}
+
 // The type of a shape in a sprite geometry structure
 UENUM()
 enum class ESpriteShapeType : uint8
@@ -202,7 +219,31 @@ public:
 	// Takes all polygon shapes and generates a list of triangles from them.
 	// Output will contain a multiple of 3 points, each set is one triangle.
 	// Always ignores circles, but can include or ignore boxes based on bIncludeBoxes.
-	void Triangulate(TArray<FVector2D>& Target, bool bIncludeBoxes);
+	void Triangulate(TArray<FVector2D>& Target, bool bIncludeBoxes) const;
+};
+
+// A helper class for converting one or more FSpriteGeometryCollection objects into a UBodySetup object
+struct FSpriteGeometryCollisionBuilderBase
+{
+protected:
+	class UBodySetup* MyBodySetup;
+	float UnrealUnitsPerPixel;
+	float CollisionThickness;
+	float ZOffsetAmount;
+	ESpriteCollisionMode::Type CollisionDomain;
+
+public:
+	FSpriteGeometryCollisionBuilderBase(class UBodySetup* InBodySetup);
+
+	virtual void ProcessGeometry(const FSpriteGeometryCollection& InGeometry);
+	virtual void Finalize();
+protected:
+	void AddBoxCollisionShapesToBodySetup(const FSpriteGeometryCollection& InGeometry);
+	void AddPolygonCollisionShapesToBodySetup(const FSpriteGeometryCollection& InGeometry);
+	void AddCircleCollisionShapesToBodySetup(const FSpriteGeometryCollection& InGeometry);
+
+	virtual FVector2D ConvertTextureSpaceToPivotSpace(const FVector2D& Input) const;
+	virtual FVector2D ConvertTextureSpaceToPivotSpaceNoTranslation(const FVector2D& Input) const;
 };
 
 USTRUCT()
diff --git a/Engine/Plugins/2D/Paper2D/Source/Paper2D/Private/PaperSprite.cpp b/Engine/Plugins/2D/Paper2D/Source/Paper2D/Private/PaperSprite.cpp
--- a/Engine/Plugins/2D/Paper2D/Source/Paper2D/Private/PaperSprite.cpp
+++ b/Engine/Plugins/2D/Paper2D/Source/Paper2D/Private/PaperSprite.cpp
@@ -689,6 +689,35 @@
 	return false;
 }
 
+class FPaperSpriteToBodySetupBuilder : public FSpriteGeometryCollisionBuilderBase
+{
+public:
+	FPaperSpriteToBodySetupBuilder(UPaperSprite* InSprite, UBodySetup* InBodySetup)
+		: FSpriteGeometryCollisionBuilderBase(InBodySetup)
+		, MySprite(InSprite)
+	{
+		UnrealUnitsPerPixel = InSprite->GetUnrealUnitsPerPixel();
+		CollisionThickness = InSprite->GetCollisionThickness();
+		CollisionDomain = InSprite->GetSpriteCollisionDomain();
+	}
+
+protected:
+	// FSpriteGeometryCollisionBuilderBase interface
+	virtual FVector2D ConvertTextureSpaceToPivotSpace(const FVector2D& Input) const override
+	{
+		return MySprite->ConvertTextureSpaceToPivotSpace(Input);
+	}
+
+	virtual FVector2D ConvertTextureSpaceToPivotSpaceNoTranslation(const FVector2D& Input) const override
+	{
+		return MySprite->IsRotatedInSourceImage() ? FVector2D(Input.Y, Input.X) : Input;
+	}
+	// End of FSpriteGeometryCollisionBuilderBase
+
+protected:
+	UPaperSprite* MySprite;
+};
+
 void UPaperSprite::RebuildCollisionData()
 {
 	UBodySetup* OldBodySetup = BodySetup;
@@ -744,18 +773,13 @@
 			check(false); // unknown mode
 		};
 
-		// Take the geometry and add it to the body setup
+		// Clean up the geometry (converting polygons back to bounding boxes, etc...)
 		CollisionGeometry.ConditionGeometry();
-		AddBoxCollisionShapesToBodySetup();
-		AddPolygonCollisionShapesToBodySetup();
-		AddCircleCollisionShapesToBodySetup();
 
-		// Rebuild the body setup
-		if (BodySetup != nullptr)
-		{
-			BodySetup->InvalidatePhysicsData();
-			BodySetup->CreatePhysicsMeshes();
-		}
+		// Take the geometry and add it to the body setup
+		FPaperSpriteToBodySetupBuilder CollisionBuilder(this, BodySetup);
+		CollisionBuilder.ProcessGeometry(CollisionGeometry);
+		CollisionBuilder.Finalize();
 
 		// Copy across or initialize the only editable property we expose on the body setup
 		if (OldBodySetup != nullptr)
@@ -769,180 +793,6 @@
 	}
 }
 
-void UPaperSprite::AddPolygonCollisionShapesToBodySetup()
-{
-	// Rebuild the runtime geometry for polygons
-	TArray<FVector2D> CollisionData;
-	CollisionGeometry.Triangulate(/*out*/ CollisionData, /*bIncludeBoxes=*/ false);
-
-	// Adjust the collision data to be relative to the pivot and scaled from pixels to uu
-	const float UnitsPerPixel = GetUnrealUnitsPerPixel();
-	for (FVector2D& Point : CollisionData)
-	{
-		Point = ConvertTextureSpaceToPivotSpace(Point) * UnitsPerPixel;
-	}
-
-	//@TODO: Observe if the custom data is really a hand-edited bounding box, and generate box geom instead of convex geom!
-	//@TODO: Use this guy instead: DecomposeMeshToHulls
-	//@TODO: Merge triangles that are convex together!
-
-	// Bake it to the runtime structure
-	switch (SpriteCollisionDomain)
-	{
-	case ESpriteCollisionMode::Use3DPhysics:
-		{
-			checkSlow(BodySetup);
-			UBodySetup* BodySetup3D = BodySetup;
-
-			const FVector HalfThicknessVector = PaperAxisZ * 0.5f * CollisionThickness;
-
-			int32 RunningIndex = 0;
-			for (int32 TriIndex = 0; TriIndex < CollisionData.Num() / 3; ++TriIndex)
-			{
-				FKConvexElem& ConvexTri = *new (BodySetup3D->AggGeom.ConvexElems) FKConvexElem();
-				ConvexTri.VertexData.Empty(6);
-				for (int32 Index = 0; Index < 3; ++Index)
-				{
-					const FVector2D& Pos2D = CollisionData[RunningIndex++];
-					
-					const FVector Pos3D = (PaperAxisX * Pos2D.X) + (PaperAxisY * Pos2D.Y);
-
-					new (ConvexTri.VertexData) FVector(Pos3D - HalfThicknessVector);
-					new (ConvexTri.VertexData) FVector(Pos3D + HalfThicknessVector);
-				}
-				ConvexTri.UpdateElemBox();
-			}
-		}
-		break;
-	case ESpriteCollisionMode::Use2DPhysics:
-		{
-			UBodySetup2D* BodySetup2D = CastChecked<UBodySetup2D>(BodySetup);
-
-			int32 RunningIndex = 0;
-			for (int32 TriIndex = 0; TriIndex < CollisionData.Num() / 3; ++TriIndex)
-			{
-				FConvexElement2D& ConvexTri = *new (BodySetup2D->AggGeom2D.ConvexElements) FConvexElement2D();
-				ConvexTri.VertexData.Empty(3);
-				for (int32 Index = 0; Index < 3; ++Index)
-				{
-					const FVector2D& Pos2D = CollisionData[RunningIndex++];
-					new (ConvexTri.VertexData) FVector2D(Pos2D);
-				}
-			}
-		}
-		break;
-	default:
-		check(false);
-		break;
-	}
-}
-
-void UPaperSprite::AddBoxCollisionShapesToBodySetup()
-{
-	checkSlow(BodySetup);
-
-	// Bake all of the boxes to the body setup
-	for (const FSpriteGeometryShape& Shape : CollisionGeometry.Shapes)
-	{
-		if (Shape.ShapeType == ESpriteShapeType::Box)
-		{
-			// Determine the box size and center in pivot space
-			const FVector2D& BoxSizeInTextureSpace = Shape.BoxSize;
-			const FVector2D CenterInTextureSpace = Shape.BoxPosition;
-			const FVector2D CenterInPivotSpace = ConvertTextureSpaceToPivotSpace(CenterInTextureSpace);
-
-			// Convert from pixels to uu
-			const float UnitsPerPixel = GetUnrealUnitsPerPixel();
-			const FVector2D BoxSizeInPivotSpace = bRotatedInSourceImage ? FVector2D(BoxSizeInTextureSpace.Y, BoxSizeInTextureSpace.X) : BoxSizeInTextureSpace;
-			const FVector2D BoxSize2D = BoxSizeInPivotSpace * UnitsPerPixel;
-			const FVector2D CenterInScaledSpace = CenterInPivotSpace * UnitsPerPixel;
-
-			// Create a new box primitive
-			switch (SpriteCollisionDomain)
-			{
-				case ESpriteCollisionMode::Use3DPhysics:
-					{
-					   const FVector BoxSize3D = (PaperAxisX * BoxSize2D.X) + (PaperAxisY * BoxSize2D.Y) + (PaperAxisZ * CollisionThickness);
-
-						// Create a new box primitive
-						FKBoxElem& Box = *new (BodySetup->AggGeom.BoxElems) FKBoxElem(FMath::Abs(BoxSize3D.X), FMath::Abs(BoxSize3D.Y), FMath::Abs(BoxSize3D.Z));
-						Box.Center = (PaperAxisX * CenterInScaledSpace.X) + (PaperAxisY * CenterInScaledSpace.Y);
-						Box.Orientation = FQuat(FRotator(Shape.Rotation, 0.0f, 0.0f));
-					}
-					break;
-				case ESpriteCollisionMode::Use2DPhysics:
-					{
-						UBodySetup2D* BodySetup2D = CastChecked<UBodySetup2D>(BodySetup);
-
-						// Create a new box primitive
-						FBoxElement2D& Box = *new (BodySetup2D->AggGeom2D.BoxElements) FBoxElement2D();
-						Box.Width = FMath::Abs(BoxSize2D.X);
-						Box.Height = FMath::Abs(BoxSize2D.Y);
-						Box.Center.X = CenterInScaledSpace.X;
-						Box.Center.Y = CenterInScaledSpace.Y;
-						Box.Angle = Shape.Rotation;
-					}
-					break;
-				default:
-					check(false);
-					break;
-			}
-		}
-	}
-}
-
-void UPaperSprite::AddCircleCollisionShapesToBodySetup()
-{
-	checkSlow(BodySetup);
-
-	// Bake all of the boxes to the body setup
-	for (const FSpriteGeometryShape& Shape : CollisionGeometry.Shapes)
-	{
-		if (Shape.ShapeType == ESpriteShapeType::Circle)
-		{
-			// Determine the box size and center in pivot space
-			const FVector2D& CircleSizeInTextureSpace = Shape.BoxSize;
-			const FVector2D& CenterInTextureSpace = Shape.BoxPosition;
-			const FVector2D CenterInPivotSpace = ConvertTextureSpaceToPivotSpace(CenterInTextureSpace);
-
-			// Convert from pixels to uu
-			const float UnitsPerPixel = GetUnrealUnitsPerPixel();
-			const FVector2D CircleSizeInPivotSpace = bRotatedInSourceImage ? FVector2D(CircleSizeInTextureSpace.Y, CircleSizeInTextureSpace.X) : CircleSizeInTextureSpace;
-			const FVector2D CircleSize2D = CircleSizeInPivotSpace * UnitsPerPixel;
-			const FVector2D CenterInScaledSpace = CenterInPivotSpace * UnitsPerPixel;
-
-			//@TODO: Neither Box2D nor PhysX support ellipses, currently forcing to be circular, but should we instead convert to an n-gon?
-			const float AverageDiameter = (CircleSize2D.X + CircleSize2D.Y) * 0.5f;
-			const float AverageRadius = FMath::Abs(AverageDiameter * 0.5f);
-
-			// Create a new circle/sphere primitive
-			switch (SpriteCollisionDomain)
-			{
-				case ESpriteCollisionMode::Use3DPhysics:
-					{
-						// Create a new box primitive
-						FKSphereElem& Sphere = *new (BodySetup->AggGeom.SphereElems) FKSphereElem(AverageRadius);
-						Sphere.Center = (PaperAxisX * CenterInScaledSpace.X) + (PaperAxisY * CenterInScaledSpace.Y);
-					}
-					break;
-				case ESpriteCollisionMode::Use2DPhysics:
-					{
-						// Create a new box primitive
-						UBodySetup2D* BodySetup2D = CastChecked<UBodySetup2D>(BodySetup);
-						FCircleElement2D& Circle = *new (BodySetup2D->AggGeom2D.CircleElements) FCircleElement2D();
-						Circle.Radius = AverageRadius;
-						Circle.Center.X = CenterInScaledSpace.X;
-						Circle.Center.Y = CenterInScaledSpace.Y;
-					}
-					break;
-				default:
-					check(false);
-					break;
-			}
-		}
-	}
-}
-
 void UPaperSprite::RebuildRenderData()
 {
 	FSpriteGeometryCollection AlternateGeometry;
@@ -1500,7 +1350,7 @@
 
 		bool bUseMaskedTexture = true;
 
-		// Analyze the texture if desired (to see if it's got greyscale alpha or just binary alpha, picking either a ranslucent or masked material)
+		// Analyze the texture if desired (to see if it's got greyscale alpha or just binary alpha, picking either a translucent or masked material)
 		if (DefaultSettings->bPickBestMaterialWhenCreatingSprite)
 		{
 			if (InitParams.Texture != nullptr)
@@ -2027,7 +1877,7 @@
 	GeometryType = ESpritePolygonMode::TightBoundingBox;
 }
 
-void FSpriteGeometryCollection::Triangulate(TArray<FVector2D>& Target, bool bIncludeBoxes)
+void FSpriteGeometryCollection::Triangulate(TArray<FVector2D>& Target, bool bIncludeBoxes) const
 {
 	Target.Empty();
 
@@ -2143,3 +1993,206 @@
 		}
 	}
 }
+
+//////////////////////////////////////////////////////////////////////////
+// FSpriteGeometryCollisionBuilderBase
+
+FSpriteGeometryCollisionBuilderBase::FSpriteGeometryCollisionBuilderBase(UBodySetup* InBodySetup)
+	: MyBodySetup(InBodySetup)
+	, UnrealUnitsPerPixel(1.0f)
+	, CollisionThickness(64.0f)
+	, ZOffsetAmount(0.0f)
+	, CollisionDomain(ESpriteCollisionMode::Use3DPhysics)
+{
+	check(MyBodySetup);
+}
+
+void FSpriteGeometryCollisionBuilderBase::ProcessGeometry(const FSpriteGeometryCollection& InGeometry)
+{
+	// Add geometry to the body setup
+	AddBoxCollisionShapesToBodySetup(InGeometry);
+	AddPolygonCollisionShapesToBodySetup(InGeometry);
+	AddCircleCollisionShapesToBodySetup(InGeometry);
+}
+
+void FSpriteGeometryCollisionBuilderBase::Finalize()
+{
+	// Rebuild the body setup
+	MyBodySetup->InvalidatePhysicsData();
+	MyBodySetup->CreatePhysicsMeshes();
+}
+
+void FSpriteGeometryCollisionBuilderBase::AddBoxCollisionShapesToBodySetup(const FSpriteGeometryCollection& InGeometry)
+{
+	// Bake all of the boxes to the body setup
+	for (const FSpriteGeometryShape& Shape : InGeometry.Shapes)
+	{
+		if (Shape.ShapeType == ESpriteShapeType::Box)
+		{
+			// Determine the box size and center in pivot space
+			const FVector2D& BoxSizeInTextureSpace = Shape.BoxSize;
+			const FVector2D CenterInTextureSpace = Shape.BoxPosition;
+			const FVector2D CenterInPivotSpace = ConvertTextureSpaceToPivotSpace(CenterInTextureSpace);
+
+			// Convert from pixels to uu
+			const FVector2D BoxSizeInPivotSpace = ConvertTextureSpaceToPivotSpaceNoTranslation(BoxSizeInTextureSpace);
+			const FVector2D BoxSize2D = BoxSizeInPivotSpace * UnrealUnitsPerPixel;
+			const FVector2D CenterInScaledSpace = CenterInPivotSpace * UnrealUnitsPerPixel;
+
+			// Create a new box primitive
+			switch (CollisionDomain)
+			{
+				case ESpriteCollisionMode::Use3DPhysics:
+					{
+					   const FVector BoxSize3D = (PaperAxisX * BoxSize2D.X) + (PaperAxisY * BoxSize2D.Y) + (PaperAxisZ * CollisionThickness);
+
+						// Create a new box primitive
+						FKBoxElem& Box = *new (MyBodySetup->AggGeom.BoxElems) FKBoxElem(FMath::Abs(BoxSize3D.X), FMath::Abs(BoxSize3D.Y), FMath::Abs(BoxSize3D.Z));
+						Box.Center = (PaperAxisX * CenterInScaledSpace.X) + (PaperAxisY * CenterInScaledSpace.Y) + (PaperAxisZ * ZOffsetAmount);
+						Box.Orientation = FQuat(FRotator(Shape.Rotation, 0.0f, 0.0f));
+					}
+					break;
+				case ESpriteCollisionMode::Use2DPhysics:
+					{
+						UBodySetup2D* BodySetup2D = CastChecked<UBodySetup2D>(MyBodySetup);
+
+						// Create a new box primitive
+						FBoxElement2D& Box = *new (BodySetup2D->AggGeom2D.BoxElements) FBoxElement2D();
+						Box.Width = FMath::Abs(BoxSize2D.X);
+						Box.Height = FMath::Abs(BoxSize2D.Y);
+						Box.Center.X = CenterInScaledSpace.X;
+						Box.Center.Y = CenterInScaledSpace.Y;
+						Box.Angle = Shape.Rotation;
+					}
+					break;
+				default:
+					check(false);
+					break;
+			}
+		}
+	}
+}
+
+void FSpriteGeometryCollisionBuilderBase::AddPolygonCollisionShapesToBodySetup(const FSpriteGeometryCollection& InGeometry)
+{
+	// Rebuild the runtime geometry for polygons
+	TArray<FVector2D> CollisionData;
+	InGeometry.Triangulate(/*out*/ CollisionData, /*bIncludeBoxes=*/ false);
+
+	// Adjust the collision data to be relative to the pivot and scaled from pixels to uu
+	for (FVector2D& Point : CollisionData)
+	{
+		Point = ConvertTextureSpaceToPivotSpace(Point) * UnrealUnitsPerPixel;
+	}
+
+	//@TODO: Use this guy instead: DecomposeMeshToHulls
+	//@TODO: Merge triangles that are convex together!
+
+	// Bake it to the runtime structure
+	switch (CollisionDomain)
+	{
+	case ESpriteCollisionMode::Use3DPhysics:
+		{
+			UBodySetup* BodySetup3D = MyBodySetup;
+
+			const FVector HalfThicknessVector = PaperAxisZ * 0.5f * CollisionThickness;
+
+			int32 RunningIndex = 0;
+			for (int32 TriIndex = 0; TriIndex < CollisionData.Num() / 3; ++TriIndex)
+			{
+				FKConvexElem& ConvexTri = *new (BodySetup3D->AggGeom.ConvexElems) FKConvexElem();
+				ConvexTri.VertexData.Empty(6);
+				for (int32 Index = 0; Index < 3; ++Index)
+				{
+					const FVector2D& Pos2D = CollisionData[RunningIndex++];
+					
+					const FVector Pos3D = (PaperAxisX * Pos2D.X) + (PaperAxisY * Pos2D.Y) + (PaperAxisZ * ZOffsetAmount);
+
+					new (ConvexTri.VertexData) FVector(Pos3D - HalfThicknessVector);
+					new (ConvexTri.VertexData) FVector(Pos3D + HalfThicknessVector);
+				}
+				ConvexTri.UpdateElemBox();
+			}
+		}
+		break;
+	case ESpriteCollisionMode::Use2DPhysics:
+		{
+			UBodySetup2D* BodySetup2D = CastChecked<UBodySetup2D>(MyBodySetup);
+
+			int32 RunningIndex = 0;
+			for (int32 TriIndex = 0; TriIndex < CollisionData.Num() / 3; ++TriIndex)
+			{
+				FConvexElement2D& ConvexTri = *new (BodySetup2D->AggGeom2D.ConvexElements) FConvexElement2D();
+				ConvexTri.VertexData.Empty(3);
+				for (int32 Index = 0; Index < 3; ++Index)
+				{
+					const FVector2D& Pos2D = CollisionData[RunningIndex++];
+					new (ConvexTri.VertexData) FVector2D(Pos2D);
+				}
+			}
+		}
+		break;
+	default:
+		check(false);
+		break;
+	}
+}
+
+void FSpriteGeometryCollisionBuilderBase::AddCircleCollisionShapesToBodySetup(const FSpriteGeometryCollection& InGeometry)
+{
+	// Bake all of the boxes to the body setup
+	for (const FSpriteGeometryShape& Shape : InGeometry.Shapes)
+	{
+		if (Shape.ShapeType == ESpriteShapeType::Circle)
+		{
+			// Determine the box size and center in pivot space
+			const FVector2D& CircleSizeInTextureSpace = Shape.BoxSize;
+			const FVector2D& CenterInTextureSpace = Shape.BoxPosition;
+			const FVector2D CenterInPivotSpace = ConvertTextureSpaceToPivotSpace(CenterInTextureSpace);
+
+			// Convert from pixels to uu
+			const FVector2D CircleSizeInPivotSpace = ConvertTextureSpaceToPivotSpaceNoTranslation(CircleSizeInTextureSpace);
+			const FVector2D CircleSize2D = CircleSizeInPivotSpace * UnrealUnitsPerPixel;
+			const FVector2D CenterInScaledSpace = CenterInPivotSpace * UnrealUnitsPerPixel;
+
+			//@TODO: Neither Box2D nor PhysX support ellipses, currently forcing to be circular, but should we instead convert to an n-gon?
+			const float AverageDiameter = (FMath::Abs(CircleSize2D.X) + FMath::Abs(CircleSize2D.Y)) * 0.5f;
+			const float AverageRadius = AverageDiameter * 0.5f;
+
+			// Create a new circle/sphere primitive
+			switch (CollisionDomain)
+			{
+				case ESpriteCollisionMode::Use3DPhysics:
+					{
+						// Create a new box primitive
+						FKSphereElem& Sphere = *new (MyBodySetup->AggGeom.SphereElems) FKSphereElem(AverageRadius);
+						Sphere.Center = (PaperAxisX * CenterInScaledSpace.X) + (PaperAxisY * CenterInScaledSpace.Y) + (PaperAxisZ * ZOffsetAmount);
+					}
+					break;
+				case ESpriteCollisionMode::Use2DPhysics:
+					{
+						// Create a new box primitive
+						UBodySetup2D* BodySetup2D = CastChecked<UBodySetup2D>(MyBodySetup);
+						FCircleElement2D& Circle = *new (BodySetup2D->AggGeom2D.CircleElements) FCircleElement2D();
+						Circle.Radius = AverageRadius;
+						Circle.Center.X = CenterInScaledSpace.X;
+						Circle.Center.Y = CenterInScaledSpace.Y;
+					}
+					break;
+				default:
+					check(false);
+					break;
+			}
+		}
+	}
+}
+
+FVector2D FSpriteGeometryCollisionBuilderBase::ConvertTextureSpaceToPivotSpace(const FVector2D& Input) const
+{
+	return Input;
+}
+
+FVector2D FSpriteGeometryCollisionBuilderBase::ConvertTextureSpaceToPivotSpaceNoTranslation(const FVector2D& Input) const
+{
+	return Input;
+}
