Tile Maps: Add collision generation using data from the per-tile metadata
- Known issue: Tile Rotation/flipping is ignored for the generated collision
- Known issue: Z-depth of the layer is ignored, all collision is centered on the first layer plane

diff --git a/Engine/Plugins/2D/Paper2D/Source/Paper2D/Private/PaperTileLayer.cpp b/Engine/Plugins/2D/Paper2D/Source/Paper2D/Private/PaperTileLayer.cpp
--- a/Engine/Plugins/2D/Paper2D/Source/Paper2D/Private/PaperTileLayer.cpp
+++ b/Engine/Plugins/2D/Paper2D/Source/Paper2D/Private/PaperTileLayer.cpp
@@ -1,9 +1,88 @@
 // Copyright 1998-2015 Epic Games, Inc. All Rights Reserved.
 
 #include "Paper2DPrivatePCH.h"
-#include "PaperSprite.h"
+#include "SpriteEditorOnlyTypes.h"
 #include "ComponentReregisterContext.h"
 
+//////////////////////////////////////////////////////////////////////////
+// FTileMapLayerReregisterContext
+
+/** Removes all components that use the specified tile map layer from their scenes for the lifetime of the class. */
+class FTileMapLayerReregisterContext
+{
+public:
+	/** Initialization constructor. */
+	FTileMapLayerReregisterContext(UPaperTileLayer* TargetAsset)
+	{
+		// Look at tile map components
+		for (TObjectIterator<UPaperTileMapComponent> MapIt; MapIt; ++MapIt)
+		{
+			if (UPaperTileMap* TestMap = (*MapIt)->TileMap)
+			{
+				if (TestMap->TileLayers.Contains(TargetAsset))
+				{
+					AddComponentToRefresh(*MapIt);
+				}
+			}
+		}
+	}
+
+protected:
+	void AddComponentToRefresh(UActorComponent* Component)
+	{
+		if (ComponentContexts.Num() == 0)
+		{
+			// wait until resources are released
+			FlushRenderingCommands();
+		}
+
+		new (ComponentContexts) FComponentReregisterContext(Component);
+	}
+
+private:
+	/** The recreate contexts for the individual components. */
+	TIndirectArray<FComponentReregisterContext> ComponentContexts;
+};
+
+//////////////////////////////////////////////////////////////////////////
+// FPaperTileLayerToBodySetupBuilder
+
+class FPaperTileLayerToBodySetupBuilder : public FSpriteGeometryCollisionBuilderBase
+{
+public:
+	FPaperTileLayerToBodySetupBuilder(UPaperTileMap* InTileMap, UBodySetup* InBodySetup, float InZOffset)
+		: FSpriteGeometryCollisionBuilderBase(InBodySetup)
+	{
+		UnrealUnitsPerPixel = InTileMap->GetUnrealUnitsPerPixel();
+		CollisionThickness = InTileMap->GetCollisionThickness();
+		CollisionDomain = InTileMap->GetSpriteCollisionDomain();
+		CurrentCellOffset = FVector2D::ZeroVector;
+		ZOffsetAmount = InZOffset;
+	}
+
+	void SetCellOffset(const FVector2D& NewOffset)
+	{
+		CurrentCellOffset = NewOffset;
+	}
+
+protected:
+	// FSpriteGeometryCollisionBuilderBase interface
+	virtual FVector2D ConvertTextureSpaceToPivotSpace(const FVector2D& Input) const override
+	{
+		return FVector2D(CurrentCellOffset.X + Input.X, CurrentCellOffset.Y - Input.Y);
+	}
+
+	virtual FVector2D ConvertTextureSpaceToPivotSpaceNoTranslation(const FVector2D& Input) const override
+	{
+		return Input;
+	}
+	// End of FSpriteGeometryCollisionBuilderBase
+
+protected:
+	UPaperTileLayer* MySprite;
+	FVector2D CurrentCellOffset;
+};
+
 //////////////////////////////////////////////////////////////////////////
 // UPaperTileLayer
 
@@ -71,45 +150,6 @@ void UPaperTileLayer::ReallocateAndCopyMap()
 
 #if WITH_EDITOR
 
-/** Removes all components that use the specified tile map layer from their scenes for the lifetime of the class. */
-class FTileMapLayerReregisterContext
-{
-public:
-	/** Initialization constructor. */
-	FTileMapLayerReregisterContext(UPaperTileLayer* TargetAsset)
-	{
-		// Look at tile map components
-		for (TObjectIterator<UPaperTileMapComponent> MapIt; MapIt; ++MapIt)
-		{
-			if (UPaperTileMap* TestMap = (*MapIt)->TileMap)
-			{
-				if (TestMap->TileLayers.Contains(TargetAsset))
-				{
-					AddComponentToRefresh(*MapIt);
-				}
-			}
-		}
-	}
-
-protected:
-	void AddComponentToRefresh(UActorComponent* Component)
-	{
-		if (ComponentContexts.Num() == 0)
-		{
-			// wait until resources are released
-			FlushRenderingCommands();
-		}
-
-		new (ComponentContexts) FComponentReregisterContext(Component);
-	}
-
-private:
-	/** The recreate contexts for the individual components. */
-	TIndirectArray<FComponentReregisterContext> ComponentContexts;
-};
-
-
-
 void UPaperTileLayer::PostEditChangeProperty(struct FPropertyChangedEvent& PropertyChangedEvent)
 {
 	FTileMapLayerReregisterContext ReregisterExistingComponents(this);
@@ -164,13 +204,42 @@ void UPaperTileLayer::SetCell(int32 X, int32 Y, const FPaperTileInfo& NewValue)
 
 void UPaperTileLayer::AugmentBodySetup(UBodySetup* ShapeBodySetup)
 {
+	UPaperTileMap* TileMap = GetTileMap();
+	const float TileWidth = TileMap->TileWidth;
+	const float TileHeight = TileMap->TileHeight;
+
+	//@TODO: Determine if we want collision to be attached to the layer or always relative to the zero layer (probably need a per-layer config value / option, as you may even want to inset layer 0's collision so it's flush with the surface (imagine a top-down game))
+	const float ZOffset = 0.0f;
+
+	// Generate collision for all cells that contain a tile with collision metadata
+	FPaperTileLayerToBodySetupBuilder CollisionBuilder(GetTileMap(), ShapeBodySetup, ZOffset);
+
+	for (int32 CellY = 0; CellY < LayerHeight; ++CellY)
+	{
+		for (int32 CellX = 0; CellX < LayerWidth; ++CellX)
+		{
+			const FPaperTileInfo CellInfo = GetCell(CellX, CellY);
+
+			if (CellInfo.IsValid())
+			{
+				if (const FPaperTileMetadata* CellMetadata = CellInfo.TileSet->GetTileMetadata(CellInfo.GetTileIndex()))
+				{
+					//@TODO: Add support for flipped/mirrored/rotated tiles here (and inside FPaperTileLayerToBodySetupBuilder::ConvertTextureSpaceToPivotSpace(NoTranslation))
+					const FVector2D CellOffset(TileWidth * CellX, TileHeight * -CellY);
+					CollisionBuilder.SetCellOffset(CellOffset);
+
+					CollisionBuilder.ProcessGeometry(CellMetadata->CollisionData);
+				}
+			}
+		}
+	}
+
+	//@TODO: Remove this legacy code path; there is no way to author new layers with bCollisionLayer and they're superceeded by per-tile collision data
+	// Create a box element for every non-zero value in the layer
 	if (bCollisionLayer)
 	{
 		//@TODO: Tile pivot issue
-		//@TODO: Layer thickness issue
-		const float TileWidth = GetTileMap()->TileWidth;
-		const float TileHeight = GetTileMap()->TileHeight;
-		const float TileThickness = 64.0f;
+		const float TileThickness = TileMap->GetCollisionThickness();
 
 		//@TODO: When the origin of the component changes, this logic will need to be adjusted as well
 		// The origin is currently the top left
@@ -184,7 +253,8 @@ void UPaperTileLayer::AugmentBodySetup(UBodySetup* ShapeBodySetup)
 			{
 				if (GetCell(XValue, YValue).PackedTileIndex != 0)
 				{
-					FKBoxElem* NewBox = new(ShapeBodySetup->AggGeom.BoxElems) FKBoxElem(TileWidth, TileThickness, TileHeight);
+					//@TODO: BOX2D: Add support for 2D physics on tile maps if this code isn't deleted entirely (see TODO above)
+					FKBoxElem* NewBox = new (ShapeBodySetup->AggGeom.BoxElems) FKBoxElem(TileWidth, TileThickness, TileHeight);
 
 					FVector BoxPosition;
 					BoxPosition.X = XOrigin + XValue * TileWidth;
