Paper2D: Gated call to InvalidatePhysicsData behind WITH_RUNTIME_PHYSICS_COOKING || WITH_EDITOR

diff --git a/Engine/Plugins/2D/Paper2D/Source/Paper2D/Private/PaperSprite.cpp b/Engine/Plugins/2D/Paper2D/Source/Paper2D/Private/PaperSprite.cpp
--- a/Engine/Plugins/2D/Paper2D/Source/Paper2D/Private/PaperSprite.cpp
+++ b/Engine/Plugins/2D/Paper2D/Source/Paper2D/Private/PaperSprite.cpp
@@ -2018,7 +2018,9 @@ void FSpriteGeometryCollisionBuilderBase::ProcessGeometry(const FSpriteGeometryC
 void FSpriteGeometryCollisionBuilderBase::Finalize()
 {
 	// Rebuild the body setup
+#if WITH_RUNTIME_PHYSICS_COOKING || WITH_EDITOR
 	MyBodySetup->InvalidatePhysicsData();
+#endif
 	MyBodySetup->CreatePhysicsMeshes();
 }
 
