Paper2D: Gated call to InvalidatePhysicsData in UPaperTileMap behind WITH_RUNTIME_PHYSICS_COOKING || WITH_EDITOR

diff --git a/Engine/Plugins/2D/Paper2D/Source/Paper2D/Private/PaperTileMap.cpp b/Engine/Plugins/2D/Paper2D/Source/Paper2D/Private/PaperTileMap.cpp
--- a/Engine/Plugins/2D/Paper2D/Source/Paper2D/Private/PaperTileMap.cpp
+++ b/Engine/Plugins/2D/Paper2D/Source/Paper2D/Private/PaperTileMap.cpp
@@ -170,7 +170,9 @@ void UPaperTileMap::UpdateBodySetup()
 		}
 
 		// Finalize the BodySetup
+#if WITH_RUNTIME_PHYSICS_COOKING || WITH_EDITOR
 		BodySetup->InvalidatePhysicsData();
+#endif
 		BodySetup->CreatePhysicsMeshes();
 	}
 }
