Paper2D: Geometry refactor phase 3 - Moved all geometry editing to an editor mode that can be used elsewhere

diff --git a/Engine/Plugins/2D/Paper2D/Source/Paper2DEditor/Private/Paper2DEditorModule.cpp b/Engine/Plugins/2D/Paper2D/Source/Paper2DEditor/Private/Paper2DEditorModule.cpp
--- a/Engine/Plugins/2D/Paper2D/Source/Paper2DEditor/Private/Paper2DEditorModule.cpp
+++ b/Engine/Plugins/2D/Paper2D/Source/Paper2DEditor/Private/Paper2DEditorModule.cpp
@@ -19,6 +19,8 @@
 #include "SpriteEditor/SpriteComponentDetailsCustomization.h"
 #include "SpriteEditor/SpritePolygonCollectionCustomization.h"
 
+#include "PaperEditorShared/SpriteGeometryEditMode.h"
+
 // Flipbook support
 #include "FlipbookAssetTypeActions.h"
 #include "PaperFlipbookAssetBroker.h"
@@ -147,6 +149,12 @@ public:
 			FSlateIcon(),
 			false);
 
+		FEditorModeRegistry::Get().RegisterMode<FSpriteGeometryEditMode>(
+			FSpriteGeometryEditMode::EM_SpriteGeometry,
+			LOCTEXT("SpriteGeometryEditMode", "Sprite Geometry Editor"),
+			FSlateIcon(),
+			false);
+
 		// Integrate Paper2D actions associated with existing engine types (e.g., Texture2D) into the content browser
 		FPaperContentBrowserExtensions::InstallHooks();
 
@@ -187,6 +195,7 @@ public:
 			FComponentAssetBrokerage::UnregisterBroker(PaperSpriteBroker);
 
 			// Unregister the editor modes
+			FEditorModeRegistry::Get().UnregisterMode(FSpriteGeometryEditMode::EM_SpriteGeometry);
 			FEditorModeRegistry::Get().UnregisterMode(FEdModeTileMap::EM_TileMap);
 
 			// Unregister the thumbnail renderers
diff --git a/Engine/Plugins/2D/Paper2D/Source/Paper2DEditor/Private/PaperEditorShared/SpriteGeometryEditMode.cpp b/Engine/Plugins/2D/Paper2D/Source/Paper2DEditor/Private/PaperEditorShared/SpriteGeometryEditMode.cpp
new file mode 100644
--- /dev/null
+++ b/Engine/Plugins/2D/Paper2D/Source/Paper2DEditor/Private/PaperEditorShared/SpriteGeometryEditMode.cpp
@@ -0,0 +1,554 @@
+// Copyright 1998-2015 Epic Games, Inc. All Rights Reserved.
+
+#include "Paper2DEditorPrivatePCH.h"
+#include "SpriteGeometryEditMode.h"
+#include "../SpriteEditor/SpriteEditorCommands.h"
+
+#define LOCTEXT_NAMESPACE "PaperGeometryEditing"
+
+//////////////////////////////////////////////////////////////////////////
+// FSpriteGeometryEditMode
+
+const FEditorModeID FSpriteGeometryEditMode::EM_SpriteGeometry(TEXT("SpriteGeometryEditMode"));
+
+const FLinearColor FSpriteGeometryEditMode::MarqueeDrawColor(1.0f, 1.0f, 1.0f, 0.5f);
+
+FSpriteGeometryEditMode::FSpriteGeometryEditMode()
+	: BoundsForNewShapes(ForceInitToZero)
+	, GeometryVertexColor(FLinearColor::White)
+	, NegativeGeometryVertexColor(FLinearColor::White)
+	, SpriteGeometryHelper(nullptr)
+	, bIsMarqueeTracking(false)
+	, MarqueeStartPos(ForceInitToZero)
+	, MarqueeEndPos(ForceInitToZero)
+{
+	BoundsForNewShapes.Max = FVector2D(20.0f, 20.0f);
+
+	bDrawPivot = false;
+	bDrawGrid = false;
+}
+
+void FSpriteGeometryEditMode::Initialize()
+{
+}
+
+void FSpriteGeometryEditMode::DrawHUD(FEditorViewportClient* ViewportClient, FViewport* Viewport, const FSceneView* View, FCanvas* Canvas)
+{
+	FEdMode::DrawHUD(ViewportClient, Viewport, View, Canvas);
+
+	int32 YPos = 42;
+	SpriteGeometryHelper.DrawGeometry_CanvasPass(*Viewport, *View, *Canvas, /*inout*/ YPos, GeometryVertexColor, NegativeGeometryVertexColor);
+
+	if (bIsMarqueeTracking)
+	{
+		DrawMarquee(*Viewport, *View, *Canvas, MarqueeDrawColor);
+	}
+}
+
+void FSpriteGeometryEditMode::Render(const FSceneView* View, FViewport* Viewport, FPrimitiveDrawInterface* PDI)
+{
+	FEdMode::Render(View, Viewport, PDI);
+
+	SpriteGeometryHelper.DrawGeometry(*View, *PDI, GeometryVertexColor, NegativeGeometryVertexColor);
+}
+
+bool FSpriteGeometryEditMode::HandleClick(FEditorViewportClient* InViewportClient, HHitProxy* HitProxy, const FViewportClick& Click)
+{
+	FViewport* Viewport = InViewportClient->Viewport;
+
+	const bool bIsCtrlKeyDown = Viewport->KeyState(EKeys::LeftControl) || Viewport->KeyState(EKeys::RightControl);
+	const bool bIsShiftKeyDown = Viewport->KeyState(EKeys::LeftShift) || Viewport->KeyState(EKeys::RightShift);
+	const bool bIsAltKeyDown = Viewport->KeyState(EKeys::LeftAlt) || Viewport->KeyState(EKeys::RightAlt);
+	bool bHandled = false;
+
+	const bool bAllowSelectVertex = !(IsEditingGeometry() && SpriteGeometryHelper.IsAddingPolygon()) && !bIsShiftKeyDown;
+
+	const bool bClearSelectionModifier = bIsCtrlKeyDown;
+	const bool bDeleteClickedVertex = bIsAltKeyDown;
+	const bool bInsertVertexModifier = bIsShiftKeyDown;
+	HSpriteSelectableObjectHitProxy* SelectedItemProxy = HitProxyCast<HSpriteSelectableObjectHitProxy>(HitProxy);
+
+	if (bAllowSelectVertex && (SelectedItemProxy != nullptr))
+	{
+		if (!bClearSelectionModifier)
+		{
+			SpriteGeometryHelper.ClearSelectionSet();
+		}
+
+		if (bDeleteClickedVertex)
+		{
+			// Delete selection
+			if (const FSpriteSelectedVertex* SelectedVertex = SelectedItemProxy->Data->CastTo<const FSpriteSelectedVertex>(FSelectionTypes::Vertex))
+			{
+				SpriteGeometryHelper.ClearSelectionSet();
+				SpriteGeometryHelper.AddPolygonVertexToSelection(SelectedVertex->ShapeIndex, SelectedVertex->VertexIndex);
+				SpriteGeometryHelper.DeleteSelectedItems();
+			}
+			else if (const FSpriteSelectedShape* SelectedShape = SelectedItemProxy->Data->CastTo<const FSpriteSelectedShape>(FSelectionTypes::GeometryShape))
+			{
+				SpriteGeometryHelper.ClearSelectionSet();
+				SpriteGeometryHelper.AddShapeToSelection(SelectedShape->ShapeIndex);
+				SpriteGeometryHelper.DeleteSelectedItems();
+			}
+		}
+		else if (Click.GetEvent() == EInputEvent::IE_DoubleClick)
+		{
+			// Double click to select a polygon
+			if (const FSpriteSelectedVertex* SelectedVertex = SelectedItemProxy->Data->CastTo<const FSpriteSelectedVertex>(FSelectionTypes::Vertex))
+			{
+				SpriteGeometryHelper.ClearSelectionSet();
+				SpriteGeometryHelper.AddShapeToSelection(SelectedVertex->ShapeIndex);
+			}
+		}
+		else
+		{
+			//@TODO: This needs to be generalized!
+			if (const FSpriteSelectedEdge* SelectedEdge = SelectedItemProxy->Data->CastTo<const FSpriteSelectedEdge>(FSelectionTypes::Edge))
+			{
+				// Add the next vertex defined by this edge
+				SpriteGeometryHelper.AddPolygonEdgeToSelection(SelectedEdge->ShapeIndex, SelectedEdge->VertexIndex);
+			}
+			else if (const FSpriteSelectedVertex* SelectedVertex = SelectedItemProxy->Data->CastTo<const FSpriteSelectedVertex>(FSelectionTypes::Vertex))
+			{
+				SpriteGeometryHelper.AddPolygonVertexToSelection(SelectedVertex->ShapeIndex, SelectedVertex->VertexIndex);
+			}
+			else if (const FSpriteSelectedShape* SelectedShape = SelectedItemProxy->Data->CastTo<const FSpriteSelectedShape>(FSelectionTypes::GeometryShape))
+			{
+				SpriteGeometryHelper.AddShapeToSelection(SelectedShape->ShapeIndex);
+			}
+			else
+			{
+				SpriteGeometryHelper.SelectItem(SelectedItemProxy->Data);
+			}
+		}
+
+		bHandled = true;
+	}
+	// 	else if (HWidgetUtilProxy* PivotProxy = HitProxyCast<HWidgetUtilProxy>(HitProxy))
+	// 	{
+	// 		//const bool bUserWantsPaint = bIsLeftButtonDown && ( !GetDefault<ULevelEditorViewportSettings>()->bLeftMouseDragMovesCamera ||  bIsCtrlDown );
+	// 		//findme
+	// 		WidgetAxis = WidgetProxy->Axis;
+	// 
+	// 			// Calculate the screen-space directions for this drag.
+	// 			FSceneViewFamilyContext ViewFamily( FSceneViewFamily::ConstructionValues( Viewport, GetScene(), EngineShowFlags ));
+	// 			FSceneView* View = CalcSceneView(&ViewFamily);
+	// 			WidgetProxy->CalcVectors(View, FViewportClick(View, this, Key, Event, HitX, HitY), LocalManipulateDir, WorldManipulateDir, DragX, DragY);
+	// 			bHandled = true;
+	// 	}
+	else
+	{
+		if (IsEditingGeometry() && !SpriteGeometryHelper.IsAddingPolygon())
+		{
+			FSpriteGeometryCollection* Geometry = SpriteGeometryHelper.GetGeometryBeingEdited();
+
+			if (bInsertVertexModifier)
+			{
+				const FPlane SpritePlane(PaperAxisX, FVector::ZeroVector, PaperAxisY);
+				const FVector WorldPoint = FMath::LinePlaneIntersection(Click.GetOrigin(), Click.GetOrigin() + Click.GetDirection(), SpritePlane);
+				const FVector2D SpriteSpaceClickPoint = SpriteGeometryHelper.GetEditorContext()->WorldSpaceToTextureSpace(WorldPoint);
+
+				// find a polygon to add vert to
+				bool bFoundShapeToAddTo = false;
+				for (TSharedPtr<FSelectedItem> SelectedItemPtr : SpriteGeometryHelper.GetSelectionSet())
+				{
+					if (const FSpriteSelectedVertex* SelectedVertex = SelectedItemPtr->CastTo<const FSpriteSelectedVertex>(FSelectionTypes::Vertex)) //@TODO:  Inflexible?
+					{
+						SpriteGeometryHelper.AddPointToGeometry(SpriteSpaceClickPoint, SelectedVertex->ShapeIndex);
+						bFoundShapeToAddTo = true;
+						break;
+					}
+				}
+
+				if (!bFoundShapeToAddTo)
+				{
+					SpriteGeometryHelper.AddPointToGeometry(SpriteSpaceClickPoint);
+				}
+				bHandled = true;
+			}
+		}
+		else if (!IsEditingGeometry())
+		{
+			// Clicked on the background (missed any proxies), deselect the socket or whatever was selected
+			SpriteGeometryHelper.ClearSelectionSet();
+		}
+	}
+
+	return bHandled ? true : FEdMode::HandleClick(InViewportClient, HitProxy, Click);
+}
+
+bool FSpriteGeometryEditMode::InputKey(FEditorViewportClient* ViewportClient, FViewport* Viewport, FKey Key, EInputEvent Event)
+{
+	bool bHandled = false;
+	FInputEventState InputState(Viewport, Key, Event);
+
+	// Handle marquee tracking in source region edit mode
+	if (IsEditingGeometry())
+	{
+		if (SpriteGeometryHelper.IsAddingPolygon())
+		{
+			if (Key == EKeys::LeftMouseButton && Event == IE_Pressed)
+			{
+				const int32 HitX = Viewport->GetMouseX();
+				const int32 HitY = Viewport->GetMouseY();
+
+				// Calculate world space positions
+				FSceneViewFamilyContext ViewFamily(FSceneViewFamily::ConstructionValues(Viewport, ViewportClient->GetScene(), ViewportClient->EngineShowFlags));
+				FSceneView* View = ViewportClient->CalcSceneView(&ViewFamily);
+				FVector WorldPoint = View->PixelToWorld(HitX, HitY, 0);
+
+				const FVector2D TexturePoint = SpriteGeometryHelper.GetEditorContext()->WorldSpaceToTextureSpace(WorldPoint);
+
+				const bool bMakeSubtractiveIfAllowed = Viewport->KeyState(EKeys::LeftControl) || Viewport->KeyState(EKeys::RightControl);
+				SpriteGeometryHelper.TEMP_HandleAddPolygonClick(TexturePoint, bMakeSubtractiveIfAllowed);
+			}
+			else if (Key == EKeys::Enter)
+			{
+				SpriteGeometryHelper.ResetAddPolygonMode();
+			}
+			else if (Key == EKeys::Escape)
+			{
+				SpriteGeometryHelper.AbandonAddPolygonMode();
+			}
+		}
+		else
+		{
+ 			if (ProcessMarquee(Viewport, Key, Event, true))
+ 			{
+ 				const bool bAddingToSelection = InputState.IsShiftButtonPressed(); //@TODO: control button moves widget? Hopefully make this more consistent when that is changed
+ 				SelectVerticesInMarquee(ViewportClient, Viewport, bAddingToSelection);
+ 			}
+		}
+	}
+
+	//@TODO: Support select-and-drag in a single operation (may involve InputAxis and StartTracking)
+
+	// Pass keys to standard controls, if we didn't consume input
+	return bHandled ? true : FEdMode::InputKey(ViewportClient, Viewport, Key, Event);
+}
+
+void FSpriteGeometryEditMode::Tick(FEditorViewportClient* ViewportClient, float DeltaTime)
+{
+	if (bIsMarqueeTracking)
+	{
+		const int32 HitX = ViewportClient->Viewport->GetMouseX();
+		const int32 HitY = ViewportClient->Viewport->GetMouseY();
+		MarqueeEndPos = FVector2D(HitX, HitY);
+	}
+
+	FEdMode::Tick(ViewportClient, DeltaTime);
+}
+
+bool FSpriteGeometryEditMode::ShouldDrawWidget() const
+{
+	return SpriteGeometryHelper.HasAnySelectedItems();
+}
+
+FVector FSpriteGeometryEditMode::GetWidgetLocation() const
+{
+	FVector SummedPos(ForceInitToZero);
+
+	if (SpriteGeometryHelper.HasAnySelectedItems())
+	{
+		// Find the center of the selection set
+		for (TSharedPtr<FSelectedItem> SelectedItem : SpriteGeometryHelper.GetSelectionSet())
+		{
+			SummedPos += SelectedItem->GetWorldPos();
+		}
+		return (SummedPos / SpriteGeometryHelper.GetSelectionSet().Num());
+	}
+
+	return SummedPos;
+}
+
+bool FSpriteGeometryEditMode::InputDelta(FEditorViewportClient* InViewportClient, FViewport* InViewport, FVector& InDrag, FRotator& InRot, FVector& InScale)
+{
+	bool bHandled = false;
+
+	const bool bManipulating = GetModeManager()->IsTracking();
+	EAxisList::Type CurrentAxis = GetCurrentWidgetAxis();
+	
+	if (bManipulating && (CurrentAxis != EAxisList::None))
+	{
+		bHandled = true;
+
+		const FWidget::EWidgetMode MoveMode = GetModeManager()->GetWidgetMode();
+
+		// Negate Y because vertices are in source texture space, not world space
+		const FVector2D Drag2D(FVector::DotProduct(InDrag, PaperAxisX), -FVector::DotProduct(InDrag, PaperAxisY));
+
+		// Apply the delta to all of the selected objects
+		for (TSharedPtr<FSelectedItem> SelectedItem : SpriteGeometryHelper.GetSelectionSet())
+		{
+			SelectedItem->ApplyDelta(Drag2D, InRot, InScale, MoveMode);
+		}
+
+		if (SpriteGeometryHelper.HasAnySelectedItems())
+		{
+			SpriteGeometryHelper.GetEditorContext()->MarkTransactionAsDirty();
+		}
+	}
+
+	return bHandled ? true : FEdMode::InputDelta(InViewportClient, InViewport, InDrag, InRot, InScale);
+}
+
+void FSpriteGeometryEditMode::SetEditorContext(class ISpriteSelectionContext* InNewEditorContext)
+{
+	SpriteGeometryHelper.SetEditorContext(InNewEditorContext);
+}
+
+void FSpriteGeometryEditMode::SetNewGeometryPreferredBounds(FBox2D& NewDesiredBounds)
+{
+	BoundsForNewShapes = NewDesiredBounds;
+}
+
+void FSpriteGeometryEditMode::SetGeometryColors(const FLinearColor& NewVertexColor, const FLinearColor& NewNegativeVertexColor)
+{
+	GeometryVertexColor = NewVertexColor;
+	NegativeGeometryVertexColor = NewNegativeVertexColor;
+}
+
+void FSpriteGeometryEditMode::SetGeometryBeingEdited(FSpriteGeometryCollection* NewGeometryBeingEdited, bool bInAllowCircles, bool bInAllowSubtractivePolygons)
+{
+	SpriteGeometryHelper.SetGeometryBeingEdited(NewGeometryBeingEdited, bInAllowCircles, bInAllowSubtractivePolygons);
+	bIsMarqueeTracking = false;
+}
+// 
+// FVector2D FSpriteGeometryEditMode::SelectedItemConvertWorldSpaceDeltaToLocalSpace(const FVector& WorldSpaceDelta) const
+// {
+// 	const FVector ProjectionX = WorldSpaceDelta.ProjectOnTo(PaperAxisX);
+// 	const FVector ProjectionY = WorldSpaceDelta.ProjectOnTo(PaperAxisY);
+// 
+// 	const float XValue = FMath::Sign(ProjectionX | PaperAxisX) * ProjectionX.Size();
+// 	const float YValue = FMath::Sign(ProjectionY | PaperAxisY) * ProjectionY.Size();
+// 
+// 	return FVector2D(XValue, YValue);
+// }
+// 
+// FVector2D FSpriteGeometryEditMode::WorldSpaceToTextureSpace(const FVector& SourcePoint) const
+// {
+// 	const FVector ProjectionX = SourcePoint.ProjectOnTo(PaperAxisX);
+// 	const FVector ProjectionY = SourcePoint.ProjectOnTo(PaperAxisY);
+// 
+// 	const float XValue = FMath::Sign(ProjectionX | PaperAxisX) * ProjectionX.Size();
+// 	const float YValue = FMath::Sign(ProjectionY | PaperAxisY) * ProjectionY.Size();
+// 
+// 	return FVector2D(XValue, YValue);
+// }
+// 
+// FVector FSpriteGeometryEditMode::TextureSpaceToWorldSpace(const FVector2D& SourcePoint) const
+// {
+// 	return (SourcePoint.X * PaperAxisX) + (SourcePoint.Y * PaperAxisY);
+// }
+// 
+// float FSpriteGeometryEditMode::SelectedItemGetUnitsPerPixel() const
+// {
+// 	return 1.0f;
+// }
+// 
+// void FSpriteGeometryEditMode::BeginTransaction(const FText& SessionName)
+// {
+// 	// Dummy implementation does nothing
+// }
+// 
+// void FSpriteGeometryEditMode::MarkTransactionAsDirty()
+// {
+// 	// Dummy implementation does nothing
+// }
+// 
+// void FSpriteGeometryEditMode::EndTransaction()
+// {
+// 	// Dummy implementation does nothing
+// }
+// 
+// void FSpriteGeometryEditMode::InvalidateViewportAndHitProxies()
+// {
+// 	// Dummy implementation does nothing
+// }
+
+void FSpriteGeometryEditMode::BindCommands(TSharedPtr<FUICommandList> CommandList)
+{
+	const FSpriteGeometryEditCommands& Commands = FSpriteGeometryEditCommands::Get();
+
+	// Show toggles
+	CommandList->MapAction(
+		Commands.SetShowNormals,
+		FExecuteAction::CreateRaw(&SpriteGeometryHelper, &FSpriteGeometryEditingHelper::ToggleShowNormals),
+		FCanExecuteAction(),
+		FIsActionChecked::CreateRaw(&SpriteGeometryHelper, &FSpriteGeometryEditingHelper::IsShowNormalsEnabled));
+
+	// Geometry editing commands
+	CommandList->MapAction(
+		Commands.DeleteSelection,
+		FExecuteAction::CreateRaw(&SpriteGeometryHelper, &FSpriteGeometryEditingHelper::DeleteSelectedItems, /*bShouldTransact=*/ true),
+		FCanExecuteAction::CreateRaw(&SpriteGeometryHelper, &FSpriteGeometryEditingHelper::CanDeleteSelection));
+	CommandList->MapAction(
+		Commands.AddBoxShape,
+		FExecuteAction::CreateSP(this, &FSpriteGeometryEditMode::AddBoxShape),
+		FCanExecuteAction::CreateRaw(&SpriteGeometryHelper, &FSpriteGeometryEditingHelper::CanAddBoxShape),
+		FIsActionChecked(),
+		FIsActionButtonVisible::CreateRaw(&SpriteGeometryHelper, &FSpriteGeometryEditingHelper::CanAddBoxShape));
+	CommandList->MapAction(
+		Commands.ToggleAddPolygonMode,
+		FExecuteAction::CreateRaw(&SpriteGeometryHelper, &FSpriteGeometryEditingHelper::ToggleAddPolygonMode),
+		FCanExecuteAction::CreateRaw(&SpriteGeometryHelper, &FSpriteGeometryEditingHelper::CanAddPolygon),
+		FIsActionChecked::CreateRaw(&SpriteGeometryHelper, &FSpriteGeometryEditingHelper::IsAddingPolygon),
+		FIsActionButtonVisible::CreateRaw(&SpriteGeometryHelper, &FSpriteGeometryEditingHelper::CanAddPolygon));
+	CommandList->MapAction(
+		Commands.AddCircleShape,
+		FExecuteAction::CreateSP(this, &FSpriteGeometryEditMode::AddCircleShape),
+		FCanExecuteAction::CreateRaw(&SpriteGeometryHelper, &FSpriteGeometryEditingHelper::CanAddCircleShape),
+		FIsActionChecked(),
+		FIsActionButtonVisible::CreateRaw(&SpriteGeometryHelper, &FSpriteGeometryEditingHelper::CanAddCircleShape));
+	CommandList->MapAction(
+		Commands.SnapAllVertices,
+		FExecuteAction::CreateRaw(&SpriteGeometryHelper, &FSpriteGeometryEditingHelper::SnapAllVerticesToPixelGrid),
+		FCanExecuteAction::CreateRaw(&SpriteGeometryHelper, &FSpriteGeometryEditingHelper::CanSnapVerticesToPixelGrid),
+		FIsActionChecked(),
+		FIsActionButtonVisible::CreateRaw(&SpriteGeometryHelper, &FSpriteGeometryEditingHelper::CanSnapVerticesToPixelGrid));
+}
+
+
+void FSpriteGeometryEditMode::AddBoxShape()
+{
+	SpriteGeometryHelper.AddNewBoxShape(BoundsForNewShapes.GetCenter(), BoundsForNewShapes.GetSize());
+}
+
+void FSpriteGeometryEditMode::AddCircleShape()
+{
+	const float SmallerBoundingAxisSize = BoundsForNewShapes.GetSize().GetMin();
+	const float CircleRadius = SmallerBoundingAxisSize * 0.5f;
+
+	SpriteGeometryHelper.AddNewCircleShape(BoundsForNewShapes.GetCenter(), CircleRadius);
+}
+
+bool FSpriteGeometryEditMode::IsEditingGeometry() const
+{
+	return SpriteGeometryHelper.IsEditingGeometry();
+}
+
+void FSpriteGeometryEditMode::SelectVerticesInMarquee(FEditorViewportClient* ViewportClient, FViewport* Viewport, bool bAddToSelection)
+{
+	if (!bAddToSelection)
+	{
+		SpriteGeometryHelper.ClearSelectionSet();
+	}
+
+	{
+		// Calculate world space positions
+		FSceneViewFamilyContext ViewFamily(FSceneViewFamily::ConstructionValues(Viewport, ViewportClient->GetScene(), ViewportClient->EngineShowFlags));
+		FSceneView* View = ViewportClient->CalcSceneView(&ViewFamily);
+		const FVector StartPos = View->PixelToWorld(MarqueeStartPos.X, MarqueeStartPos.Y, 0);
+		const FVector EndPos = View->PixelToWorld(MarqueeEndPos.X, MarqueeEndPos.Y, 0);
+
+		// Convert to source texture space to work out the pixels dragged
+		FVector2D TextureSpaceStartPos = SpriteGeometryHelper.GetEditorContext()->WorldSpaceToTextureSpace(StartPos);
+		FVector2D TextureSpaceEndPos = SpriteGeometryHelper.GetEditorContext()->WorldSpaceToTextureSpace(EndPos);
+
+		if (TextureSpaceStartPos.X > TextureSpaceEndPos.X)
+		{
+			Swap(TextureSpaceStartPos.X, TextureSpaceEndPos.X);
+		}
+		if (TextureSpaceStartPos.Y > TextureSpaceEndPos.Y)
+		{
+			Swap(TextureSpaceStartPos.Y, TextureSpaceEndPos.Y);
+		}
+
+		const FBox2D QueryBounds(TextureSpaceStartPos, TextureSpaceEndPos);
+
+		if (FSpriteGeometryCollection* Geometry = SpriteGeometryHelper.GetGeometryBeingEdited())
+		{
+			for (int32 ShapeIndex = 0; ShapeIndex < Geometry->Shapes.Num(); ++ShapeIndex)
+			{
+				const FSpriteGeometryShape& Shape = Geometry->Shapes[ShapeIndex];
+
+				bool bSelectWholeShape = false;
+
+				if ((Shape.ShapeType == ESpriteShapeType::Circle) || (Shape.ShapeType == ESpriteShapeType::Box))
+				{
+					// First see if we are fully contained
+					const FBox2D ShapeBoxBounds(Shape.BoxPosition - Shape.BoxSize * 0.5f, Shape.BoxPosition + Shape.BoxSize * 0.5f);
+					if (QueryBounds.IsInside(ShapeBoxBounds))
+					{
+						bSelectWholeShape = true;
+					}
+				}
+
+				//@TODO: Try intersecting with the circle if it wasn't entirely enclosed
+
+				if (bSelectWholeShape)
+				{
+					SpriteGeometryHelper.AddShapeToSelection(ShapeIndex);
+				}
+				else
+				{
+					// Try to select some subset of the vertices
+					for (int32 VertexIndex = 0; VertexIndex < Shape.Vertices.Num(); ++VertexIndex)
+					{
+						const FVector2D TextureSpaceVertex = Shape.ConvertShapeSpaceToTextureSpace(Shape.Vertices[VertexIndex]);
+						if (QueryBounds.IsInside(TextureSpaceVertex))
+						{
+							SpriteGeometryHelper.AddPolygonVertexToSelection(ShapeIndex, VertexIndex);
+						}
+					}
+				}
+			}
+		}
+	}
+}
+
+bool FSpriteGeometryEditMode::ProcessMarquee(FViewport* Viewport, FKey Key, EInputEvent Event, bool bMarqueeStartModifierPressed)
+{
+	bool bMarqueeReady = false;
+
+	if (Key == EKeys::LeftMouseButton)
+	{
+		int32 HitX = Viewport->GetMouseX();
+		int32 HitY = Viewport->GetMouseY();
+		if (Event == IE_Pressed && bMarqueeStartModifierPressed)
+		{
+			HHitProxy* HitResult = Viewport->GetHitProxy(HitX, HitY);
+
+			if ((HitResult == nullptr) || (HitResult->Priority == HPP_World))
+			{
+				bIsMarqueeTracking = true;
+				MarqueeStartPos = FVector2D(HitX, HitY);
+				MarqueeEndPos = MarqueeStartPos;
+			}
+		}
+		else if (bIsMarqueeTracking && (Event == IE_Released))
+		{
+			MarqueeEndPos = FVector2D(HitX, HitY);
+			bIsMarqueeTracking = false;
+			bMarqueeReady = true;
+		}
+	}
+	else if (bIsMarqueeTracking && (Key == EKeys::Escape))
+	{
+		// Cancel marquee selection
+		bIsMarqueeTracking = false;
+	}
+
+	return bMarqueeReady;
+}
+
+void FSpriteGeometryEditMode::DrawMarquee(FViewport& InViewport, const FSceneView& View, FCanvas& Canvas, const FLinearColor& Color)
+{
+	FVector2D MarqueeDelta = MarqueeEndPos - MarqueeStartPos;
+	FVector2D MarqueeVertices[4];
+	MarqueeVertices[0] = FVector2D(MarqueeStartPos.X, MarqueeStartPos.Y);
+	MarqueeVertices[1] = MarqueeVertices[0] + FVector2D(MarqueeDelta.X, 0);
+	MarqueeVertices[2] = MarqueeVertices[0] + FVector2D(MarqueeDelta.X, MarqueeDelta.Y);
+	MarqueeVertices[3] = MarqueeVertices[0] + FVector2D(0, MarqueeDelta.Y);
+	for (int32 MarqueeVertexIndex = 0; MarqueeVertexIndex < 4; ++MarqueeVertexIndex)
+	{
+		const int32 NextVertexIndex = (MarqueeVertexIndex + 1) % 4;
+		FCanvasLineItem MarqueeLine(MarqueeVertices[MarqueeVertexIndex], MarqueeVertices[NextVertexIndex]);
+		MarqueeLine.SetColor(Color);
+		Canvas.DrawItem(MarqueeLine);
+	}
+}
+
+//////////////////////////////////////////////////////////////////////////
+
+#undef LOCTEXT_NAMESPACE
diff --git a/Engine/Plugins/2D/Paper2D/Source/Paper2DEditor/Private/PaperEditorShared/SpriteGeometryEditMode.h b/Engine/Plugins/2D/Paper2D/Source/Paper2DEditor/Private/PaperEditorShared/SpriteGeometryEditMode.h
new file mode 100644
--- /dev/null
+++ b/Engine/Plugins/2D/Paper2D/Source/Paper2DEditor/Private/PaperEditorShared/SpriteGeometryEditMode.h
@@ -0,0 +1,83 @@
+// Copyright 1998-2015 Epic Games, Inc. All Rights Reserved.
+
+#pragma once
+
+#include "SpriteGeometryEditing.h"
+#include "../SpriteEditor/SpriteEditorSelections.h"
+
+//////////////////////////////////////////////////////////////////////////
+// FSpriteGeometryEditMode
+
+class FSpriteGeometryEditMode : public FEdMode//, public ISpriteSelectionContext
+{
+public:
+	static const FEditorModeID EM_SpriteGeometry;
+	static const FLinearColor MarqueeDrawColor;
+public:
+
+	FSpriteGeometryEditMode();
+
+	// FEdMode interface
+	virtual void Initialize() override;
+	virtual void DrawHUD(FEditorViewportClient* ViewportClient, FViewport* Viewport, const FSceneView* View, FCanvas* Canvas) override;
+	virtual void Render(const FSceneView* View, FViewport* Viewport, FPrimitiveDrawInterface* PDI) override;
+	virtual bool HandleClick(FEditorViewportClient* InViewportClient, HHitProxy* HitProxy, const FViewportClick& Click) override;
+	virtual bool InputKey(FEditorViewportClient* ViewportClient, FViewport* Viewport, FKey Key, EInputEvent Event) override;
+	virtual FVector GetWidgetLocation() const override;
+	virtual bool InputDelta(FEditorViewportClient* InViewportClient, FViewport* InViewport, FVector& InDrag, FRotator& InRot, FVector& InScale) override;
+	virtual void Tick(FEditorViewportClient* ViewportClient, float DeltaTime) override;
+	virtual bool ShouldDrawWidget() const override;
+	// End of FEdMode interface
+
+	// Changes the editor interface to point to the hosting editor; this is basically required
+	void SetEditorContext(class ISpriteSelectionContext* InNewEditorContext);
+
+	// Sets the default bounds for newly created boxes/circles/etc...
+	void SetNewGeometryPreferredBounds(FBox2D& NewDesiredBounds);
+
+	// Sets the draw color for geometry
+	void SetGeometryColors(const FLinearColor& NewVertexColor, const FLinearColor& NewNegativeVertexColor);
+
+	// Changes the geometry being edited (clears the selection set in the process)
+	void SetGeometryBeingEdited(FSpriteGeometryCollection* NewGeometryBeingEdited, bool bInAllowCircles, bool bInAllowSubtractivePolygons);
+
+// 	// Dummy implementation of ISpriteSelectionContext; users of this mode should provide their own and register via SetEditorContext
+// 	virtual FVector2D SelectedItemConvertWorldSpaceDeltaToLocalSpace(const FVector& WorldSpaceDelta) const override;
+// 	virtual FVector2D WorldSpaceToTextureSpace(const FVector& SourcePoint) const override;
+// 	virtual FVector TextureSpaceToWorldSpace(const FVector2D& SourcePoint) const override;
+// 	virtual float SelectedItemGetUnitsPerPixel() const override;
+// 	virtual void BeginTransaction(const FText& SessionName) override;
+// 	virtual void MarkTransactionAsDirty() override;
+// 	virtual void EndTransaction() override;
+// 	virtual void InvalidateViewportAndHitProxies() override;
+// 	// End of ISpriteSelectionContext interface
+// 
+	void BindCommands(TSharedPtr<FUICommandList> InCommandList);
+
+	FVector2D GetMarqueeStartPos() const { return MarqueeStartPos; }
+	FVector2D GetMarqueeEndPos() const { return MarqueeEndPos; }
+	bool ProcessMarquee(FViewport* Viewport, FKey Key, EInputEvent Event, bool bMarqueeStartModifierPressed);
+
+protected:
+	FBox2D BoundsForNewShapes;
+	FLinearColor GeometryVertexColor;
+	FLinearColor NegativeGeometryVertexColor;
+
+	// Sprite geometry editing/rendering helper
+	FSpriteGeometryEditingHelper SpriteGeometryHelper;
+
+	// Marquee tracking
+	bool bIsMarqueeTracking;
+	FVector2D MarqueeStartPos;
+	FVector2D MarqueeEndPos;
+
+protected:
+	void AddBoxShape();
+	void AddCircleShape();
+
+	bool IsEditingGeometry() const;
+
+	void SelectVerticesInMarquee(FEditorViewportClient* ViewportClient, FViewport* Viewport, bool bAddToSelection);
+
+	void DrawMarquee(FViewport& InViewport, const FSceneView& View, FCanvas& Canvas, const FLinearColor& Color);
+};
diff --git a/Engine/Plugins/2D/Paper2D/Source/Paper2DEditor/Private/PaperEditorShared/SpriteGeometryEditing.cpp b/Engine/Plugins/2D/Paper2D/Source/Paper2DEditor/Private/PaperEditorShared/SpriteGeometryEditing.cpp
--- a/Engine/Plugins/2D/Paper2D/Source/Paper2DEditor/Private/PaperEditorShared/SpriteGeometryEditing.cpp
+++ b/Engine/Plugins/2D/Paper2D/Source/Paper2DEditor/Private/PaperEditorShared/SpriteGeometryEditing.cpp
@@ -41,10 +41,15 @@ void FSpriteSelectionHelper::SelectItem(TSharedPtr<FSelectedItem> NewItem)
 	SelectedItemSet.Add(NewItem);
 }
 
+bool FSpriteSelectionHelper::CanDeleteSelection() const
+{
+	return SelectedItemSet.Num() > 0;
+}
+
 //////////////////////////////////////////////////////////////////////////
 // FSpriteSelectedShape
 
-FSpriteSelectedShape::FSpriteSelectedShape(ISpriteSelectionContext& InEditorContext, FSpriteGeometryCollection& InGeometry, int32 InShapeIndex, bool bInIsBackground)
+FSpriteSelectedShape::FSpriteSelectedShape(ISpriteSelectionContext* InEditorContext, FSpriteGeometryCollection& InGeometry, int32 InShapeIndex, bool bInIsBackground)
 	: FSelectedItem(FSelectionTypes::GeometryShape)
 	, EditorContext(InEditorContext)
 	, Geometry(InGeometry)
@@ -96,7 +101,7 @@ void FSpriteSelectedShape::ApplyDelta(const FVector2D& Delta, const FRotator& Ro
 		if (bDoTranslation)
 		{
 			const FVector WorldSpaceDelta = (PaperAxisX * Delta.X) + (PaperAxisY * Delta.Y);
-			const FVector2D TextureSpaceDelta = EditorContext.SelectedItemConvertWorldSpaceDeltaToLocalSpace(WorldSpaceDelta);
+			const FVector2D TextureSpaceDelta = EditorContext->SelectedItemConvertWorldSpaceDeltaToLocalSpace(WorldSpaceDelta);
 
 			Shape.BoxPosition += TextureSpaceDelta;
 
@@ -162,12 +167,12 @@ FVector FSpriteSelectedShape::GetWorldPos() const
 		{
 		case ESpriteShapeType::Box:
 		case ESpriteShapeType::Circle:
-			Result = EditorContext.TextureSpaceToWorldSpace(Shape.BoxPosition);
+			Result = EditorContext->TextureSpaceToWorldSpace(Shape.BoxPosition);
 			break;
 		case ESpriteShapeType::Polygon:
 			// Average the vertex positions
 			//@TODO: Eventually this will just be BoxPosition as well once the vertex positions are relative
-			Result = EditorContext.TextureSpaceToWorldSpace(Shape.GetPolygonCentroid());
+			Result = EditorContext->TextureSpaceToWorldSpace(Shape.GetPolygonCentroid());
 			break;
 		default:
 			check(false);
@@ -180,7 +185,7 @@ FVector FSpriteSelectedShape::GetWorldPos() const
 //////////////////////////////////////////////////////////////////////////
 // FSpriteGeometryEditingHelper
 
-FSpriteGeometryEditingHelper::FSpriteGeometryEditingHelper(ISpriteSelectionContext& InEditorContext)
+FSpriteGeometryEditingHelper::FSpriteGeometryEditingHelper(ISpriteSelectionContext* InEditorContext)
 	: EditorContext(InEditorContext)
 	, GeometryBeingEdited(nullptr)
 	, bIsAddingPolygon(false)
@@ -199,7 +204,7 @@ void FSpriteGeometryEditingHelper::ClearSelectionSet()
 
 	ResetAddPolygonMode();
 
-	EditorContext.InvalidateViewportAndHitProxies();
+	EditorContext->InvalidateViewportAndHitProxies();
 }
 
 void FSpriteGeometryEditingHelper::DeleteSelectedItems(bool bShouldTransact)
@@ -230,7 +235,7 @@ void FSpriteGeometryEditingHelper::DeleteSelectedItems(bool bShouldTransact)
 
 	if ((CompositeIndicesSet.Num() > 0) || (ShapesToDeleteSet.Num() > 0))
 	{
-		EditorContext.BeginTransaction(LOCTEXT("DeleteSelectionTransaction", "Delete Selection"));
+		EditorContext->BeginTransaction(LOCTEXT("DeleteSelectionTransaction", "Delete Selection"));
 
 		// Delete the selected vertices first, as they may cause entire shapes to need to be deleted (sort so we delete from the back first)
 		TArray<FShapeVertexPair> CompositeIndices = CompositeIndicesSet.Array();
@@ -258,13 +263,13 @@ void FSpriteGeometryEditingHelper::DeleteSelectedItems(bool bShouldTransact)
 		}
 
 		Geometry.GeometryType = ESpritePolygonMode::FullyCustom;
-		EditorContext.MarkTransactionAsDirty();
+		EditorContext->MarkTransactionAsDirty();
 
 		//@TODO: Should allow deleting other things when geometry is also selected!
 		//FSpriteGeometryEditingHelper::DeleteSelectedItems(/*bShouldTransact=*/ false);
 
 		ClearSelectionSet();
-		EditorContext.EndTransaction();
+		EditorContext->EndTransaction();
 		ResetAddPolygonMode();
 	}
 	else
@@ -278,6 +283,16 @@ void FSpriteGeometryEditingHelper::AddReferencedObjects(FReferenceCollector& Col
 	Collector.AddReferencedObject(WidgetVertexColorMaterial);
 }
 
+void FSpriteGeometryEditingHelper::SetEditorContext(class ISpriteSelectionContext* InNewEditorContext)
+{
+	EditorContext = InNewEditorContext;
+}
+
+class ISpriteSelectionContext* FSpriteGeometryEditingHelper::GetEditorContext() const
+{
+	return EditorContext;
+}
+
 void FSpriteGeometryEditingHelper::DrawGeometry(const FSceneView& View, FPrimitiveDrawInterface& PDI, const FLinearColor& GeometryVertexColor, const FLinearColor& NegativeGeometryVertexColor)
 {
 	if (GeometryBeingEdited == nullptr)
@@ -287,14 +302,14 @@ void FSpriteGeometryEditingHelper::DrawGeometry(const FSceneView& View, FPrimiti
 
 	FSpriteGeometryCollection& Geometry = GetGeometryChecked();
 	const bool bIsHitTesting = PDI.IsHitTesting();
-	const float UnitsPerPixel = EditorContext.SelectedItemGetUnitsPerPixel();
+	const float UnitsPerPixel = EditorContext->SelectedItemGetUnitsPerPixel();
 
 	// Run thru the geometry shapes and draw hit proxies for them
 	for (int32 ShapeIndex = 0; ShapeIndex < Geometry.Shapes.Num(); ++ShapeIndex)
 	{
 		const FSpriteGeometryShape& Shape = Geometry.Shapes[ShapeIndex];
 
-		const bool bIsShapeSelected = EditorContext.SelectedItemIsSelected(FShapeVertexPair(ShapeIndex, INDEX_NONE));
+		const bool bIsShapeSelected = IsGeometrySelected(FShapeVertexPair(ShapeIndex, INDEX_NONE));
 		const FLinearColor LineColorRaw = Shape.bNegativeWinding ? NegativeGeometryVertexColor : GeometryVertexColor;
 		const FLinearColor VertexColor = Shape.bNegativeWinding ? NegativeGeometryVertexColor : GeometryVertexColor;
 
@@ -317,7 +332,7 @@ void FSpriteGeometryEditingHelper::DrawGeometry(const FSceneView& View, FPrimiti
 			const FVector2D PixelSpaceRadius = Shape.BoxSize * 0.5f;
 			const float WorldSpaceRadius = PixelSpaceRadius.X * UnitsPerPixel;
 
-			const FVector CircleCenterWorldPos = EditorContext.TextureSpaceToWorldSpace(Shape.BoxPosition);
+			const FVector CircleCenterWorldPos = EditorContext->TextureSpaceToWorldSpace(Shape.BoxPosition);
 
 			DrawDisc(&PDI, CircleCenterWorldPos, PaperAxisX, PaperAxisY, BackgroundColor, WorldSpaceRadius, SpriteEditingConstantsEX::CircleShapeNumSides, ShapeMaterialProxy, SDPG_Foreground);
 		}
@@ -340,7 +355,7 @@ void FSpriteGeometryEditingHelper::DrawGeometry(const FSceneView& View, FPrimiti
 
 				for (const FVector2D& SrcTriangleVertex : TriangulatedPolygonVertices)
 				{
-					MeshVertex.Position = EditorContext.TextureSpaceToWorldSpace(SrcTriangleVertex);
+					MeshVertex.Position = EditorContext->TextureSpaceToWorldSpace(SrcTriangleVertex);
 					MeshBuilder.AddVertex(MeshVertex);
 				}
 
@@ -360,8 +375,13 @@ void FSpriteGeometryEditingHelper::DrawGeometry(const FSceneView& View, FPrimiti
 	}
 }
 
-void FSpriteGeometryEditingHelper::DrawGeometry_CanvasPass(FViewport& InViewport, FSceneView& View, FCanvas& Canvas, /*inout*/ int32& YPos, const FLinearColor& GeometryVertexColor, const FLinearColor& NegativeGeometryVertexColor)
+void FSpriteGeometryEditingHelper::DrawGeometry_CanvasPass(FViewport& InViewport, const FSceneView& View, FCanvas& Canvas, /*inout*/ int32& YPos, const FLinearColor& GeometryVertexColor, const FLinearColor& NegativeGeometryVertexColor)
 {
+	if (GeometryBeingEdited == nullptr)
+	{
+		return;
+	}
+
 	//@TODO: Display 'action required' text in a bright color for Add Polygon mode
 	//@TODO: Move all of the line drawing to the PDI pass
 	FSpriteGeometryCollection& Geometry = GetGeometryChecked();
@@ -588,12 +608,12 @@ void FSpriteGeometryEditingHelper::AddPointToGeometry(const FVector2D& TextureSp
 		FSpriteGeometryShape& Shape = Geometry.Shapes[ClosestShapeIndex];
 		if (Shape.ShapeType != ESpriteShapeType::Circle)
 		{
-			EditorContext.BeginTransaction(LOCTEXT("AddPolygonVertexTransaction", "Add Vertex to Polygon"));
+			EditorContext->BeginTransaction(LOCTEXT("AddPolygonVertexTransaction", "Add Vertex to Polygon"));
 			Shape.Vertices.Insert(Shape.ConvertTextureSpaceToShapeSpace(TextureSpacePoint), ClosestVertexInsertIndex);
 			Shape.ShapeType = ESpriteShapeType::Polygon;
 			Geometry.GeometryType = ESpritePolygonMode::FullyCustom;
-			EditorContext.MarkTransactionAsDirty();
-			EditorContext.EndTransaction();
+			EditorContext->MarkTransactionAsDirty();
+			EditorContext->EndTransaction();
 
 			// Select this vertex
 			ClearSelectionSet();
@@ -605,7 +625,7 @@ void FSpriteGeometryEditingHelper::AddPointToGeometry(const FVector2D& TextureSp
 void FSpriteGeometryEditingHelper::SelectGeometry(const FShapeVertexPair& GeometyItem)
 {
 	SelectedIDSet.Add(GeometyItem);
-	EditorContext.InvalidateViewportAndHitProxies();
+	EditorContext->InvalidateViewportAndHitProxies();
 }
 
 bool FSpriteGeometryEditingHelper::IsGeometrySelected(const FShapeVertexPair& GeometyItem) const
@@ -665,7 +685,7 @@ void FSpriteGeometryEditingHelper::AddPolygonEdgeToSelection(const int32 ShapeIn
 void FSpriteGeometryEditingHelper::SetShowNormals(bool bShouldShowNormals)
 {
 	bShowNormals = bShouldShowNormals;
-	EditorContext.InvalidateViewportAndHitProxies();
+	EditorContext->InvalidateViewportAndHitProxies();
 }
 
 void FSpriteGeometryEditingHelper::ToggleShowNormals()
@@ -681,10 +701,15 @@ void FSpriteGeometryEditingHelper::SetGeometryBeingEdited(FSpriteGeometryCollect
 	bAllowSubtractivePolygons = bInAllowSubtractivePolygons;
 }
 
+FSpriteGeometryCollection* FSpriteGeometryEditingHelper::GetGeometryBeingEdited() const
+{
+	return GeometryBeingEdited;
+}
+
 void FSpriteGeometryEditingHelper::AddNewCircleShape(const FVector2D& CircleLocation, float Radius)
 {
 	FSpriteGeometryCollection& Geometry = GetGeometryChecked();
-	EditorContext.BeginTransaction(LOCTEXT("AddCircleShapeTransaction", "Add Circle Shape"));
+	EditorContext->BeginTransaction(LOCTEXT("AddCircleShapeTransaction", "Add Circle Shape"));
 
 	// Create the new shape
 	const FVector2D CircleSize2D(Radius*2.0f, Radius*2.0f);
@@ -695,14 +720,14 @@ void FSpriteGeometryEditingHelper::AddNewCircleShape(const FVector2D& CircleLoca
 	ClearSelectionSet();
 	AddShapeToSelection(Geometry.Shapes.Num() - 1);
 
-	EditorContext.MarkTransactionAsDirty();
-	EditorContext.EndTransaction();
+	EditorContext->MarkTransactionAsDirty();
+	EditorContext->EndTransaction();
 }
 
 void FSpriteGeometryEditingHelper::AddNewBoxShape(const FVector2D& BoxLocation, const FVector2D& BoxSize)
 {
 	FSpriteGeometryCollection& Geometry = GetGeometryChecked();
-	EditorContext.BeginTransaction(LOCTEXT("AddBoxShapeTransaction", "Add Box Shape"));
+	EditorContext->BeginTransaction(LOCTEXT("AddBoxShapeTransaction", "Add Box Shape"));
 
 	// Create the new shape
 	Geometry.AddRectangleShape(BoxLocation, BoxSize);
@@ -712,8 +737,8 @@ void FSpriteGeometryEditingHelper::AddNewBoxShape(const FVector2D& BoxLocation,
 	ClearSelectionSet();
 	AddShapeToSelection(Geometry.Shapes.Num() - 1);
 
-	EditorContext.MarkTransactionAsDirty();
-	EditorContext.EndTransaction();
+	EditorContext->MarkTransactionAsDirty();
+	EditorContext->EndTransaction();
 }
 
 void FSpriteGeometryEditingHelper::ToggleAddPolygonMode()
@@ -728,7 +753,7 @@ void FSpriteGeometryEditingHelper::ToggleAddPolygonMode()
 
 		bIsAddingPolygon = true;
 		AddingPolygonIndex = INDEX_NONE;
-		EditorContext.InvalidateViewportAndHitProxies();
+		EditorContext->InvalidateViewportAndHitProxies();
 	}
 }
 
@@ -740,20 +765,59 @@ void FSpriteGeometryEditingHelper::AbandonAddPolygonMode()
 
 	if ((AddingPolygonIndex != INDEX_NONE) && Geometry.Shapes.IsValidIndex(AddingPolygonIndex))
 	{
-		EditorContext.BeginTransaction(LOCTEXT("DeletePolygon", "Delete Polygon"));
+		EditorContext->BeginTransaction(LOCTEXT("DeletePolygon", "Delete Polygon"));
 
 		Geometry.Shapes.RemoveAt(AddingPolygonIndex);
 
-		EditorContext.MarkTransactionAsDirty();
-		EditorContext.EndTransaction();
+		EditorContext->MarkTransactionAsDirty();
+		EditorContext->EndTransaction();
 	}
 
 	ResetAddPolygonMode();
 }
 
+void FSpriteGeometryEditingHelper::SnapAllVerticesToPixelGrid()
+{
+	FSpriteGeometryCollection& Geometry = GetGeometryChecked();
+
+	EditorContext->BeginTransaction(LOCTEXT("SnapAllVertsToPixelGridTransaction", "Snap All Verts to Pixel Grid"));
+
+	for (FSpriteGeometryShape& Shape : Geometry.Shapes)
+	{
+		EditorContext->MarkTransactionAsDirty();
+
+		Shape.BoxPosition.X = FMath::RoundToInt(Shape.BoxPosition.X);
+		Shape.BoxPosition.Y = FMath::RoundToInt(Shape.BoxPosition.Y);
+
+		if ((Shape.ShapeType == ESpriteShapeType::Box) || (Shape.ShapeType == ESpriteShapeType::Circle))
+		{
+			//@TODO: Should we snap BoxPosition also, or just the verts?
+			const FVector2D OldHalfSize = Shape.BoxSize * 0.5f;
+			FVector2D TopLeft = Shape.BoxPosition - OldHalfSize;
+			FVector2D BottomRight = Shape.BoxPosition + OldHalfSize;
+			TopLeft.X = FMath::RoundToInt(TopLeft.X);
+			TopLeft.Y = FMath::RoundToInt(TopLeft.Y);
+			BottomRight.X = FMath::RoundToInt(BottomRight.X);
+			BottomRight.Y = FMath::RoundToInt(BottomRight.Y);
+			Shape.BoxPosition = (TopLeft + BottomRight) * 0.5f;
+			Shape.BoxSize = BottomRight - TopLeft;
+		}
+
+		for (FVector2D& Vertex : Shape.Vertices)
+		{
+			FVector2D TextureSpaceVertex = Shape.ConvertShapeSpaceToTextureSpace(Vertex);
+			TextureSpaceVertex.X = FMath::RoundToInt(TextureSpaceVertex.X);
+			TextureSpaceVertex.Y = FMath::RoundToInt(TextureSpaceVertex.Y);
+			Vertex = Shape.ConvertTextureSpaceToShapeSpace(TextureSpaceVertex);
+		}
+	}
+
+	EditorContext->EndTransaction();
+}
+
 void FSpriteGeometryEditingHelper::TEMP_HandleAddPolygonClick(const FVector2D& TexturePoint, bool bWantsSubtractive)
 {
-	EditorContext.BeginTransaction(LOCTEXT("AddPolygonVertexTransaction", "Add Vertex to Polygon"));
+	EditorContext->BeginTransaction(LOCTEXT("AddPolygonVertexTransaction", "Add Vertex to Polygon"));
 
 	FSpriteGeometryCollection& Geometry = GetGeometryChecked();
 	if (AddingPolygonIndex == INDEX_NONE)
@@ -779,8 +843,8 @@ void FSpriteGeometryEditingHelper::TEMP_HandleAddPolygonClick(const FVector2D& T
 
 	Geometry.GeometryType = ESpritePolygonMode::FullyCustom;
 
-	EditorContext.MarkTransactionAsDirty();
-	EditorContext.EndTransaction();
+	EditorContext->MarkTransactionAsDirty();
+	EditorContext->EndTransaction();
 }
 
 bool FSpriteGeometryEditingHelper::ClosestPointOnLine(const FVector2D& Point, const FVector2D& LineStart, const FVector2D& LineEnd, FVector2D& OutClosestPoint)
@@ -828,7 +892,7 @@ bool FSpriteGeometryEditingHelper::DeleteVertexInPolygonInternal(FSpriteGeometry
 
 FVector2D FSpriteGeometryEditingHelper::TextureSpaceToScreenSpace(const FSceneView& View, const FVector2D& SourcePoint) const
 {
-	const FVector WorldSpacePoint = EditorContext.TextureSpaceToWorldSpace(SourcePoint);
+	const FVector WorldSpacePoint = EditorContext->TextureSpaceToWorldSpace(SourcePoint);
 
 	FVector2D PixelLocation;
 	View.WorldToPixel(WorldSpacePoint, /*out*/ PixelLocation);
diff --git a/Engine/Plugins/2D/Paper2D/Source/Paper2DEditor/Private/PaperEditorShared/SpriteGeometryEditing.h b/Engine/Plugins/2D/Paper2D/Source/Paper2DEditor/Private/PaperEditorShared/SpriteGeometryEditing.h
--- a/Engine/Plugins/2D/Paper2D/Source/Paper2DEditor/Private/PaperEditorShared/SpriteGeometryEditing.h
+++ b/Engine/Plugins/2D/Paper2D/Source/Paper2DEditor/Private/PaperEditorShared/SpriteGeometryEditing.h
@@ -60,6 +60,8 @@ public:
 
 	virtual void DeleteSelectedItems(bool bShouldTransact = true);
 
+	virtual bool CanDeleteSelection() const;
+
 	~FSpriteSelectionHelper()
 	{
 	}
@@ -75,7 +77,7 @@ private:
 class FSpriteGeometryEditingHelper : public FSpriteSelectionHelper, public FGCObject
 {
 public:
-	FSpriteGeometryEditingHelper(class ISpriteSelectionContext& InEditorContext);
+	FSpriteGeometryEditingHelper(class ISpriteSelectionContext* InEditorContext);
 
 	// FSpriteSelectionHelper interface
 	virtual void ClearSelectionSet() override;
@@ -86,8 +88,11 @@ public:
 	virtual void AddReferencedObjects(FReferenceCollector& Collector) override;
 	// End of FGCObject interface
 
+	void SetEditorContext(class ISpriteSelectionContext* InNewEditorContext);
+	class ISpriteSelectionContext* GetEditorContext() const;
+
 	void DrawGeometry(const FSceneView& View, FPrimitiveDrawInterface& PDI, const FLinearColor& GeometryVertexColor, const FLinearColor& NegativeGeometryVertexColor);
-	void DrawGeometry_CanvasPass(FViewport& InViewport, FSceneView& View, FCanvas& Canvas, /*inout*/ int32& YPos, const FLinearColor& GeometryVertexColor, const FLinearColor& NegativeGeometryVertexColor);
+	void DrawGeometry_CanvasPass(FViewport& InViewport, const FSceneView& View, FCanvas& Canvas, /*inout*/ int32& YPos, const FLinearColor& GeometryVertexColor, const FLinearColor& NegativeGeometryVertexColor);
 
 	// Adds a point to the currently edited index
 	// If SelectedPolygonIndex is set, only that polygon is considered
@@ -107,6 +112,7 @@ public:
 
 	// Changes the geometry being edited (clears the selection set in the process)
 	void SetGeometryBeingEdited(FSpriteGeometryCollection* NewGeometryBeingEdited, bool bInAllowCircles, bool bInAllowSubtractivePolygons);
+	FSpriteGeometryCollection* GetGeometryBeingEdited() const;
 
 	// Adds a new circle shape and selects it
 	void AddNewCircleShape(const FVector2D& CircleLocation, float Radius);
@@ -123,9 +129,16 @@ public:
 	bool CanAddSubtractivePolygon() const { return CanAddPolygon() && bAllowSubtractivePolygons; }
 	void AbandonAddPolygonMode();
 
+	void SnapAllVerticesToPixelGrid();
+	bool CanSnapVerticesToPixelGrid() const { return GeometryBeingEdited != nullptr; }
+
 	//@TODO: The code calling this will eventually be folded in here too
 	void TEMP_HandleAddPolygonClick(const FVector2D& TexturePoint, bool bWantsSubtractive);
 
+	bool IsEditingGeometry() const
+	{
+		return GeometryBeingEdited != nullptr;
+	}
 private:
 	static bool ClosestPointOnLine(const FVector2D& Point, const FVector2D& LineStart, const FVector2D& LineEnd, FVector2D& OutClosestPoint);
 
@@ -143,7 +156,8 @@ private:
 	FVector2D TextureSpaceToScreenSpace(const FSceneView& View, const FVector2D& SourcePoint) const;
 
 private:
-	class ISpriteSelectionContext& EditorContext;
+	class ISpriteSelectionContext* EditorContext;
+
 	UMaterial* WidgetVertexColorMaterial;
 
 	// Set of selected vertices/shapes
diff --git a/Engine/Plugins/2D/Paper2D/Source/Paper2DEditor/Private/PaperStyle.cpp b/Engine/Plugins/2D/Paper2D/Source/Paper2DEditor/Private/PaperStyle.cpp
--- a/Engine/Plugins/2D/Paper2D/Source/Paper2DEditor/Private/PaperStyle.cpp
+++ b/Engine/Plugins/2D/Paper2D/Source/Paper2DEditor/Private/PaperStyle.cpp
@@ -104,22 +104,9 @@ void FPaperStyle::Initialize()
 
 		StyleSet->Set("SpriteEditor.SetShowSockets", new IMAGE_BRUSH(TEXT("Icons/icon_StaticMeshEd_ShowSockets_40x"), Icon40x40));
 		StyleSet->Set("SpriteEditor.SetShowSockets.Small", new IMAGE_BRUSH(TEXT("Icons/icon_StaticMeshEd_ShowSockets_40x"), Icon20x20));
-		StyleSet->Set("SpriteEditor.SetShowNormals", new IMAGE_BRUSH(TEXT("Icons/icon_StaticMeshEd_Normals_40x"), Icon40x40));
-		StyleSet->Set("SpriteEditor.SetShowNormals.Small", new IMAGE_BRUSH(TEXT("Icons/icon_StaticMeshEd_Normals_40x"), Icon20x20));
 		StyleSet->Set("SpriteEditor.SetShowPivot", new IMAGE_BRUSH(TEXT("Icons/icon_StaticMeshEd_ShowPivot_40x"), Icon40x40));
 		StyleSet->Set("SpriteEditor.SetShowPivot.Small", new IMAGE_BRUSH(TEXT("Icons/icon_StaticMeshEd_ShowPivot_40x"), Icon20x20));
 
-		StyleSet->Set("SpriteEditor.AddBoxShape", new IMAGE_PLUGIN_BRUSH(TEXT("Icons/icon_Paper2D_AddBoxShape_40x"), Icon40x40));
-		StyleSet->Set("SpriteEditor.AddBoxShape.Small", new IMAGE_PLUGIN_BRUSH(TEXT("Icons/icon_Paper2D_AddBoxShape_40x"), Icon20x20));
-		StyleSet->Set("SpriteEditor.AddCircleShape", new IMAGE_PLUGIN_BRUSH(TEXT("Icons/icon_Paper2D_AddCircleShape_40x"), Icon40x40));
-		StyleSet->Set("SpriteEditor.AddCircleShape.Small", new IMAGE_PLUGIN_BRUSH(TEXT("Icons/icon_Paper2D_AddCircleShape_40x"), Icon20x20));
-
-		StyleSet->Set("SpriteEditor.ToggleAddPolygonMode", new IMAGE_PLUGIN_BRUSH(TEXT("Icons/icon_Paper2D_AddPolygon_40x"), Icon40x40));
-		StyleSet->Set("SpriteEditor.ToggleAddPolygonMode.Small", new IMAGE_PLUGIN_BRUSH(TEXT("Icons/icon_Paper2D_AddPolygon_40x"), Icon20x20));
-
-		StyleSet->Set("SpriteEditor.SnapAllVertices", new IMAGE_PLUGIN_BRUSH(TEXT("Icons/icon_Paper2D_SnapToPixelGrid_40x"), Icon40x40));
-		StyleSet->Set("SpriteEditor.SnapAllVertices.Small", new IMAGE_PLUGIN_BRUSH(TEXT("Icons/icon_Paper2D_SnapToPixelGrid_40x"), Icon20x20));
-
 		StyleSet->Set("SpriteEditor.EnterViewMode", new IMAGE_PLUGIN_BRUSH(TEXT("Icons/icon_Paper2D_ViewSprite_40x"), Icon40x40));
 		StyleSet->Set("SpriteEditor.EnterViewMode.Small", new IMAGE_PLUGIN_BRUSH(TEXT("Icons/icon_Paper2D_ViewSprite_40x"), Icon20x20));
 		StyleSet->Set("SpriteEditor.EnterCollisionEditMode", new IMAGE_PLUGIN_BRUSH("Icons/icon_Paper2D_EditCollision_40x", Icon40x40));
@@ -130,6 +117,23 @@ void FPaperStyle::Initialize()
 		StyleSet->Set("SpriteEditor.EnterRenderingEditMode.Small", new IMAGE_PLUGIN_BRUSH("Icons/icon_Paper2D_RenderGeom_40x", Icon20x20));
 	}
 
+	// Sprite geometry editor (shared between the sprite editor, tile set editor, etc...)
+	{
+		StyleSet->Set("SpriteGeometryEditor.AddBoxShape", new IMAGE_PLUGIN_BRUSH(TEXT("Icons/icon_Paper2D_AddBoxShape_40x"), Icon40x40));
+		StyleSet->Set("SpriteGeometryEditor.AddBoxShape.Small", new IMAGE_PLUGIN_BRUSH(TEXT("Icons/icon_Paper2D_AddBoxShape_40x"), Icon20x20));
+		StyleSet->Set("SpriteGeometryEditor.AddCircleShape", new IMAGE_PLUGIN_BRUSH(TEXT("Icons/icon_Paper2D_AddCircleShape_40x"), Icon40x40));
+		StyleSet->Set("SpriteGeometryEditor.AddCircleShape.Small", new IMAGE_PLUGIN_BRUSH(TEXT("Icons/icon_Paper2D_AddCircleShape_40x"), Icon20x20));
+
+		StyleSet->Set("SpriteGeometryEditor.ToggleAddPolygonMode", new IMAGE_PLUGIN_BRUSH(TEXT("Icons/icon_Paper2D_AddPolygon_40x"), Icon40x40));
+		StyleSet->Set("SpriteGeometryEditor.ToggleAddPolygonMode.Small", new IMAGE_PLUGIN_BRUSH(TEXT("Icons/icon_Paper2D_AddPolygon_40x"), Icon20x20));
+
+		StyleSet->Set("SpriteGeometryEditor.SnapAllVertices", new IMAGE_PLUGIN_BRUSH(TEXT("Icons/icon_Paper2D_SnapToPixelGrid_40x"), Icon40x40));
+		StyleSet->Set("SpriteGeometryEditor.SnapAllVertices.Small", new IMAGE_PLUGIN_BRUSH(TEXT("Icons/icon_Paper2D_SnapToPixelGrid_40x"), Icon20x20));
+
+		StyleSet->Set("SpriteGeometryEditor.SetShowNormals", new IMAGE_BRUSH(TEXT("Icons/icon_StaticMeshEd_Normals_40x"), Icon40x40));
+		StyleSet->Set("SpriteGeometryEditor.SetShowNormals.Small", new IMAGE_BRUSH(TEXT("Icons/icon_StaticMeshEd_Normals_40x"), Icon20x20));
+	}
+
 	// Flipbook editor
 	{
 		StyleSet->Set("FlipbookEditor.SetShowGrid", new IMAGE_BRUSH(TEXT("Icons/icon_MatEd_Grid_40x"), Icon40x40));
diff --git a/Engine/Plugins/2D/Paper2D/Source/Paper2DEditor/Private/SpriteEditor/SSpriteEditorViewportToolbar.cpp b/Engine/Plugins/2D/Paper2D/Source/Paper2DEditor/Private/SpriteEditor/SSpriteEditorViewportToolbar.cpp
--- a/Engine/Plugins/2D/Paper2D/Source/Paper2DEditor/Private/SpriteEditor/SSpriteEditorViewportToolbar.cpp
+++ b/Engine/Plugins/2D/Paper2D/Source/Paper2DEditor/Private/SpriteEditor/SSpriteEditorViewportToolbar.cpp
@@ -31,7 +31,7 @@ TSharedRef<SWidget> SSpriteEditorViewportToolbar::GenerateShowMenu() const
 
 		ShowMenuBuilder.AddMenuEntry(FSpriteEditorCommands::Get().SetShowGrid);
 		ShowMenuBuilder.AddMenuEntry(FSpriteEditorCommands::Get().SetShowBounds);
-		ShowMenuBuilder.AddMenuEntry(FSpriteEditorCommands::Get().SetShowNormals);
+		ShowMenuBuilder.AddMenuEntry(FSpriteGeometryEditCommands::Get().SetShowNormals);
 
 		ShowMenuBuilder.AddMenuSeparator();
 
diff --git a/Engine/Plugins/2D/Paper2D/Source/Paper2DEditor/Private/SpriteEditor/SpriteEditor.cpp b/Engine/Plugins/2D/Paper2D/Source/Paper2DEditor/Private/SpriteEditor/SpriteEditor.cpp
--- a/Engine/Plugins/2D/Paper2D/Source/Paper2DEditor/Private/SpriteEditor/SpriteEditor.cpp
+++ b/Engine/Plugins/2D/Paper2D/Source/Paper2DEditor/Private/SpriteEditor/SpriteEditor.cpp
@@ -75,6 +75,11 @@ public:
 	{
 		return EditorViewportClient->GetCurrentMode();
 	}
+
+	void ActivateEditMode()
+	{
+		EditorViewportClient->ActivateEditMode();
+	}
 private:
 	// Pointer back to owning sprite editor instance (the keeper of state)
 	TWeakPtr<class FSpriteEditor> SpriteEditorPtr;
@@ -101,8 +106,6 @@ void SSpriteEditorViewport::BindCommands()
 	const FSpriteEditorCommands& Commands = FSpriteEditorCommands::Get();
 
 	TSharedRef<FSpriteEditorViewportClient> EditorViewportClientRef = EditorViewportClient.ToSharedRef();
-	FSpriteGeometryEditingHelper* GeometryEditingHelper = EditorViewportClientRef->GetGeometryEditHelper();
-
 
 	// Show toggles
 	CommandList->MapAction(
@@ -140,12 +143,6 @@ void SSpriteEditorViewport::BindCommands()
  		FExecuteAction::CreateSP( EditorViewportClientRef, &FSpriteEditorViewportClient::ToggleShowSockets ),
  		FCanExecuteAction(),
  		FIsActionChecked::CreateSP( EditorViewportClientRef, &FSpriteEditorViewportClient::IsShowSocketsChecked ) );
-
-	CommandList->MapAction(
-		Commands.SetShowNormals,
-		FExecuteAction::CreateRaw(GeometryEditingHelper, &FSpriteGeometryEditingHelper::ToggleShowNormals),
-		FCanExecuteAction(),
-		FIsActionChecked::CreateRaw(GeometryEditingHelper, &FSpriteGeometryEditingHelper::IsShowNormalsEnabled));
  
 	CommandList->MapAction(
 		Commands.SetShowPivot,
@@ -174,34 +171,6 @@ void SSpriteEditorViewport::BindCommands()
 		FExecuteAction::CreateSP( EditorViewportClientRef, &FSpriteEditorViewportClient::EnterRenderingEditMode ),
 		FCanExecuteAction(),
 		FIsActionChecked::CreateSP( EditorViewportClientRef, &FSpriteEditorViewportClient::IsInRenderingEditMode ) );
-
-	// Geometry editing commands
-	CommandList->MapAction(
-		Commands.DeleteSelection,
-		FExecuteAction::CreateSP( EditorViewportClientRef, &FSpriteEditorViewportClient::DeleteSelection ),
-		FCanExecuteAction::CreateSP( EditorViewportClientRef, &FSpriteEditorViewportClient::CanDeleteSelection ));
-	CommandList->MapAction(
-		Commands.AddBoxShape,
-		FExecuteAction::CreateSP(EditorViewportClientRef, &FSpriteEditorViewportClient::AddBoxShape),
-		FCanExecuteAction::CreateRaw(GeometryEditingHelper, &FSpriteGeometryEditingHelper::CanAddBoxShape),
-		FIsActionChecked(),
-		FIsActionButtonVisible::CreateRaw(GeometryEditingHelper, &FSpriteGeometryEditingHelper::CanAddBoxShape));
-	CommandList->MapAction(
-		Commands.ToggleAddPolygonMode,
-		FExecuteAction::CreateRaw(GeometryEditingHelper, &FSpriteGeometryEditingHelper::ToggleAddPolygonMode),
-		FCanExecuteAction::CreateRaw(GeometryEditingHelper, &FSpriteGeometryEditingHelper::CanAddPolygon),
-		FIsActionChecked::CreateRaw(GeometryEditingHelper, &FSpriteGeometryEditingHelper::IsAddingPolygon),
-		FIsActionButtonVisible::CreateRaw(GeometryEditingHelper, &FSpriteGeometryEditingHelper::CanAddPolygon));
-	CommandList->MapAction(
-		Commands.AddCircleShape,
-		FExecuteAction::CreateSP(EditorViewportClientRef, &FSpriteEditorViewportClient::AddCircleShape),
-		FCanExecuteAction::CreateRaw(GeometryEditingHelper, &FSpriteGeometryEditingHelper::CanAddCircleShape),
-		FIsActionChecked(),
-		FIsActionButtonVisible::CreateRaw(GeometryEditingHelper, &FSpriteGeometryEditingHelper::CanAddCircleShape));
-	CommandList->MapAction(
-		Commands.SnapAllVertices,
-		FExecuteAction::CreateSP( EditorViewportClientRef, &FSpriteEditorViewportClient::SnapAllVerticesToPixelGrid ),
-		FCanExecuteAction::CreateSP( EditorViewportClientRef, &FSpriteEditorViewportClient::CanSnapVerticesToPixelGrid ));
 }
 
 TSharedRef<FEditorViewportClient> SSpriteEditorViewport::MakeEditorViewportClient()
@@ -364,6 +333,7 @@ void FSpriteEditor::InitSpriteEditor(const EToolkitMode::Type Mode, const TShare
 	SpriteBeingEdited = InitSprite;
 
 	FSpriteEditorCommands::Register();
+	FSpriteGeometryEditCommands::Register();
 
 	BindCommands();
 
@@ -418,6 +388,8 @@ void FSpriteEditor::InitSpriteEditor(const EToolkitMode::Type Mode, const TShare
 	// Initialize the asset editor
 	InitAssetEditor(Mode, InitToolkitHost, SpriteEditorAppName, StandaloneDefaultLayout, /*bCreateDefaultStandaloneMenu=*/ true, /*bCreateDefaultToolbar=*/ true, InitSprite);
 
+	ViewportPtr->ActivateEditMode();
+
 	// Extend things
 	ExtendMenu();
 	ExtendToolbar();
@@ -426,24 +398,6 @@ void FSpriteEditor::InitSpriteEditor(const EToolkitMode::Type Mode, const TShare
 
 void FSpriteEditor::BindCommands()
 {
-// 	const FSpriteEditorCommands& Commands = FSpriteEditorCommands::Get();
-// 
-// 	const TSharedRef<FUICommandList>& UICommandList = GetToolkitCommands();
-// 
-// 	UICommandList->MapAction( FGenericCommands::Get().Delete,
-// 		FExecuteAction::CreateSP( this, &FStaticMeshEditor::DeleteSelectedSockets ),
-// 		FCanExecuteAction::CreateSP( this, &FStaticMeshEditor::HasSelectedSockets ));
-// 
-// 	UICommandList->MapAction( FGenericCommands::Get().Undo, 
-// 		FExecuteAction::CreateSP( this, &FStaticMeshEditor::UndoAction ) );
-// 
-// 	UICommandList->MapAction( FGenericCommands::Get().Redo, 
-// 		FExecuteAction::CreateSP( this, &FStaticMeshEditor::RedoAction ) );
-// 
-// 	UICommandList->MapAction(
-// 		FGenericCommands::Get().Duplicate,
-// 		FExecuteAction::CreateSP(this, &FStaticMeshEditor::DuplicateSelectedSocket),
-// 		FCanExecuteAction::CreateSP(this, &FStaticMeshEditor::HasSelectedSockets));
 }
 
 FName FSpriteEditor::GetToolkitFName() const
@@ -476,6 +430,22 @@ FString FSpriteEditor::GetDocumentationLink() const
 	return TEXT("Engine/Paper2D/SpriteEditor");
 }
 
+void FSpriteEditor::OnToolkitHostingStarted(const TSharedRef<class IToolkit>& Toolkit)
+{
+	//@TODO: MODETOOLS: Need to be able to register the widget in the toolbox panel with ToolkitHost, so it can instance the ed mode widgets into it
+	// 	TSharedPtr<SWidget> InlineContent = Toolkit->GetInlineContent();
+// 	if (InlineContent.IsValid())
+// 	{
+// 		ToolboxPtr->SetContent(InlineContent.ToSharedRef());
+// 	}
+}
+
+void FSpriteEditor::OnToolkitHostingFinished(const TSharedRef<class IToolkit>& Toolkit)
+{
+	//ToolboxPtr->SetContent(SNullWidget::NullWidget);
+	//@TODO: MODETOOLS: How to handle multiple ed modes at once in a standalone asset editor?
+}
+
 FLinearColor FSpriteEditor::GetWorldCentricTabColorScale() const
 {
 	return FLinearColor::White;
@@ -515,10 +485,10 @@ void FSpriteEditor::ExtendToolbar()
 
 			ToolbarBuilder.BeginSection("Tools");
 			{
-				ToolbarBuilder.AddToolBarButton(FSpriteEditorCommands::Get().AddBoxShape);
-				ToolbarBuilder.AddToolBarButton(FSpriteEditorCommands::Get().ToggleAddPolygonMode);
-				ToolbarBuilder.AddToolBarButton(FSpriteEditorCommands::Get().AddCircleShape);
-				ToolbarBuilder.AddToolBarButton(FSpriteEditorCommands::Get().SnapAllVertices);
+				ToolbarBuilder.AddToolBarButton(FSpriteGeometryEditCommands::Get().AddBoxShape);
+				ToolbarBuilder.AddToolBarButton(FSpriteGeometryEditCommands::Get().ToggleAddPolygonMode);
+				ToolbarBuilder.AddToolBarButton(FSpriteGeometryEditCommands::Get().AddCircleShape);
+				ToolbarBuilder.AddToolBarButton(FSpriteGeometryEditCommands::Get().SnapAllVertices);
 			}
 			ToolbarBuilder.EndSection();
 		}
diff --git a/Engine/Plugins/2D/Paper2D/Source/Paper2DEditor/Private/SpriteEditor/SpriteEditor.h b/Engine/Plugins/2D/Paper2D/Source/Paper2DEditor/Private/SpriteEditor/SpriteEditor.h
--- a/Engine/Plugins/2D/Paper2D/Source/Paper2DEditor/Private/SpriteEditor/SpriteEditor.h
+++ b/Engine/Plugins/2D/Paper2D/Source/Paper2DEditor/Private/SpriteEditor/SpriteEditor.h
@@ -37,6 +37,8 @@ public:
 	virtual FLinearColor GetWorldCentricTabColorScale() const override;
 	virtual FString GetWorldCentricTabPrefix() const override;
 	virtual FString GetDocumentationLink() const override;
+	virtual void OnToolkitHostingStarted(const TSharedRef<class IToolkit>& Toolkit) override;
+	virtual void OnToolkitHostingFinished(const TSharedRef<class IToolkit>& Toolkit) override;
 	// End of FAssetEditorToolkit
 
 	// FSerializableObject interface
diff --git a/Engine/Plugins/2D/Paper2D/Source/Paper2DEditor/Private/SpriteEditor/SpriteEditorCommands.cpp b/Engine/Plugins/2D/Paper2D/Source/Paper2DEditor/Private/SpriteEditor/SpriteEditorCommands.cpp
--- a/Engine/Plugins/2D/Paper2D/Source/Paper2DEditor/Private/SpriteEditor/SpriteEditorCommands.cpp
+++ b/Engine/Plugins/2D/Paper2D/Source/Paper2DEditor/Private/SpriteEditor/SpriteEditorCommands.cpp
@@ -3,10 +3,27 @@
 #include "Paper2DEditorPrivatePCH.h"
 #include "SpriteEditorCommands.h"
 
+#define LOCTEXT_NAMESPACE ""
+
 //////////////////////////////////////////////////////////////////////////
-// FSpriteEditorCommands
+// FSpriteGeometryEditCommands
 
-#define LOCTEXT_NAMESPACE ""
+void FSpriteGeometryEditCommands::RegisterCommands()
+{
+	// Show toggles
+	UI_COMMAND(SetShowNormals, "Normals", "Toggles display of vertex normals in the Preview Pane.", EUserInterfaceActionType::ToggleButton, FInputChord());
+
+	// Geometry editing commands
+	UI_COMMAND(DeleteSelection, "Delete", "Delete the selection.", EUserInterfaceActionType::Button, FInputChord(EKeys::Platform_Delete, EModifierKey::None));
+	UI_COMMAND(SplitEdge, "Split", "Split edge.", EUserInterfaceActionType::Button, FInputChord(EKeys::Insert, EModifierKey::None));
+	UI_COMMAND(AddBoxShape, "Add Box", "Adds a new box shape.", EUserInterfaceActionType::Button, FInputChord());
+	UI_COMMAND(ToggleAddPolygonMode, "Add Polygon", "Adds a new polygon shape.", EUserInterfaceActionType::ToggleButton, FInputChord());
+	UI_COMMAND(AddCircleShape, "Add Circle", "Adds a new circle shape.", EUserInterfaceActionType::Button, FInputChord());
+	UI_COMMAND(SnapAllVertices, "Snap to pixel grid", "Snaps all vertices to the pixel grid.", EUserInterfaceActionType::Button, FInputChord());
+}
+
+//////////////////////////////////////////////////////////////////////////
+// FSpriteEditorCommands
 
 void FSpriteEditorCommands::RegisterCommands()
 {
@@ -18,7 +35,6 @@ void FSpriteEditorCommands::RegisterCommands()
 
 	UI_COMMAND(SetShowSockets, "Sockets", "Displays the sprite sockets.", EUserInterfaceActionType::ToggleButton, FInputChord());
 
-	UI_COMMAND(SetShowNormals, "Normals", "Toggles display of vertex normals in the Preview Pane.", EUserInterfaceActionType::ToggleButton, FInputChord());
 	UI_COMMAND(SetShowPivot, "Pivot", "Display the pivot location of the sprite.", EUserInterfaceActionType::ToggleButton, FInputChord());
 	UI_COMMAND(SetShowMeshEdges, "Mesh Edges", "Overlays the mesh edges on top of the view.", EUserInterfaceActionType::ToggleButton, FInputChord());
 
@@ -27,14 +43,8 @@ void FSpriteEditorCommands::RegisterCommands()
 	UI_COMMAND(EnterSourceRegionEditMode, "Edit Source Region", "Edit the sprite source region.", EUserInterfaceActionType::ToggleButton, FInputChord());
 	UI_COMMAND(EnterCollisionEditMode, "Edit Collision", "Edit the collision geometry.", EUserInterfaceActionType::ToggleButton, FInputChord());
 	UI_COMMAND(EnterRenderingEditMode, "Edit RenderGeom", "Edit the rendering geometry (useful to reduce overdraw).", EUserInterfaceActionType::ToggleButton, FInputChord());
-
-	// Geometry editing commands
-	UI_COMMAND(DeleteSelection, "Delete", "Delete the selection.", EUserInterfaceActionType::Button, FInputChord(EKeys::Platform_Delete, EModifierKey::None));
-	UI_COMMAND(SplitEdge, "Split", "Split edge.", EUserInterfaceActionType::Button, FInputChord(EKeys::Insert, EModifierKey::None));
-	UI_COMMAND(AddBoxShape, "Add Box", "Adds a new box shape.", EUserInterfaceActionType::Button, FInputChord());
-	UI_COMMAND(ToggleAddPolygonMode, "Add Polygon", "Adds a new polygon shape.", EUserInterfaceActionType::ToggleButton, FInputChord());
-	UI_COMMAND(AddCircleShape, "Add Circle", "Adds a new circle shape.", EUserInterfaceActionType::Button, FInputChord());
-	UI_COMMAND(SnapAllVertices, "Snap to pixel grid", "Snaps all vertices to the pixel grid.", EUserInterfaceActionType::Button, FInputChord());
 }
 
+//////////////////////////////////////////////////////////////////////////
+
 #undef LOCTEXT_NAMESPACE
\ No newline at end of file
diff --git a/Engine/Plugins/2D/Paper2D/Source/Paper2DEditor/Private/SpriteEditor/SpriteEditorCommands.h b/Engine/Plugins/2D/Paper2D/Source/Paper2DEditor/Private/SpriteEditor/SpriteEditorCommands.h
--- a/Engine/Plugins/2D/Paper2D/Source/Paper2DEditor/Private/SpriteEditor/SpriteEditorCommands.h
+++ b/Engine/Plugins/2D/Paper2D/Source/Paper2DEditor/Private/SpriteEditor/SpriteEditorCommands.h
@@ -4,6 +4,37 @@
 
 #include "PaperStyle.h"
 
+class FSpriteGeometryEditCommands : public TCommands<FSpriteGeometryEditCommands>
+{
+public:
+	FSpriteGeometryEditCommands()
+		: TCommands<FSpriteGeometryEditCommands>(
+			TEXT("SpriteGeometryEditor"), // Context name for fast lookup
+			NSLOCTEXT("Contexts", "PaperEditor", "Sprite Geometry Editor"), // Localized context name for displaying
+			NAME_None, // Parent
+			FPaperStyle::Get()->GetStyleSetName() // Icon Style Set
+			)
+	{
+	}
+
+	// TCommand<> interface
+	virtual void RegisterCommands() override;
+	// End of TCommand<> interface
+
+public:
+	// Show toggles
+	TSharedPtr<FUICommandInfo> SetShowNormals;
+
+	// Geometry editing commands
+	TSharedPtr<FUICommandInfo> DeleteSelection;
+	TSharedPtr<FUICommandInfo> SplitEdge;
+	TSharedPtr<FUICommandInfo> AddBoxShape;
+	TSharedPtr<FUICommandInfo> AddCircleShape;
+	TSharedPtr<FUICommandInfo> ToggleAddPolygonMode;
+	TSharedPtr<FUICommandInfo> SnapAllVertices;
+};
+
+
 class FSpriteEditorCommands : public TCommands<FSpriteEditorCommands>
 {
 public:
@@ -30,7 +61,6 @@ public:
 	
 	TSharedPtr<FUICommandInfo> SetShowSockets;
 
-	TSharedPtr<FUICommandInfo> SetShowNormals;
 	TSharedPtr<FUICommandInfo> SetShowPivot;
 	TSharedPtr<FUICommandInfo> SetShowMeshEdges;
 
@@ -39,12 +69,4 @@ public:
 	TSharedPtr<FUICommandInfo> EnterSourceRegionEditMode;
 	TSharedPtr<FUICommandInfo> EnterCollisionEditMode;
 	TSharedPtr<FUICommandInfo> EnterRenderingEditMode;
-
-	// Geometry editing commands
-	TSharedPtr<FUICommandInfo> DeleteSelection;
-	TSharedPtr<FUICommandInfo> SplitEdge;
-	TSharedPtr<FUICommandInfo> AddBoxShape;
-	TSharedPtr<FUICommandInfo> AddCircleShape;
-	TSharedPtr<FUICommandInfo> ToggleAddPolygonMode;
-	TSharedPtr<FUICommandInfo> SnapAllVertices;
 };
\ No newline at end of file
diff --git a/Engine/Plugins/2D/Paper2D/Source/Paper2DEditor/Private/SpriteEditor/SpriteEditorSelections.h b/Engine/Plugins/2D/Paper2D/Source/Paper2DEditor/Private/SpriteEditor/SpriteEditorSelections.h
--- a/Engine/Plugins/2D/Paper2D/Source/Paper2DEditor/Private/SpriteEditor/SpriteEditorSelections.h
+++ b/Engine/Plugins/2D/Paper2D/Source/Paper2DEditor/Private/SpriteEditor/SpriteEditorSelections.h
@@ -26,8 +26,8 @@ class ISpriteSelectionContext
 {
 public:
 	virtual FVector2D SelectedItemConvertWorldSpaceDeltaToLocalSpace(const FVector& WorldSpaceDelta) const = 0;
+	virtual FVector2D WorldSpaceToTextureSpace(const FVector& SourcePoint) const = 0;
 	virtual FVector TextureSpaceToWorldSpace(const FVector2D& SourcePoint) const = 0;
-	virtual bool SelectedItemIsSelected(const struct FShapeVertexPair& Item) const = 0;
 	virtual float SelectedItemGetUnitsPerPixel() const = 0;
 	virtual void BeginTransaction(const FText& SessionName) = 0;
 	virtual void MarkTransactionAsDirty() = 0;
@@ -211,7 +211,7 @@ class FSpriteSelectedShape : public FSelectedItem
 {
 public:
 	// The editor context
-	ISpriteSelectionContext& EditorContext;
+	ISpriteSelectionContext* EditorContext;
 
 	// The geometry that this shape belongs to
 	FSpriteGeometryCollection& Geometry;
@@ -224,7 +224,7 @@ public:
 
 	TWeakObjectPtr<UPaperSprite> SpritePtr;
 public:
-	FSpriteSelectedShape(ISpriteSelectionContext& InEditorContext, FSpriteGeometryCollection& InGeometry, int32 InShapeIndex, bool bInIsBackground = false);
+	FSpriteSelectedShape(ISpriteSelectionContext* InEditorContext, FSpriteGeometryCollection& InGeometry, int32 InShapeIndex, bool bInIsBackground = false);
 
 	// FSelectedItem interface
 	virtual uint32 GetTypeHash() const override;
@@ -243,7 +243,7 @@ class FSpriteSelectedVertex : public FSelectedItem
 {
 public:
 	// The editor context
-	ISpriteSelectionContext& EditorContext;
+	ISpriteSelectionContext* EditorContext;
 
 	// The geometry that this shape belongs to
 	FSpriteGeometryCollection& Geometry;
@@ -252,7 +252,7 @@ public:
 	const int32 VertexIndex;
 
 public:
-	FSpriteSelectedVertex(ISpriteSelectionContext& InEditorContext, FSpriteGeometryCollection& InGeometry, int32 InShapeIndex, int32 InVertexIndex)
+	FSpriteSelectedVertex(ISpriteSelectionContext* InEditorContext, FSpriteGeometryCollection& InGeometry, int32 InShapeIndex, int32 InVertexIndex)
 		: FSelectedItem(FSelectionTypes::Vertex)
 		, EditorContext(InEditorContext)
 		, Geometry(InGeometry)
@@ -293,7 +293,7 @@ public:
 				FVector2D& ShapeSpaceVertex = Vertices[TargetVertexIndex];
 
 				const FVector WorldSpaceDelta = PaperAxisX * Delta.X + PaperAxisY * Delta.Y;
-				const FVector2D TextureSpaceDelta = EditorContext.SelectedItemConvertWorldSpaceDeltaToLocalSpace(WorldSpaceDelta);
+				const FVector2D TextureSpaceDelta = EditorContext->SelectedItemConvertWorldSpaceDeltaToLocalSpace(WorldSpaceDelta);
 				const FVector2D NewTextureSpacePos = Shape.ConvertShapeSpaceToTextureSpace(ShapeSpaceVertex) + TextureSpaceDelta;
 				ShapeSpaceVertex = Shape.ConvertTextureSpaceToShapeSpace(NewTextureSpacePos);
 
@@ -316,7 +316,7 @@ public:
 			{
 				const FVector2D& ShapeSpacePos = Vertices[TargetVertexIndex];
 				const FVector2D TextureSpacePos = Shape.ConvertShapeSpaceToTextureSpace(ShapeSpacePos);
-				Result = EditorContext.TextureSpaceToWorldSpace(TextureSpacePos);
+				Result = EditorContext->TextureSpaceToWorldSpace(TextureSpacePos);
 			}
 		}
 
@@ -345,7 +345,7 @@ public:
 class FSpriteSelectedEdge : public FSpriteSelectedVertex
 {
 public:
-	FSpriteSelectedEdge(ISpriteSelectionContext& InEditorContext, FSpriteGeometryCollection& InGeometry, int32 InShapeIndex, int32 InVertexIndex)
+	FSpriteSelectedEdge(ISpriteSelectionContext* InEditorContext, FSpriteGeometryCollection& InGeometry, int32 InShapeIndex, int32 InVertexIndex)
 		: FSpriteSelectedVertex(InEditorContext, InGeometry, InShapeIndex, InVertexIndex)
 	{
 		TypeName = FSelectionTypes::Edge;
diff --git a/Engine/Plugins/2D/Paper2D/Source/Paper2DEditor/Private/SpriteEditor/SpriteEditorViewportClient.cpp b/Engine/Plugins/2D/Paper2D/Source/Paper2DEditor/Private/SpriteEditor/SpriteEditorViewportClient.cpp
--- a/Engine/Plugins/2D/Paper2D/Source/Paper2DEditor/Private/SpriteEditor/SpriteEditorViewportClient.cpp
+++ b/Engine/Plugins/2D/Paper2D/Source/Paper2DEditor/Private/SpriteEditor/SpriteEditorViewportClient.cpp
@@ -15,7 +15,10 @@
 #include "CanvasItem.h"
 #include "PaperEditorShared/SocketEditing.h"
 
+#include "PaperEditorShared/SpriteGeometryEditMode.h"
+
 #include "PhysicsEngine/BodySetup2D.h"
+#include "SEditorViewport.h"
 
 #define LOCTEXT_NAMESPACE "SpriteEditor"
 
@@ -42,8 +45,6 @@ namespace SpriteEditingConstants
 	const FLinearColor BakedCollisionRenderColor(1.0f, 1.0f, 0.0f, 0.5f);
 	const float BakedCollisionVertexSize = 3.0f;
 
-	const FLinearColor MarqueeColor(1.0f, 1.0f, 1.0f, 0.5f);
-
 	const FLinearColor SourceRegionBoundsColor(1.0f, 1.0f, 1.0f, 0.8f);
 	// Alpha is being ignored in FBatchedElements::AddLine :(
 	const FLinearColor SourceRegionRelatedBoundsColor(0.3f, 0.3f, 0.3f, 0.2f);
@@ -56,19 +57,21 @@ namespace SpriteEditingConstants
 //////////////////////////////////////////////////////////////////////////
 // FSpriteEditorViewportClient
 
-FSpriteEditorViewportClient::FSpriteEditorViewportClient(TWeakPtr<FSpriteEditor> InSpriteEditor, TWeakPtr<class SSpriteEditorViewport> InSpriteEditorViewportPtr)
+FSpriteEditorViewportClient::FSpriteEditorViewportClient(TWeakPtr<FSpriteEditor> InSpriteEditor, TWeakPtr<class SEditorViewport> InSpriteEditorViewportPtr)
 	: CurrentMode(ESpriteEditorMode::ViewMode)
 	, SpriteEditorPtr(InSpriteEditor)
 	, SpriteEditorViewportPtr(InSpriteEditorViewportPtr)
-	, SpriteGeometryHelper(*this)
 {
 	check(SpriteEditorPtr.IsValid() && SpriteEditorViewportPtr.IsValid());
 
+	// The tile map editor fully supports mode tools and isn't doing any incompatible stuff with the Widget
+	Widget->SetUsesEditorModeTools(ModeTools);
+
 	PreviewScene = &OwnedPreviewScene;
+	((FAssetEditorModeManager*)ModeTools)->SetPreviewScene(PreviewScene);
 
 	SetRealtime(true);
 
-	DesiredWidgetMode = FWidget::WM_Translate;
 	bManipulating = false;
 	bManipulationDirtiedSomething = false;
 	ScopedTransaction = nullptr;
@@ -78,8 +81,6 @@ FSpriteEditorViewportClient::FSpriteEditorViewportClient(TWeakPtr<FSpriteEditor>
 	bShowPivot = true;
 	bShowRelatedSprites = true;
 
-	bIsMarqueeTracking = false;
-
 	DrawHelper.bDrawGrid = false;
 
 	EngineShowFlags.DisableAdvancedFeatures();
@@ -113,6 +114,20 @@ FSpriteEditorViewportClient::FSpriteEditorViewportClient(TWeakPtr<FSpriteEditor>
 	}
 }
 
+void FSpriteEditorViewportClient::ActivateEditMode()
+{
+	// Activate the sprite geometry edit mode
+	ModeTools->SetToolkitHost(SpriteEditorPtr.Pin()->GetToolkitHost());
+	ModeTools->SetDefaultMode(FSpriteGeometryEditMode::EM_SpriteGeometry);
+	ModeTools->ActivateDefaultMode();
+
+	FSpriteGeometryEditMode* GeometryEditMode = ModeTools->GetActiveModeTyped<FSpriteGeometryEditMode>(FSpriteGeometryEditMode::EM_SpriteGeometry);
+	check(GeometryEditMode);
+	GeometryEditMode->SetEditorContext(this);
+	GeometryEditMode->BindCommands(SpriteEditorViewportPtr.Pin()->GetCommandList());
+	ModeTools->SetWidgetMode(FWidget::WM_Translate);
+}
+
 void FSpriteEditorViewportClient::UpdateSourceTextureSpriteFromSprite(UPaperSprite* SourceSprite)
 {
 	UPaperSprite* TargetSprite = SourceTextureViewComponent->GetSprite();
@@ -287,23 +302,6 @@ void FSpriteEditorViewportClient::DrawSourceRegion(FViewport& InViewport, FScene
 	}
 }
 
-void FSpriteEditorViewportClient::DrawMarquee(FViewport& InViewport, FSceneView& View, FCanvas& Canvas, const FLinearColor& MarqueeColor)
-{
-	FVector2D MarqueeDelta = MarqueeEndPos - MarqueeStartPos;
-	FVector2D MarqueeVertices[4];
-	MarqueeVertices[0] = FVector2D(MarqueeStartPos.X, MarqueeStartPos.Y);
-	MarqueeVertices[1] = MarqueeVertices[0] + FVector2D(MarqueeDelta.X, 0);
-	MarqueeVertices[2] = MarqueeVertices[0] + FVector2D(MarqueeDelta.X, MarqueeDelta.Y);
-	MarqueeVertices[3] = MarqueeVertices[0] + FVector2D(0, MarqueeDelta.Y);
-	for (int32 MarqueeVertexIndex = 0; MarqueeVertexIndex < 4; ++MarqueeVertexIndex)
-	{
-		const int32 NextVertexIndex = (MarqueeVertexIndex + 1) % 4;
-		FCanvasLineItem MarqueeLine(MarqueeVertices[MarqueeVertexIndex], MarqueeVertices[NextVertexIndex]);
-		MarqueeLine.SetColor(MarqueeColor);
-		Canvas.DrawItem(MarqueeLine);
-	}
-}
-
 void FSpriteEditorViewportClient::DrawGeometryStats(FViewport& InViewport, FSceneView& View, FCanvas& Canvas, const FSpriteGeometryCollection& Geometry, bool bIsRenderGeometry, int32& YPos)
 {
 	// Draw the type of geometry we're displaying stats for
@@ -591,8 +589,8 @@ void FSpriteEditorViewportClient::DrawCanvas(FViewport& Viewport, FSceneView& Vi
 		break;
 	case ESpriteEditorMode::EditCollisionMode:
 		{
-			// Draw the custom collision geometry
-			SpriteGeometryHelper.DrawGeometry_CanvasPass(Viewport, View, Canvas, /*inout*/ YPos, SpriteEditingConstants::CollisionShapeColor, FLinearColor::White);
+			// Draw the collision geometry stats
+			YPos += 60; //@TODO: Need a better way to determine this from the editor mode
 			if (Sprite->BodySetup != nullptr)
 			{
 				DrawGeometryStats(Viewport, View, Canvas, Sprite->CollisionGeometry, false, /*inout*/ YPos);
@@ -608,8 +606,8 @@ void FSpriteEditorViewportClient::DrawCanvas(FViewport& Viewport, FSceneView& Vi
 		break;
 	case ESpriteEditorMode::EditRenderingGeomMode:
 		{
-			// Draw the custom render geometry
-			SpriteGeometryHelper.DrawGeometry_CanvasPass(Viewport, View, Canvas, /*inout*/ YPos, SpriteEditingConstants::RenderShapeColor, SpriteEditingConstants::SubtractiveRenderShapeColor);
+			// Draw the render geometry stats
+			YPos += 60; //@TODO: Need a better way to determine this from the editor mode
 			DrawGeometryStats(Viewport, View, Canvas, Sprite->RenderGeometry, true, /*inout*/ YPos);
 			DrawRenderStats(Viewport, View, Canvas, Sprite, /*inout*/ YPos);
 
@@ -637,11 +635,6 @@ void FSpriteEditorViewportClient::DrawCanvas(FViewport& Viewport, FSceneView& Vi
 		break;
 	}
 
-	if (bIsMarqueeTracking)
-	{
-		DrawMarquee(Viewport, View, Canvas, SpriteEditingConstants::MarqueeColor);
-	}
-
 	if (bShowSockets && !IsInSourceRegionEditMode())
 	{
 		FSocketEditingHelper::DrawSocketNames(RenderSpriteComponent, Viewport, View, Canvas);
@@ -667,20 +660,6 @@ void FSpriteEditorViewportClient::Draw(const FSceneView* View, FPrimitiveDrawInt
 	{
 		FSocketEditingHelper::DrawSockets(RenderSpriteComponent, View, PDI);
 	}
-
-	switch (CurrentMode)
-	{
-	case ESpriteEditorMode::EditCollisionMode:
-		// Draw the custom collision geometry
-		SpriteGeometryHelper.DrawGeometry(*View, *PDI, SpriteEditingConstants::CollisionShapeColor, FLinearColor::White);
-		break;
-	case ESpriteEditorMode::EditRenderingGeomMode:
-		// Draw the custom render geometry
-		SpriteGeometryHelper.DrawGeometry(*View, *PDI, SpriteEditingConstants::RenderShapeColor, SpriteEditingConstants::SubtractiveRenderShapeColor);
-		break;
-	default:
-		break;
-	}
 }
 
 void FSpriteEditorViewportClient::Draw(FViewport* Viewport, FCanvas* Canvas)
@@ -722,13 +701,17 @@ void FSpriteEditorViewportClient::Tick(float DeltaSeconds)
 			RequestFocusOnSelection(/*bInstant=*/ true);
 			RenderSpriteComponent->SetVisibility(bRenderTextureViewComponentVisibility);
 		}
-	}
 
-	if (bIsMarqueeTracking)
-	{
-		int32 HitX = Viewport->GetMouseX();
-		int32 HitY = Viewport->GetMouseY();
-		MarqueeEndPos = FVector2D(HitX, HitY);
+		const FVector2D BoxSize(Sprite->GetSourceSize());
+		const FVector2D BoxLocation(Sprite->GetSourceUV() + (BoxSize * 0.5f));
+		FBox2D SpriteBounds(ForceInitToZero);
+		SpriteBounds.Min = BoxSize - BoxLocation * 0.5f;
+		SpriteBounds.Max = BoxSize + BoxLocation * 0.5f;
+
+		if (FSpriteGeometryEditMode* GeometryEditMode = ModeTools->GetActiveModeTyped<FSpriteGeometryEditMode>(FSpriteGeometryEditMode::EM_SpriteGeometry))
+		{
+			GeometryEditMode->SetNewGeometryPreferredBounds(SpriteBounds);
+		}
 	}
 
 	FPaperEditorViewportClient::Tick(DeltaSeconds);
@@ -812,173 +795,59 @@ void FSpriteEditorViewportClient::ProcessClick(FSceneView& View, HHitProxy* HitP
 	const bool bIsAltKeyDown = Viewport->KeyState(EKeys::LeftAlt) || Viewport->KeyState(EKeys::RightAlt);
 	bool bHandled = false;
 
-	const bool bAllowSelectVertex = !(IsEditingGeometry() && SpriteGeometryHelper.IsAddingPolygon()) && !bIsShiftKeyDown;
-
-	const bool bClearSelectionModifier = bIsCtrlKeyDown;
-	const bool bDeleteClickedVertex = bIsAltKeyDown;
-	const bool bInsertVertexModifier = bIsShiftKeyDown;
 	HSpriteSelectableObjectHitProxy* SelectedItemProxy = HitProxyCast<HSpriteSelectableObjectHitProxy>(HitProxy);
 
-	if (bAllowSelectVertex && (SelectedItemProxy != nullptr))
-	{
-		if (!bClearSelectionModifier)
-		{
-			SpriteGeometryHelper.ClearSelectionSet();
-		}
-
-		if (bDeleteClickedVertex)
-		{
-			// Delete selection
-			if (const FSpriteSelectedVertex* SelectedVertex = SelectedItemProxy->Data->CastTo<const FSpriteSelectedVertex>(FSelectionTypes::Vertex))
-			{
-				SpriteGeometryHelper.ClearSelectionSet();
-				SpriteGeometryHelper.AddPolygonVertexToSelection(SelectedVertex->ShapeIndex, SelectedVertex->VertexIndex);
-				DeleteSelection();
-			}
-			else if (const FSpriteSelectedShape* SelectedShape = SelectedItemProxy->Data->CastTo<const FSpriteSelectedShape>(FSelectionTypes::GeometryShape))
-			{
-				SpriteGeometryHelper.ClearSelectionSet();
-				SpriteGeometryHelper.AddShapeToSelection(SelectedShape->ShapeIndex);
-				DeleteSelection();
-			}
-		}
-		else if (Event == EInputEvent::IE_DoubleClick)
-		{
-			// Double click to select a polygon
-			if (const FSpriteSelectedVertex* SelectedVertex = SelectedItemProxy->Data->CastTo<const FSpriteSelectedVertex>(FSelectionTypes::Vertex))
-			{
-				SpriteGeometryHelper.ClearSelectionSet();
-				SpriteGeometryHelper.AddShapeToSelection(SelectedVertex->ShapeIndex);
-			}
-		}
-		else
-		{
-			//@TODO: This needs to be generalized!
-			if (const FSpriteSelectedEdge* SelectedEdge = SelectedItemProxy->Data->CastTo<const FSpriteSelectedEdge>(FSelectionTypes::Edge))
-			{
-				// Add the next vertex defined by this edge
-				SpriteGeometryHelper.AddPolygonEdgeToSelection(SelectedEdge->ShapeIndex, SelectedEdge->VertexIndex);
-			}
-			else if (const FSpriteSelectedVertex* SelectedVertex = SelectedItemProxy->Data->CastTo<const FSpriteSelectedVertex>(FSelectionTypes::Vertex))
-			{
-				SpriteGeometryHelper.AddPolygonVertexToSelection(SelectedVertex->ShapeIndex, SelectedVertex->VertexIndex);
-			}
-			else if (const FSpriteSelectedShape* SelectedShape = SelectedItemProxy->Data->CastTo<const FSpriteSelectedShape>(FSelectionTypes::GeometryShape))
-			{
-				SpriteGeometryHelper.AddShapeToSelection(SelectedShape->ShapeIndex);
-			}
-			else
-			{
-				SpriteGeometryHelper.SelectItem(SelectedItemProxy->Data);
-			}
-		}
-
-		bHandled = true;
-	}
-// 	else if (HWidgetUtilProxy* PivotProxy = HitProxyCast<HWidgetUtilProxy>(HitProxy))
-// 	{
-// 		//const bool bUserWantsPaint = bIsLeftButtonDown && ( !GetDefault<ULevelEditorViewportSettings>()->bLeftMouseDragMovesCamera ||  bIsCtrlDown );
-// 		//findme
-// 		WidgetAxis = WidgetProxy->Axis;
-// 
-// 			// Calculate the screen-space directions for this drag.
-// 			FSceneViewFamilyContext ViewFamily( FSceneViewFamily::ConstructionValues( Viewport, GetScene(), EngineShowFlags ));
-// 			FSceneView* View = CalcSceneView(&ViewFamily);
-// 			WidgetProxy->CalcVectors(View, FViewportClick(View, this, Key, Event, HitX, HitY), LocalManipulateDir, WorldManipulateDir, DragX, DragY);
-// 			bHandled = true;
-// 	}
-	else
+	if (IsInSourceRegionEditMode())
 	{
-		if (IsInSourceRegionEditMode())
+		if (Event == EInputEvent::IE_DoubleClick)
 		{
-			SpriteGeometryHelper.ClearSelectionSet();
-
-			if (Event == EInputEvent::IE_DoubleClick)
+			FVector4 WorldPoint = View.PixelToWorld(HitX, HitY, 0);
+			UPaperSprite* Sprite = GetSpriteBeingEdited();
+			FVector2D TexturePoint = SourceTextureViewComponent->GetSprite()->ConvertWorldSpaceToTextureSpace(WorldPoint);
+			if (bIsCtrlKeyDown)
 			{
-				FVector4 WorldPoint = View.PixelToWorld(HitX, HitY, 0);
-				UPaperSprite* Sprite = GetSpriteBeingEdited();
-				FVector2D TexturePoint = SourceTextureViewComponent->GetSprite()->ConvertWorldSpaceToTextureSpace(WorldPoint);
-				if (bIsCtrlKeyDown)
+				UPaperSprite* NewSprite = CreateNewSprite(Sprite->GetSourceUV(), Sprite->GetSourceSize());
+				if (NewSprite != nullptr)
 				{
-					UPaperSprite* NewSprite = CreateNewSprite(Sprite->GetSourceUV(), Sprite->GetSourceSize());
-					if (NewSprite != nullptr)
-					{
-						NewSprite->ExtractSourceRegionFromTexturePoint(TexturePoint);
-						bHandled = true;
-					}
-				}
-				else
-				{
-					Sprite->ExtractSourceRegionFromTexturePoint(TexturePoint);
+					NewSprite->ExtractSourceRegionFromTexturePoint(TexturePoint);
 					bHandled = true;
 				}
 			}
 			else
 			{
-				FVector4 WorldPoint = View.PixelToWorld(HitX, HitY, 0);
-				FVector2D TexturePoint = SourceTextureViewComponent->GetSprite()->ConvertWorldSpaceToTextureSpace(WorldPoint);
-				for (int32 RelatedSpriteIndex = 0; RelatedSpriteIndex < RelatedSprites.Num(); ++RelatedSpriteIndex)
-				{
-					FRelatedSprite& RelatedSprite = RelatedSprites[RelatedSpriteIndex];
-					if ((TexturePoint.X >= RelatedSprite.SourceUV.X) && (TexturePoint.Y >= RelatedSprite.SourceUV.Y) &&
-						(TexturePoint.X < (RelatedSprite.SourceUV.X + RelatedSprite.SourceDimension.X)) &&
-						(TexturePoint.Y < (RelatedSprite.SourceUV.Y + RelatedSprite.SourceDimension.Y)))
-					{
-						// Select this sprite
-						if (UPaperSprite* LoadedSprite = Cast<UPaperSprite>(RelatedSprite.AssetData.GetAsset()))
-						{
-							if (SpriteEditorPtr.IsValid())
-							{
-								SpriteEditorPtr.Pin()->SetSpriteBeingEdited(LoadedSprite);
-								break;
-							}
-						}
-						bHandled = true;
-					}
-				}
+				Sprite->ExtractSourceRegionFromTexturePoint(TexturePoint);
+				bHandled = true;
 			}
 		}
-		else if (IsEditingGeometry() && !SpriteGeometryHelper.IsAddingPolygon())
+		else
 		{
-			FSpriteGeometryCollection* Geometry = GetGeometryBeingEdited();
-
-			if (bInsertVertexModifier)
+			FVector4 WorldPoint = View.PixelToWorld(HitX, HitY, 0);
+			FVector2D TexturePoint = SourceTextureViewComponent->GetSprite()->ConvertWorldSpaceToTextureSpace(WorldPoint);
+			for (int32 RelatedSpriteIndex = 0; RelatedSpriteIndex < RelatedSprites.Num(); ++RelatedSpriteIndex)
 			{
-				FVector4 WorldPoint = View.PixelToWorld(HitX, HitY, 0);
-				if (UPaperSprite* Sprite = GetSpriteBeingEdited())
+				FRelatedSprite& RelatedSprite = RelatedSprites[RelatedSpriteIndex];
+				if ((TexturePoint.X >= RelatedSprite.SourceUV.X) && (TexturePoint.Y >= RelatedSprite.SourceUV.Y) &&
+					(TexturePoint.X < (RelatedSprite.SourceUV.X + RelatedSprite.SourceDimension.X)) &&
+					(TexturePoint.Y < (RelatedSprite.SourceUV.Y + RelatedSprite.SourceDimension.Y)))
 				{
-					FVector2D SpriteSpaceClickPoint = Sprite->ConvertWorldSpaceToTextureSpace(WorldPoint);
-
-					// find a polygon to add vert to
-					bool bFoundShapeToAddTo = false;
-					for (TSharedPtr<FSelectedItem> SelectedItemPtr : SpriteGeometryHelper.GetSelectionSet())
+					// Select this sprite
+					if (UPaperSprite* LoadedSprite = Cast<UPaperSprite>(RelatedSprite.AssetData.GetAsset()))
 					{
-						if (const FSpriteSelectedVertex* SelectedVertex = SelectedItemPtr->CastTo<const FSpriteSelectedVertex>(FSelectionTypes::Vertex)) //@TODO:  Inflexible?
+						if (SpriteEditorPtr.IsValid())
 						{
-							SpriteGeometryHelper.AddPointToGeometry(SpriteSpaceClickPoint, SelectedVertex->ShapeIndex);
-							bFoundShapeToAddTo = true;
+							SpriteEditorPtr.Pin()->SetSpriteBeingEdited(LoadedSprite);
 							break;
 						}
 					}
-
-					if (!bFoundShapeToAddTo)
-					{
-						SpriteGeometryHelper.AddPointToGeometry(SpriteSpaceClickPoint);
-					}
 					bHandled = true;
 				}
 			}
 		}
-		else if (!IsEditingGeometry())
-		{
-			// Clicked on the background (missed any proxies), deselect the socket or whatever was selected
-			SpriteGeometryHelper.ClearSelectionSet();
-		}
+	}
 
-		if (!bHandled)
-		{
-			FPaperEditorViewportClient::ProcessClick(View, HitProxy, Key, Event, HitX, HitY);
-		}
+	if (!bHandled)
+	{
+		FPaperEditorViewportClient::ProcessClick(View, HitProxy, Key, Event, HitX, HitY);
 	}
 }
 
@@ -1022,41 +891,6 @@ UPaperSprite* FSpriteEditorViewportClient::CreateNewSprite(FVector2D TopLeft, FV
 	return CreatedSprite;
 }
 
-bool FSpriteEditorViewportClient::ProcessMarquee(FViewport* Viewport, int32 ControllerId, FKey Key, EInputEvent Event, bool bMarqueeStartModifierPressed)
-{
-	bool bMarqueeReady = false;
-
-	if (Key == EKeys::LeftMouseButton)
-	{
-		int32 HitX = Viewport->GetMouseX();
-		int32 HitY = Viewport->GetMouseY();
-		if (Event == IE_Pressed && bMarqueeStartModifierPressed)
-		{
-			HHitProxy* HitResult = Viewport->GetHitProxy(HitX, HitY);
-			
-			if ((HitResult == nullptr) || (HitResult->Priority == HPP_World))
-			{
-				bIsMarqueeTracking = true;
-				MarqueeStartPos = FVector2D(HitX, HitY);
-				MarqueeEndPos = MarqueeStartPos;
-			}
-		}
-		else if (bIsMarqueeTracking && (Event == IE_Released))
-		{
-			MarqueeEndPos = FVector2D(HitX, HitY);
-			bIsMarqueeTracking = false;
-			bMarqueeReady = true;
-		}
-	}
-	else if (bIsMarqueeTracking && (Key == EKeys::Escape))
-	{
-		// Cancel marquee selection
-		bIsMarqueeTracking = false;
-	}
-
-	return bMarqueeReady;
-}
-
 bool FSpriteEditorViewportClient::InputKey(FViewport* Viewport, int32 ControllerId, FKey Key, EInputEvent Event, float AmountDepressed, bool bGamepad)
 {
 	bool bHandled = false;
@@ -1065,10 +899,14 @@ bool FSpriteEditorViewportClient::InputKey(FViewport* Viewport, int32 Controller
 	// Handle marquee tracking in source region edit mode
 	if (IsInSourceRegionEditMode())
 	{
-		bool bMarqueeStartModifier = InputState.IsCtrlButtonPressed();
-		if (ProcessMarquee(Viewport, ControllerId, Key, Event, bMarqueeStartModifier))
+		FSpriteGeometryEditMode* GeometryEditMode = ModeTools->GetActiveModeTyped<FSpriteGeometryEditMode>(FSpriteGeometryEditMode::EM_SpriteGeometry);
+		check(GeometryEditMode);
+
+		const bool bMarqueeStartModifier = InputState.IsCtrlButtonPressed();
+		if (GeometryEditMode->ProcessMarquee(Viewport, Key, Event, bMarqueeStartModifier))
 		{
-			FVector2D TextureSpaceStartPos, TextureSpaceDimensions;
+			FVector2D TextureSpaceStartPos;
+			FVector2D TextureSpaceDimensions;
 			if (ConvertMarqueeToSourceTextureSpace(/*out*/TextureSpaceStartPos, /*out*/TextureSpaceDimensions))
 			{
 				//@TODO: Warn if overlapping with another sprite
@@ -1076,92 +914,18 @@ bool FSpriteEditorViewportClient::InputKey(FViewport* Viewport, int32 Controller
 			}
 		}
 	}
-	else if (IsEditingGeometry())
-	{
-		if (SpriteGeometryHelper.IsAddingPolygon())
-		{
-			if (Key == EKeys::LeftMouseButton && Event == IE_Pressed)
-			{
-				const int32 HitX = Viewport->GetMouseX();
-				const int32 HitY = Viewport->GetMouseY();
-
-				// Calculate world space positions
-				FSceneViewFamilyContext ViewFamily(FSceneViewFamily::ConstructionValues(Viewport, GetScene(), EngineShowFlags));
-				FSceneView* View = CalcSceneView(&ViewFamily);
-				FVector WorldPoint = View->PixelToWorld(HitX, HitY, 0);
-
-				UPaperSprite* Sprite = GetSpriteBeingEdited();
-				FVector2D TexturePoint = Sprite->ConvertWorldSpaceToTextureSpace(WorldPoint);
-
-				const bool bMakeSubtractiveIfAllowed = Viewport->KeyState(EKeys::LeftControl) || Viewport->KeyState(EKeys::RightControl);
-				SpriteGeometryHelper.TEMP_HandleAddPolygonClick(TexturePoint, bMakeSubtractiveIfAllowed);
-			}
-			else if (Key == EKeys::Enter)
-			{
-				SpriteGeometryHelper.ResetAddPolygonMode();
-			}
-			else if (Key == EKeys::Escape)
-			{
-				SpriteGeometryHelper.AbandonAddPolygonMode();
-			}
-		}
-		else
-		{
-			if (ProcessMarquee(Viewport, ControllerId, Key, Event, true))
-			{
-				const bool bAddingToSelection = InputState.IsShiftButtonPressed(); //@TODO: control button moves widget? Hopefully make this more consistent when that is changed
-				SelectVerticesInMarquee(bAddingToSelection);
-			}
-		}
-	}
-
-	// Start the drag
-	//@TODO: EKeys::LeftMouseButton
-	//@TODO: Event.IE_Pressed
-	// Implement InputAxis
-	// StartTracking
-
+	
 	// Pass keys to standard controls, if we didn't consume input
-	return (bHandled) ? true : FEditorViewportClient::InputKey(Viewport,  ControllerId, Key, Event, AmountDepressed, bGamepad);
-}
-
-bool FSpriteEditorViewportClient::InputWidgetDelta(FViewport* Viewport, EAxisList::Type CurrentAxis, FVector& Drag, FRotator& Rot, FVector& Scale)
-{
-	bool bHandled = false;
-	if (bManipulating && (CurrentAxis != EAxisList::None))
-	{
-		bHandled = true;
-
-		const FWidget::EWidgetMode MoveMode = GetWidgetMode();
-
-		// Negate Y because vertices are in source texture space, not world space
-		const FVector2D Drag2D(FVector::DotProduct(Drag, PaperAxisX), -FVector::DotProduct(Drag, PaperAxisY));
-
-		// Apply the delta to all of the selected objects
-		for (TSharedPtr<FSelectedItem> SelectedItem : SpriteGeometryHelper.GetSelectionSet())
-		{
-			SelectedItem->ApplyDelta(Drag2D, Rot, Scale, MoveMode);
-		}
-
-		if (SpriteGeometryHelper.HasAnySelectedItems())
-		{
-			if (!IsInSourceRegionEditMode())
-			{
-				UPaperSprite* Sprite = GetSpriteBeingEdited();
-				Sprite->PostEditChange();
-				Invalidate();
-			}
-
-			bManipulationDirtiedSomething = true;
-		}
-	}
-
-	return bHandled;
+	return (bHandled) ? true : FEditorViewportClient::InputKey(Viewport, ControllerId, Key, Event, AmountDepressed, bGamepad);
 }
 
 void FSpriteEditorViewportClient::TrackingStarted(const struct FInputEventState& InInputState, bool bIsDragging, bool bNudge)
 {
-	if (!bManipulating && bIsDragging)
+	//@TODO: Should push this into FEditorViewportClient
+	// Begin transacting.  Give the current editor mode an opportunity to do the transacting.
+	const bool bTrackingHandledExternally = ModeTools->StartTracking(this, Viewport);
+
+	if (!bManipulating && bIsDragging && !bTrackingHandledExternally)
 	{
 		BeginTransaction(LOCTEXT("ModificationInViewport", "Modification in Viewport"));
 		bManipulating = true;
@@ -1171,64 +935,16 @@ void FSpriteEditorViewportClient::TrackingStarted(const struct FInputEventState&
 
 void FSpriteEditorViewportClient::TrackingStopped() 
 {
-	if (bManipulating)
+	// Stop transacting.  Give the current editor mode an opportunity to do the transacting.
+	const bool bTransactingHandledByEditorMode = ModeTools->EndTracking(this, Viewport);
+
+	if (bManipulating && !bTransactingHandledByEditorMode)
 	{
 		EndTransaction();
 		bManipulating = false;
 	}
 }
 
-FWidget::EWidgetMode FSpriteEditorViewportClient::GetWidgetMode() const
-{
-	return SpriteGeometryHelper.HasAnySelectedItems() ? DesiredWidgetMode : FWidget::WM_None;
-}
-
-void FSpriteEditorViewportClient::SetWidgetMode(FWidget::EWidgetMode NewMode)
-{
-	DesiredWidgetMode = NewMode;
-}
-
-bool FSpriteEditorViewportClient::CanSetWidgetMode(FWidget::EWidgetMode NewMode) const
-{
-	return CanCycleWidgetMode();
-}
-
-bool FSpriteEditorViewportClient::CanCycleWidgetMode() const
-{
-	if (!Widget->IsDragging())
-	{
-		return SpriteGeometryHelper.HasAnySelectedItems();
-	}
-	return false;
-}
-
-FVector FSpriteEditorViewportClient::GetWidgetLocation() const
-{
-	FVector SummedPos(ForceInitToZero);
-
-	if (SpriteGeometryHelper.HasAnySelectedItems())
-	{
-		// Find the center of the selection set
-		for (TSharedPtr<FSelectedItem> SelectedItem : SpriteGeometryHelper.GetSelectionSet())
-		{
-			SummedPos += SelectedItem->GetWorldPos();
-		}
-		return (SummedPos / SpriteGeometryHelper.GetSelectionSet().Num());
-	}
-
-	return SummedPos;
-}
-
-FMatrix FSpriteEditorViewportClient::GetWidgetCoordSystem() const
-{
-	return FMatrix::Identity;
-}
-
-ECoordSystem FSpriteEditorViewportClient::GetWidgetCoordSystemSpace() const
-{
-	return COORD_World;
-}
-
 FLinearColor FSpriteEditorViewportClient::GetBackgroundColor() const
 {
 	return FEditorViewportClient::GetBackgroundColor();
@@ -1240,15 +956,16 @@ FVector2D FSpriteEditorViewportClient::SelectedItemConvertWorldSpaceDeltaToLocal
 	return Sprite->ConvertWorldSpaceDeltaToTextureSpace(WorldSpaceDelta);
 }
 
-FVector FSpriteEditorViewportClient::TextureSpaceToWorldSpace(const FVector2D& SourcePoint) const
+FVector2D FSpriteEditorViewportClient::WorldSpaceToTextureSpace(const FVector& SourcePoint) const
 {
 	UPaperSprite* Sprite = GetSpriteBeingEdited();
-	return Sprite->ConvertTextureSpaceToWorldSpace(SourcePoint);
+	return Sprite->ConvertWorldSpaceToTextureSpace(SourcePoint);
 }
 
-bool FSpriteEditorViewportClient::SelectedItemIsSelected(const struct FShapeVertexPair& Item) const
+FVector FSpriteEditorViewportClient::TextureSpaceToWorldSpace(const FVector2D& SourcePoint) const
 {
-	return SpriteGeometryHelper.IsGeometrySelected(Item);
+	UPaperSprite* Sprite = GetSpriteBeingEdited();
+	return Sprite->ConvertTextureSpaceToWorldSpace(SourcePoint);
 }
 
 float FSpriteEditorViewportClient::SelectedItemGetUnitsPerPixel() const
@@ -1271,6 +988,9 @@ void FSpriteEditorViewportClient::BeginTransaction(const FText& SessionName)
 void FSpriteEditorViewportClient::MarkTransactionAsDirty()
 {
 	bManipulationDirtiedSomething = true;
+	Invalidate();
+	//@TODO: Can add a call to Sprite->PostEditChange here if we want to update the baked sprite data during a drag operation
+	// (maybe passing in Interactive - if so, the EndTransaction PostEditChange needs to be a ValueSet)
 }
 
 void FSpriteEditorViewportClient::EndTransaction()
@@ -1309,7 +1029,6 @@ void FSpriteEditorViewportClient::NotifySpriteBeingEditedHasChanged()
 {
 	//@TODO: Ideally we do this before switching
 	EndTransaction();
-	SpriteGeometryHelper.ClearSelectionSet();
 
 	// Update components to know about the new sprite being edited
 	UPaperSprite* Sprite = GetSpriteBeingEdited();
@@ -1323,93 +1042,19 @@ void FSpriteEditorViewportClient::NotifySpriteBeingEditedHasChanged()
 	RequestFocusOnSelection(/*bInstant=*/ true);
 }
 
-void FSpriteEditorViewportClient::FocusOnSprite()
-{
-	FocusViewportOnBox(RenderSpriteComponent->Bounds.GetBox());
-}
-
-void FSpriteEditorViewportClient::DeleteSelection()
-{
-	SpriteGeometryHelper.DeleteSelectedItems();
-}
-
-void FSpriteEditorViewportClient::AddBoxShape()
-{
-	UPaperSprite* Sprite = GetSpriteBeingEdited();
-	check(Sprite);
-
-	const FVector2D BoxSize(Sprite->GetSourceSize());
-	const FVector2D BoxLocation(Sprite->GetSourceUV() + (BoxSize * 0.5f));
-
-	SpriteGeometryHelper.AddNewBoxShape(BoxLocation, BoxSize);
-}
-
-void FSpriteEditorViewportClient::AddCircleShape()
-{
-	check(IsInCollisionEditMode());
-
-	UPaperSprite* Sprite = GetSpriteBeingEdited();
-	check(Sprite);
-
-	const FVector2D SpriteBounds = Sprite->GetSourceSize();
-	const float SmallerBoundingAxisSize = SpriteBounds.GetMin();
-	const float CircleRadius = SmallerBoundingAxisSize * 0.5f;
-	const FVector2D EllipseLocation = Sprite->GetSourceUV() + (SpriteBounds * 0.5f);
-
-	SpriteGeometryHelper.AddNewCircleShape(EllipseLocation, CircleRadius);
-}
-
-void FSpriteEditorViewportClient::SnapAllVerticesToPixelGrid()
-{
-	if (FSpriteGeometryCollection* Geometry = GetGeometryBeingEdited())
-	{
-		BeginTransaction(LOCTEXT("SnapAllVertsToPixelGridTransaction", "Snap All Verts to Pixel Grid"));
-
-		for (FSpriteGeometryShape& Shape : Geometry->Shapes)
-		{
-			bManipulationDirtiedSomething = true;
-
-			Shape.BoxPosition.X = FMath::RoundToInt(Shape.BoxPosition.X);
-			Shape.BoxPosition.Y = FMath::RoundToInt(Shape.BoxPosition.Y);
-			
-			if ((Shape.ShapeType == ESpriteShapeType::Box) || (Shape.ShapeType == ESpriteShapeType::Circle))
-			{
-				//@TODO: Should we snap BoxPosition also, or just the verts?
-				const FVector2D OldHalfSize = Shape.BoxSize * 0.5f;
-				FVector2D TopLeft = Shape.BoxPosition - OldHalfSize;
-				FVector2D BottomRight = Shape.BoxPosition + OldHalfSize;
-				TopLeft.X = FMath::RoundToInt(TopLeft.X);
-				TopLeft.Y = FMath::RoundToInt(TopLeft.Y);
-				BottomRight.X = FMath::RoundToInt(BottomRight.X);
-				BottomRight.Y = FMath::RoundToInt(BottomRight.Y);
-				Shape.BoxPosition = (TopLeft + BottomRight) * 0.5f;
-				Shape.BoxSize = BottomRight - TopLeft;
-			}
-
-			for (FVector2D& Vertex : Shape.Vertices)
-			{
-				FVector2D TextureSpaceVertex = Shape.ConvertShapeSpaceToTextureSpace(Vertex);
-				TextureSpaceVertex.X = FMath::RoundToInt(TextureSpaceVertex.X);
-				TextureSpaceVertex.Y = FMath::RoundToInt(TextureSpaceVertex.Y);
-				Vertex = Shape.ConvertTextureSpaceToShapeSpace(TextureSpaceVertex);
-			}
-		}
-
-		EndTransaction();
-	}
-}
-
 void FSpriteEditorViewportClient::InternalActivateNewMode(ESpriteEditorMode::Type NewMode)
 {
 	CurrentMode = NewMode;
 	Viewport->InvalidateHitProxy();
-	ResetMarqueeTracking();
 
 	UPaperSprite* Sprite = GetSpriteBeingEdited();
 
+	FSpriteGeometryEditMode* GeometryEditMode = ModeTools->GetActiveModeTyped<FSpriteGeometryEditMode>(FSpriteGeometryEditMode::EM_SpriteGeometry);
+	check(GeometryEditMode);
+
 	// Note: This has side effects (clearing the selection set, ensuring the geometry is correct if the sprite being edited changed, etc...).
 	// Do not skip even if the mode is not really changing.
-	SpriteGeometryHelper.SetGeometryBeingEdited(nullptr, /*bAllowCircles=*/ false, /*bAllowSubtractivePolygons=*/ false);
+	GeometryEditMode->SetGeometryBeingEdited(nullptr, /*bAllowCircles=*/ false, /*bAllowSubtractivePolygons=*/ false);
 
 	switch (CurrentMode)
 	{
@@ -1419,51 +1064,29 @@ void FSpriteEditorViewportClient::InternalActivateNewMode(ESpriteEditorMode::Typ
 		UpdateRelatedSpritesList();
 		break;
 	case ESpriteEditorMode::EditCollisionMode:
+		GeometryEditMode->SetGeometryColors(SpriteEditingConstants::CollisionShapeColor, FLinearColor::White);
 		if (Sprite != nullptr)
 		{
-			SpriteGeometryHelper.SetGeometryBeingEdited(&(Sprite->CollisionGeometry), /*bAllowCircles=*/ true, /*bAllowSubtractivePolygons=*/ false);
+			GeometryEditMode->SetGeometryBeingEdited(&(Sprite->CollisionGeometry), /*bAllowCircles=*/ true, /*bAllowSubtractivePolygons=*/ false);
 		}
 		break;
 	case ESpriteEditorMode::EditRenderingGeomMode:
+		GeometryEditMode->SetGeometryColors(SpriteEditingConstants::RenderShapeColor, SpriteEditingConstants::SubtractiveRenderShapeColor);
 		if (Sprite != nullptr)
 		{
-			SpriteGeometryHelper.SetGeometryBeingEdited(&(Sprite->RenderGeometry), /*bAllowCircles=*/ false, /*bAllowSubtractivePolygons=*/ true);
+			GeometryEditMode->SetGeometryBeingEdited(&(Sprite->RenderGeometry), /*bAllowCircles=*/ false, /*bAllowSubtractivePolygons=*/ true);
 		}
 		break;
 	}
 }
 
-FSpriteGeometryCollection* FSpriteEditorViewportClient::GetGeometryBeingEdited() const
-{
-	UPaperSprite* Sprite = GetSpriteBeingEdited();
-	switch (CurrentMode)
-	{
-	case ESpriteEditorMode::EditCollisionMode:
-		return &(Sprite->CollisionGeometry);
-	case ESpriteEditorMode::EditRenderingGeomMode:
-		return &(Sprite->RenderGeometry);
-	default:
-		return nullptr;
-	}
-}
-
-bool FSpriteEditorViewportClient::IsShapeSelected(const int32 ShapeIndex) const
-{
-	return SpriteGeometryHelper.IsGeometrySelected(FShapeVertexPair(ShapeIndex, INDEX_NONE));
-}
-
-bool FSpriteEditorViewportClient::IsPolygonVertexSelected(const int32 ShapeIndex, const int32 VertexIndex) const
-{
-	return SpriteGeometryHelper.IsGeometrySelected(FShapeVertexPair(ShapeIndex, VertexIndex));
-}
-
-void FSpriteEditorViewportClient::ResetMarqueeTracking()
-{
-	bIsMarqueeTracking = false;
-}
-
 bool FSpriteEditorViewportClient::ConvertMarqueeToSourceTextureSpace(/*out*/FVector2D& OutStartPos, /*out*/FVector2D& OutDimension)
 {
+	FSpriteGeometryEditMode* GeometryEditMode = ModeTools->GetActiveModeTyped<FSpriteGeometryEditMode>(FSpriteGeometryEditMode::EM_SpriteGeometry);
+	check(GeometryEditMode);
+	const FVector2D MarqueeStartPos = GeometryEditMode->GetMarqueeStartPos();
+	const FVector2D MarqueeEndPos = GeometryEditMode->GetMarqueeEndPos();
+
 	bool bSuccessful = false;
 	UPaperSprite* Sprite = SourceTextureViewComponent->GetSprite();
 	UTexture2D* SpriteSourceTexture = Sprite->GetSourceTexture();
@@ -1507,77 +1130,6 @@ bool FSpriteEditorViewportClient::ConvertMarqueeToSourceTextureSpace(/*out*/FVec
 	return bSuccessful;
 }
 
-void FSpriteEditorViewportClient::SelectVerticesInMarquee(bool bAddToSelection)
-{
-	if (!bAddToSelection)
-	{
-		SpriteGeometryHelper.ClearSelectionSet();
-	}
-
-	if (UPaperSprite* Sprite = GetSpriteBeingEdited())
-	{
-		// Calculate world space positions
-		FSceneViewFamilyContext ViewFamily(FSceneViewFamily::ConstructionValues(Viewport, GetScene(), EngineShowFlags));
-		FSceneView* View = CalcSceneView(&ViewFamily);
-		const FVector StartPos = View->PixelToWorld(MarqueeStartPos.X, MarqueeStartPos.Y, 0);
-		const FVector EndPos = View->PixelToWorld(MarqueeEndPos.X, MarqueeEndPos.Y, 0);
-
-		// Convert to source texture space to work out the pixels dragged
-		FVector2D TextureSpaceStartPos = Sprite->ConvertWorldSpaceToTextureSpace(StartPos);
-		FVector2D TextureSpaceEndPos = Sprite->ConvertWorldSpaceToTextureSpace(EndPos);
-
-		if (TextureSpaceStartPos.X > TextureSpaceEndPos.X)
-		{
-			Swap(TextureSpaceStartPos.X, TextureSpaceEndPos.X);
-		}
-		if (TextureSpaceStartPos.Y > TextureSpaceEndPos.Y)
-		{
-			Swap(TextureSpaceStartPos.Y, TextureSpaceEndPos.Y);
-		}
-
-		const FBox2D QueryBounds(TextureSpaceStartPos, TextureSpaceEndPos);
-
-		if (FSpriteGeometryCollection* Geometry = GetGeometryBeingEdited())
-		{
-			for (int32 ShapeIndex = 0; ShapeIndex < Geometry->Shapes.Num(); ++ShapeIndex)
-			{
-				const FSpriteGeometryShape& Shape = Geometry->Shapes[ShapeIndex];
-
-				bool bSelectWholeShape = false;
-
-				if ((Shape.ShapeType == ESpriteShapeType::Circle) || (Shape.ShapeType == ESpriteShapeType::Box))
-				{
-					// First see if we are fully contained
-					const FBox2D ShapeBoxBounds(Shape.BoxPosition - Shape.BoxSize * 0.5f, Shape.BoxPosition + Shape.BoxSize * 0.5f);
-					if (QueryBounds.IsInside(ShapeBoxBounds))
-					{
-						bSelectWholeShape = true;
-					}
-				}
-
-				//@TODO: Try intersecting with the circle if it wasn't entirely enclosed
-				
-				if (bSelectWholeShape)
-				{
-					SpriteGeometryHelper.AddShapeToSelection(ShapeIndex);
-				}
-				else
-				{
-					// Try to select some subset of the vertices
-					for (int32 VertexIndex = 0; VertexIndex < Shape.Vertices.Num(); ++VertexIndex)
-					{
-						const FVector2D TextureSpaceVertex = Shape.ConvertShapeSpaceToTextureSpace(Shape.Vertices[VertexIndex]);
-						if (QueryBounds.IsInside(TextureSpaceVertex))
-						{
-							SpriteGeometryHelper.AddPolygonVertexToSelection(ShapeIndex, VertexIndex);
-						}
-					}
-				}
-			}
-		}
-	}
-}
-
 //////////////////////////////////////////////////////////////////////////
 
 #undef LOCTEXT_NAMESPACE
\ No newline at end of file
diff --git a/Engine/Plugins/2D/Paper2D/Source/Paper2DEditor/Private/SpriteEditor/SpriteEditorViewportClient.h b/Engine/Plugins/2D/Paper2D/Source/Paper2DEditor/Private/SpriteEditor/SpriteEditorViewportClient.h
--- a/Engine/Plugins/2D/Paper2D/Source/Paper2DEditor/Private/SpriteEditor/SpriteEditorViewportClient.h
+++ b/Engine/Plugins/2D/Paper2D/Source/Paper2DEditor/Private/SpriteEditor/SpriteEditorViewportClient.h
@@ -30,7 +30,7 @@ class FSpriteEditorViewportClient : public FPaperEditorViewportClient, public IS
 {
 public:
 	/** Constructor */
-	FSpriteEditorViewportClient(TWeakPtr<FSpriteEditor> InSpriteEditor, TWeakPtr<class SSpriteEditorViewport> InSpriteEditorViewportPtr);
+	FSpriteEditorViewportClient(TWeakPtr<FSpriteEditor> InSpriteEditor, TWeakPtr<class SEditorViewport> InSpriteEditorViewportPtr);
 
 	// FViewportClient interface
 	virtual void Draw(FViewport* Viewport, FCanvas* Canvas) override;
@@ -42,23 +42,15 @@ public:
 	// FEditorViewportClient interface
 	virtual void ProcessClick(FSceneView& View, HHitProxy* HitProxy, FKey Key, EInputEvent Event, uint32 HitX, uint32 HitY) override;
 	virtual bool InputKey(FViewport* Viewport, int32 ControllerId, FKey Key, EInputEvent Event, float AmountDepressed, bool bGamepad) override;
-	virtual bool InputWidgetDelta(FViewport* Viewport, EAxisList::Type CurrentAxis, FVector& Drag, FRotator& Rot, FVector& Scale) override;
 	virtual void TrackingStarted(const struct FInputEventState& InInputState, bool bIsDragging, bool bNudge) override;
 	virtual void TrackingStopped() override;
-	virtual FWidget::EWidgetMode GetWidgetMode() const override;
-	virtual void SetWidgetMode(FWidget::EWidgetMode NewMode) override;
-	virtual bool CanSetWidgetMode(FWidget::EWidgetMode NewMode) const override;
-	virtual bool CanCycleWidgetMode() const override;
-	virtual FVector GetWidgetLocation() const override;
-	virtual FMatrix GetWidgetCoordSystem() const override;
-	virtual ECoordSystem GetWidgetCoordSystemSpace() const override;
 	virtual FLinearColor GetBackgroundColor() const override;
 	// End of FEditorViewportClient interface
 
 	// ISpriteSelectionContext interface
 	virtual FVector2D SelectedItemConvertWorldSpaceDeltaToLocalSpace(const FVector& WorldSpaceDelta) const override;
+	virtual FVector2D WorldSpaceToTextureSpace(const FVector& SourcePoint) const override;
 	virtual FVector TextureSpaceToWorldSpace(const FVector2D& SourcePoint) const override;
-	virtual bool SelectedItemIsSelected(const struct FShapeVertexPair& Item) const override;
 	virtual float SelectedItemGetUnitsPerPixel() const override;
 	virtual void BeginTransaction(const FText& SessionName) override;
 	virtual void MarkTransactionAsDirty() override;
@@ -66,8 +58,7 @@ public:
 	virtual void InvalidateViewportAndHitProxies() override;
 	// End of ISpriteSelectionContext interface
 
-	// Process marquee selection, return true of a selection has been performed
-	bool ProcessMarquee(FViewport* Viewport, int32 ControllerId, FKey Key, EInputEvent Event, bool bMarqueeStartModifierPressed);
+	void ActivateEditMode();
 
 	void ToggleShowSockets() { bShowSockets = !bShowSockets; Invalidate(); }
 	bool IsShowSocketsChecked() const { return bShowSockets; }
@@ -95,18 +86,6 @@ public:
 	bool IsShowSourceTextureChecked() const { return bShowSourceTexture; }
 	bool CanShowSourceTexture() const { return !IsInSourceRegionEditMode(); }
 
-	void FocusOnSprite();
-
-	// Geometry editing commands
-	void DeleteSelection();
-	bool CanDeleteSelection() const { return IsEditingGeometry(); } //@TODO: Need a selection
-
-	void AddBoxShape();
-	void AddCircleShape();
-
-	void SnapAllVerticesToPixelGrid();
-	bool CanSnapVerticesToPixelGrid() const { return IsEditingGeometry(); }
-
 	// Invalidate any references to the sprite being edited; it has changed
 	void NotifySpriteBeingEditedHasChanged();
 
@@ -143,7 +122,7 @@ private:
 	UPaperSpriteComponent* RenderSpriteComponent;
 
 	// Widget mode
-	FWidget::EWidgetMode DesiredWidgetMode;
+// 	FWidget::EWidgetMode DesiredWidgetMode;
 
 	// Are we currently manipulating something?
 	bool bManipulating;
@@ -152,7 +131,7 @@ private:
 	bool bManipulationDirtiedSomething;
 
 	// Pointer back to the sprite editor viewport control that owns us
-	TWeakPtr<class SSpriteEditorViewport> SpriteEditorViewportPtr;
+	TWeakPtr<class SEditorViewport> SpriteEditorViewportPtr;
 
 	// The current transaction for undo/redo
 	class FScopedTransaction* ScopedTransaction;
@@ -169,22 +148,9 @@ private:
 	// Should we show related sprites in the source texture?
 	bool bShowRelatedSprites;
 
-	// Marquee tracking
-	bool bIsMarqueeTracking;
-	FVector2D MarqueeStartPos, MarqueeEndPos;
-
 	// Other sprites that share the same source texture
 	TArray<FRelatedSprite> RelatedSprites;
 
-	// Sprite geometry editing/rendering helper
-	FSpriteGeometryEditingHelper SpriteGeometryHelper;
-
-public:
-	FSpriteGeometryEditingHelper* GetGeometryEditHelper()
-	{
-		return &SpriteGeometryHelper;
-	}
-
 private:
 	UPaperSprite* GetSpriteBeingEdited() const
 	{
@@ -202,25 +168,11 @@ private:
 	void DrawRenderStats(FViewport& InViewport, FSceneView& View, FCanvas& Canvas, class UPaperSprite* Sprite, int32& YPos);
 	void DrawSourceRegion(FViewport& InViewport, FSceneView& View, FCanvas& Canvas, const FLinearColor& GeometryVertexColor);
 	void DrawRelatedSprites(FViewport& InViewport, FSceneView& View, FCanvas& Canvas, const FLinearColor& GeometryVertexColor);
-	void DrawMarquee(FViewport& InViewport, FSceneView& View, FCanvas& Canvas, const FLinearColor& MarqueeColor);
 
 	void UpdateSourceTextureSpriteFromSprite(UPaperSprite* SourceSprite);
 	
-	// Selection handling
-
-	// Indicates if the specified shape is currently selected
-	bool IsShapeSelected(const int32 ShapeIndex) const;
-
-	// Indicates if the specified vertex is selected
-	bool IsPolygonVertexSelected(const int32 PolygonIndex, const int32 VertexIndex) const;
-
-	void ResetMarqueeTracking();
 	bool ConvertMarqueeToSourceTextureSpace(/*out*/FVector2D& OutStartPos, /*out*/FVector2D& OutDimension);
-	void SelectVerticesInMarquee(bool bAddToSelection);
 
 	// Activates a new mode, clearing selection set, etc...
 	void InternalActivateNewMode(ESpriteEditorMode::Type NewMode);
-
-	// Can return null
-	FSpriteGeometryCollection* GetGeometryBeingEdited() const;
 };
