Tile Set Editor: Introduced the geometry edit mode to the single tile editor
- Added the geometry edit mode
- Added a viewport toolbar and buttons to the main toolbar

diff --git a/Engine/Plugins/2D/Paper2D/Source/Paper2DEditor/Private/TileSetEditor.cpp b/Engine/Plugins/2D/Paper2D/Source/Paper2DEditor/Private/TileSetEditor.cpp
--- a/Engine/Plugins/2D/Paper2D/Source/Paper2DEditor/Private/TileSetEditor.cpp
+++ b/Engine/Plugins/2D/Paper2D/Source/Paper2DEditor/Private/TileSetEditor.cpp
@@ -303,87 +303,341 @@ public:
 };
 
 
-//////////////////////////////////////////////////////////////////////////
-// FSingleTileEditorViewportClient
 
-class FSingleTileEditorViewportClient : public FPaperEditorViewportClient
+
+
+
+
+
+
+
+
+
+
+
+
+
+#include "Editor/UnrealEd/Public/SCommonEditorViewportToolbarBase.h"
+
+// In-viewport toolbar widget used in the tile set editor
+class STileSetEditorViewportToolbar : public SCommonEditorViewportToolbarBase
 {
 public:
-	FSingleTileEditorViewportClient(UPaperTileSet* InTileSet)
-		: TileSet(InTileSet)
-		, TileBeingEditedIndex(INDEX_NONE)
-		, PixelSize(4.0f)
+	SLATE_BEGIN_ARGS(STileSetEditorViewportToolbar) {}
+	SLATE_END_ARGS()
+
+	void Construct(const FArguments& InArgs, TSharedPtr<class ICommonEditorViewportToolbarInfoProvider> InInfoProvider);
+
+	// SCommonEditorViewportToolbarBase interface
+	virtual TSharedRef<SWidget> GenerateShowMenu() const override;
+	// End of SCommonEditorViewportToolbarBase
+};
+
+///////////////////////////////////////////////////////////
+// STileSetEditorViewportToolbar
+
+#include "SEditorViewport.h"
+#include "SpriteEditor/SpriteEditorCommands.h"
+
+void STileSetEditorViewportToolbar::Construct(const FArguments& InArgs, TSharedPtr<class ICommonEditorViewportToolbarInfoProvider> InInfoProvider)
+{
+	SCommonEditorViewportToolbarBase::Construct(SCommonEditorViewportToolbarBase::FArguments(), InInfoProvider);
+}
+
+TSharedRef<SWidget> STileSetEditorViewportToolbar::GenerateShowMenu() const
+{
+	GetInfoProvider().OnFloatingButtonClicked();
+
+	TSharedRef<SEditorViewport> ViewportRef = GetInfoProvider().GetViewportWidget();
+
+	const bool bInShouldCloseWindowAfterMenuSelection = true;
+	FMenuBuilder ShowMenuBuilder(bInShouldCloseWindowAfterMenuSelection, ViewportRef->GetCommandList());
 	{
-		ZoomPos = FVector2D(-TileSet->TileWidth * ZoomAmount * PixelSize * 0.5f, -TileSet->TileHeight * ZoomAmount * PixelSize * 0.5f);
+		ShowMenuBuilder.AddMenuEntry(FSpriteGeometryEditCommands::Get().SetShowNormals);
 	}
 
+	return ShowMenuBuilder.MakeWidget();
+}
+
+
+//////////////////////////////////////////////////////////////////////////
+// FSingleTileEditorViewportClient
+
+#include "PreviewScene.h"
+#include "SpriteEditor/SpriteEditorSelections.h"
+#include "SpriteEditor/SpriteEditorCommands.h"
+
+class FSingleTileEditorViewportClient : public FPaperEditorViewportClient, public ISpriteSelectionContext
+{
+public:
+	FSingleTileEditorViewportClient(UPaperTileSet* InTileSet);
+
 	// FViewportClient interface
+	virtual void Tick(float DeltaSeconds) override;
 	virtual void Draw(FViewport* Viewport, FCanvas* Canvas) override;
 	// End of FViewportClient interface
 
 	// FEditorViewportClient interface
 	virtual FLinearColor GetBackgroundColor() const override;
+	virtual void TrackingStarted(const struct FInputEventState& InInputState, bool bIsDragging, bool bNudge) override;
+	virtual void TrackingStopped() override;
 	// End of FEditorViewportClient interface
 
+	// ISpriteSelectionContext interface
+	virtual FVector2D SelectedItemConvertWorldSpaceDeltaToLocalSpace(const FVector& WorldSpaceDelta) const override;
+	virtual FVector2D WorldSpaceToTextureSpace(const FVector& SourcePoint) const override;
+	virtual FVector TextureSpaceToWorldSpace(const FVector2D& SourcePoint) const override;
+	virtual float SelectedItemGetUnitsPerPixel() const override;
+	virtual void BeginTransaction(const FText& SessionName) override;
+	virtual void MarkTransactionAsDirty() override;
+	virtual void EndTransaction() override;
+	virtual void InvalidateViewportAndHitProxies() override;
+	// End of ISpriteSelectionContext interface
+
 	void SetTileIndex(int32 InTileIndex);
 	int32 GetTileIndex() const;
 	void OnActiveTileIndexChanged(const FIntPoint& TopLeft, const FIntPoint& Dimensions);
+
+	void ActivateEditMode(TSharedPtr<FUICommandList> InCommandList);
+
+protected:
+	// FPaperEditorViewportClient interface
+	virtual FBox GetDesiredFocusBounds() const override;
+	// End of FPaperEditorViewportClient
+
 public:
 	// Tile set
 	UPaperTileSet* TileSet;
 
 	int32 TileBeingEditedIndex;
 
-	float PixelSize;
+	// Are we currently manipulating something?
+	bool bManipulating;
+
+	// Did we dirty something during manipulation?
+	bool bManipulationDirtiedSomething;
+
+	// Pointer back to the sprite editor viewport control that owns us
+	TWeakPtr<class SEditorViewport> SpriteEditorViewportPtr;
+
+	// The current transaction for undo/redo
+	class FScopedTransaction* ScopedTransaction;
+
+	// The preview scene
+	FPreviewScene OwnedPreviewScene;
+
+	// The preview sprite in the scene
+	UPaperSpriteComponent* PreviewTileSpriteComponent;
 };
 
+#include "PaperEditorShared/SpriteGeometryEditMode.h"
+#include "ScopedTransaction.h"
+
+FSingleTileEditorViewportClient::FSingleTileEditorViewportClient(UPaperTileSet* InTileSet)
+	: TileSet(InTileSet)
+	, TileBeingEditedIndex(INDEX_NONE)
+	, bManipulating(false)
+	, bManipulationDirtiedSomething(false)
+	, ScopedTransaction(nullptr)
+{
+	//@TODO: Should be able to set this to false eventually
+	SetRealtime(true);
+
+	// The tile map editor fully supports mode tools and isn't doing any incompatible stuff with the Widget
+	Widget->SetUsesEditorModeTools(ModeTools);
+
+	DrawHelper.bDrawGrid = true;//@TODO:
+	DrawHelper.bDrawPivot = false;
+
+	PreviewScene = &OwnedPreviewScene;
+	((FAssetEditorModeManager*)ModeTools)->SetPreviewScene(PreviewScene);
+
+	EngineShowFlags.DisableAdvancedFeatures();
+	EngineShowFlags.CompositeEditorPrimitives = true;
+
+	// Create a render component for the tile preview
+	PreviewTileSpriteComponent = NewObject<UPaperSpriteComponent>();
+
+	FStringAssetReference DefaultTranslucentMaterialName("/Paper2D/TranslucentUnlitSpriteMaterial.TranslucentUnlitSpriteMaterial");
+	UMaterialInterface* TranslucentMaterial = Cast<UMaterialInterface>(DefaultTranslucentMaterialName.TryLoad());
+	PreviewTileSpriteComponent->SetMaterial(0, TranslucentMaterial);
+
+	PreviewScene->AddComponent(PreviewTileSpriteComponent, FTransform::Identity);
+}
+
+FBox FSingleTileEditorViewportClient::GetDesiredFocusBounds() const
+{
+	UPaperSpriteComponent* ComponentToFocusOn = PreviewTileSpriteComponent;
+	return ComponentToFocusOn->Bounds.GetBox();
+}
+
+
+void FSingleTileEditorViewportClient::Tick(float DeltaSeconds)
+{
+	FPaperEditorViewportClient::Tick(DeltaSeconds);
+
+	if (!GIntraFrameDebuggingGameThread)
+	{
+		OwnedPreviewScene.GetWorld()->Tick(LEVELTICK_All, DeltaSeconds);
+	}
+}
+
 void FSingleTileEditorViewportClient::Draw(FViewport* Viewport, FCanvas* Canvas)
 {
-	// Super will clear the viewport
-	FPaperEditorViewportClient::Draw(Viewport, Canvas);
+	// Skipping the parent on purpose
+	FEditorViewportClient::Draw(Viewport, Canvas);
+}
+
+FLinearColor FSingleTileEditorViewportClient::GetBackgroundColor() const
+{
+	return TileSet->BackgroundColor;
+}
 
-	if (TileBeingEditedIndex == INDEX_NONE)
+void FSingleTileEditorViewportClient::TrackingStarted(const struct FInputEventState& InInputState, bool bIsDragging, bool bNudge)
+{
+	//@TODO: Should push this into FEditorViewportClient
+	// Begin transacting.  Give the current editor mode an opportunity to do the transacting.
+	const bool bTrackingHandledExternally = ModeTools->StartTracking(this, Viewport);
+
+	if (!bManipulating && bIsDragging && !bTrackingHandledExternally)
 	{
-		return;
+		BeginTransaction(LOCTEXT("ModificationInViewport", "Modification in Viewport"));
+		bManipulating = true;
+		bManipulationDirtiedSomething = false;
 	}
+}
+
+void FSingleTileEditorViewportClient::TrackingStopped()
+{
+	// Stop transacting.  Give the current editor mode an opportunity to do the transacting.
+	const bool bTransactingHandledByEditorMode = ModeTools->EndTracking(this, Viewport);
 
-	if (UTexture2D* Texture = TileSet->TileSheet)
+	if (bManipulating && !bTransactingHandledByEditorMode)
 	{
-		const bool bUseTranslucentBlend = Texture->HasAlphaChannel();
+		EndTransaction();
+		bManipulating = false;
+	}
+}
 
-		// Fully stream in the texture before drawing it.
-		Texture->SetForceMipLevelsToBeResident(30.0f);
-		Texture->WaitForStreaming();
+FVector2D FSingleTileEditorViewportClient::SelectedItemConvertWorldSpaceDeltaToLocalSpace(const FVector& WorldSpaceDelta) const
+{
+	const FVector ProjectionX = WorldSpaceDelta.ProjectOnTo(PaperAxisX);
+	const FVector ProjectionY = WorldSpaceDelta.ProjectOnTo(PaperAxisY);
 
-		const float XPos = -ZoomPos.X * ZoomAmount;
-		const float YPos = -ZoomPos.Y * ZoomAmount;
-		const float Width = TileSet->TileWidth * ZoomAmount * PixelSize;
-		const float Height = TileSet->TileHeight * ZoomAmount * PixelSize;
+	const float XValue = FMath::Sign(ProjectionX | PaperAxisX) * ProjectionX.Size();
+	const float YValue = FMath::Sign(ProjectionY | PaperAxisY) * ProjectionY.Size();
+
+	return FVector2D(XValue, YValue);
+}
+
+FVector2D FSingleTileEditorViewportClient::WorldSpaceToTextureSpace(const FVector& SourcePoint) const
+{
+	const FVector ProjectionX = SourcePoint.ProjectOnTo(PaperAxisX);
+	const FVector ProjectionY = -SourcePoint.ProjectOnTo(PaperAxisY);
+
+	const float XValue = FMath::Sign(ProjectionX | PaperAxisX) * ProjectionX.Size();
+	const float YValue = FMath::Sign(ProjectionY | PaperAxisY) * ProjectionY.Size();
 
-		const float InvTextureWidth = 1.0f / Texture->GetSurfaceWidth();
-		const float InvTextureHeight = 1.0f / Texture->GetSurfaceHeight();
+	return FVector2D(XValue, YValue);
+}
 
-		FVector2D TopLeft;
-		TileSet->GetTileUV(TileBeingEditedIndex, /*out*/ TopLeft);
-		const FVector2D UV0(TopLeft.X * InvTextureWidth, TopLeft.Y * InvTextureHeight);
+FVector FSingleTileEditorViewportClient::TextureSpaceToWorldSpace(const FVector2D& SourcePoint) const
+{
+	return (SourcePoint.X * PaperAxisX) - (SourcePoint.Y * PaperAxisY);
+}
 
-		const FVector2D UVSize(TileSet->TileWidth * InvTextureWidth, TileSet->TileHeight * InvTextureHeight);
+float FSingleTileEditorViewportClient::SelectedItemGetUnitsPerPixel() const
+{
+	return 1.0f;
+}
 
-		const FLinearColor TextureDrawColor = FLinearColor::White;
-		Canvas->DrawTile(XPos, YPos, Width, Height, UV0.X, UV0.Y, UVSize.X, UVSize.Y, TextureDrawColor, Texture->Resource, bUseTranslucentBlend);
+void FSingleTileEditorViewportClient::BeginTransaction(const FText& SessionName)
+{
+	if (ScopedTransaction == nullptr)
+	{
+		ScopedTransaction = new FScopedTransaction(SessionName);
+
+		TileSet->Modify();
 	}
 }
 
+void FSingleTileEditorViewportClient::MarkTransactionAsDirty()
+{
+	bManipulationDirtiedSomething = true;
+	Invalidate();
+}
 
-FLinearColor FSingleTileEditorViewportClient::GetBackgroundColor() const
+void FSingleTileEditorViewportClient::EndTransaction()
 {
-	return TileSet->BackgroundColor;
+	if (bManipulationDirtiedSomething)
+	{
+		TileSet->PostEditChange();
+	}
+
+	bManipulationDirtiedSomething = false;
+
+	if (ScopedTransaction != nullptr)
+	{
+		delete ScopedTransaction;
+		ScopedTransaction = nullptr;
+	}
+}
+
+void FSingleTileEditorViewportClient::InvalidateViewportAndHitProxies()
+{
+	Invalidate();
 }
 
 void FSingleTileEditorViewportClient::SetTileIndex(int32 InTileIndex)
 {
 	const bool bNewIndexValid = (InTileIndex >= 0) && (InTileIndex < TileSet->GetTileCount());
 	TileBeingEditedIndex = bNewIndexValid ? InTileIndex : INDEX_NONE;
+
+	FSpriteGeometryEditMode* GeometryEditMode = ModeTools->GetActiveModeTyped<FSpriteGeometryEditMode>(FSpriteGeometryEditMode::EM_SpriteGeometry);
+	check(GeometryEditMode);
+	FSpriteGeometryCollection* GeomToEdit = nullptr;
+	
+	if (TileBeingEditedIndex != INDEX_NONE)
+	{
+		if (FPaperTileMetadata* Metadata = TileSet->GetMutableTileMetadata(InTileIndex))
+		{
+			GeomToEdit = &(Metadata->CollisionData);
+		}
+	}
+
+	// Tell the geometry editor about the new tile (if it exists)
+	GeometryEditMode->SetGeometryBeingEdited(GeomToEdit, /*bAllowCircles=*/ true, /*bAllowSubtractivePolygons=*/ false);
+
+	// Update the visual representation
+	UPaperSprite* DummySprite = nullptr;
+	if (TileBeingEditedIndex != INDEX_NONE)
+	{
+		//@TODO: Should use this to pick the correct material
+		//const bool bUseTranslucentBlend = Texture->HasAlphaChannel();
+
+		DummySprite = NewObject<UPaperSprite>();
+ 		DummySprite->SpriteCollisionDomain = ESpriteCollisionMode::None;
+ 		DummySprite->PivotMode = ESpritePivotMode::Center_Center;
+ 		DummySprite->CollisionGeometry.GeometryType = ESpritePolygonMode::SourceBoundingBox;
+ 		DummySprite->RenderGeometry.GeometryType = ESpritePolygonMode::SourceBoundingBox;
+		DummySprite->PixelsPerUnrealUnit = 1.0f;
+
+		FSpriteAssetInitParameters SpriteReinitParams;
+		SpriteReinitParams.Texture = TileSet->TileSheet;
+		TileSet->GetTileUV(TileBeingEditedIndex, /*out*/ SpriteReinitParams.Offset);
+		SpriteReinitParams.Dimension = FVector2D(TileSet->TileWidth, TileSet->TileHeight);
+		DummySprite->InitializeSprite(SpriteReinitParams);
+	}
+	PreviewTileSpriteComponent->SetSprite(DummySprite);
+
+	// Update the default geometry bounds
+	const FVector2D HalfTileSize(TileSet->TileWidth * 0.5f, TileSet->TileHeight * 0.5f);
+	FBox2D DesiredBounds(ForceInitToZero);
+	DesiredBounds.Min = -HalfTileSize;
+	DesiredBounds.Max = HalfTileSize;
+	GeometryEditMode->SetNewGeometryPreferredBounds(DesiredBounds);
+
+	// Redraw the viewport
 	Invalidate();
 }
 
@@ -398,10 +652,29 @@ int32 FSingleTileEditorViewportClient::GetTileIndex() const
 	return TileBeingEditedIndex;
 }
 
+void FSingleTileEditorViewportClient::ActivateEditMode(TSharedPtr<FUICommandList> InCommandList)
+{
+	// Activate the sprite geometry edit mode
+	//@TODO: ModeTools->SetToolkitHost(SpriteEditorPtr.Pin()->GetToolkitHost());
+	ModeTools->SetDefaultMode(FSpriteGeometryEditMode::EM_SpriteGeometry);
+	ModeTools->ActivateDefaultMode();
+	ModeTools->SetWidgetMode(FWidget::WM_Translate);
+
+	FSpriteGeometryEditMode* GeometryEditMode = ModeTools->GetActiveModeTyped<FSpriteGeometryEditMode>(FSpriteGeometryEditMode::EM_SpriteGeometry);
+	check(GeometryEditMode);
+	GeometryEditMode->SetEditorContext(this);
+	GeometryEditMode->BindCommands(InCommandList /*SpriteEditorViewportPtr.Pin()->GetCommandList()*/);
+
+	const FLinearColor CollisionShapeColor(0.0f, 0.7f, 1.0f, 1.0f); //@TODO: Duplicated constant from SpriteEditingConstants
+	GeometryEditMode->SetGeometryColors(CollisionShapeColor, FLinearColor::White);
+}
+
 //////////////////////////////////////////////////////////////////////////
 // SSingleTileEditorViewport
 
-class SSingleTileEditorViewport : public SPaperEditorViewport
+#include "SEditorViewport.h"
+
+class SSingleTileEditorViewport : public SEditorViewport, public ICommonEditorViewportToolbarInfoProvider
 {
 public:
 	SLATE_BEGIN_ARGS(SSingleTileEditorViewport) {}
@@ -411,10 +684,21 @@ public:
 
 	void Construct(const FArguments& InArgs, TSharedPtr<class FSingleTileEditorViewportClient> InViewportClient);
 
+	// SEditorViewport interface
+	virtual TSharedPtr<SWidget> MakeViewportToolbar() override;
+	virtual TSharedRef<FEditorViewportClient> MakeEditorViewportClient() override;
+	// End of SEditorViewport interface
+
+	// ICommonEditorViewportToolbarInfoProvider interface
+	virtual TSharedRef<class SEditorViewport> GetViewportWidget() override;
+	virtual TSharedPtr<FExtender> GetExtenders() const override;
+	virtual void OnFloatingButtonClicked() override;
+	// End of ICommonEditorViewportToolbarInfoProvider interface
+
 protected:
-	// SPaperEditorViewport interface
-	virtual FText GetTitleText() const override;
-	// End of SPaperEditorViewport interface
+	FText GetTitleText() const;
+
+	bool IsVisible() const;
 
 private:
 	TSharedPtr<class FSingleTileEditorViewportClient> TypedViewportClient;
@@ -433,12 +717,72 @@ void SSingleTileEditorViewport::Construct(const FArguments& InArgs, TSharedPtr<c
 {
 	TypedViewportClient = InViewportClient;
 
-	SPaperEditorViewport::Construct(
-		SPaperEditorViewport::FArguments(),
-		TypedViewportClient.ToSharedRef());
+	SEditorViewport::Construct(SEditorViewport::FArguments());
 
-	// Make sure we get input instead of the viewport stealing it
-	ViewportWidget->SetVisibility(EVisibility::HitTestInvisible);
+	TSharedRef<SWidget> ParentContents = ChildSlot.GetWidget();
+
+	this->ChildSlot
+	[
+		SNew(SOverlay)
+		+SOverlay::Slot()
+		[
+			ParentContents
+		]
+
+		+SOverlay::Slot()
+		.VAlign(VAlign_Bottom)
+		[
+			SNew(SBorder)
+			.BorderImage( FEditorStyle::GetBrush( TEXT("Graph.TitleBackground") ) ) //@TODO: Switch this over to the Paper2D style
+			.HAlign(HAlign_Fill)
+			.Visibility(EVisibility::HitTestInvisible)
+			[
+				SNew(SVerticalBox)
+				// Title text/icon
+				+SVerticalBox::Slot()
+				.AutoHeight()
+				[
+					SNew(SHorizontalBox)
+					+SHorizontalBox::Slot()
+					.HAlign(HAlign_Center)
+					.FillWidth(1.f)
+					[
+						SNew(STextBlock)
+						.Font( FSlateFontInfo( FPaths::EngineContentDir() / TEXT("Slate/Fonts/Roboto-Regular.ttf"), 18 ) ) //@TODO: Switch this over to the Paper2D style
+						.ColorAndOpacity( FLinearColor(1,1,1,0.5) )
+						.Text(this, &SSingleTileEditorViewport::GetTitleText)
+					]
+				]
+			]
+		]
+	];
+}
+
+TSharedPtr<SWidget> SSingleTileEditorViewport::MakeViewportToolbar()
+{
+	return SNew(STileSetEditorViewportToolbar, SharedThis(this));
+}
+
+TSharedRef<FEditorViewportClient> SSingleTileEditorViewport::MakeEditorViewportClient()
+{
+	TypedViewportClient->VisibilityDelegate.BindSP(this, &SSingleTileEditorViewport::IsVisible);
+
+	return TypedViewportClient.ToSharedRef();
+}
+
+TSharedRef<class SEditorViewport> SSingleTileEditorViewport::GetViewportWidget()
+{
+	return SharedThis(this);
+}
+
+TSharedPtr<FExtender> SSingleTileEditorViewport::GetExtenders() const
+{
+	TSharedPtr<FExtender> Result(MakeShareable(new FExtender));
+	return Result;
+}
+
+void SSingleTileEditorViewport::OnFloatingButtonClicked()
+{
 }
 
 FText SSingleTileEditorViewport::GetTitleText() const
@@ -457,6 +801,11 @@ FText SSingleTileEditorViewport::GetTitleText() const
 	}
 }
 
+bool SSingleTileEditorViewport::IsVisible() const
+{
+	return true;//@TODO: Determine this better so viewport ticking optimizations can take place
+}
+
 //////////////////////////////////////////////////////////////////////////
 // FTileSetEditor
 
@@ -516,7 +865,11 @@ void FTileSetEditor::InitTileSetEditor(const EToolkitMode::Type Mode, const TSha
 	TileEditorViewportClient = MakeShareable(new FSingleTileEditorViewportClient(InitTileSet));
 	TileSetViewport->GetTileSelectionChanged().AddRaw(TileEditorViewportClient.Get(), &FSingleTileEditorViewportClient::OnActiveTileIndexChanged);
 
+	TileEditorViewport = SNew(SSingleTileEditorViewport, TileEditorViewportClient);
+
 	//@TODO: FTileSetEditorCommands::Register();
+	FSpriteGeometryEditCommands::Register();
+
 	BindCommands();
 
 	// Default layout
@@ -568,6 +921,8 @@ void FTileSetEditor::InitTileSetEditor(const EToolkitMode::Type Mode, const TSha
 
 	// Initialize the asset editor
 	InitAssetEditor(Mode, InitToolkitHost, TileSetEditorAppName, StandaloneDefaultLayout, /*bCreateDefaultStandaloneMenu=*/ true, /*bCreateDefaultToolbar=*/ true, InitTileSet);
+	
+	TileEditorViewportClient->ActivateEditMode(TileEditorViewport->GetCommandList());
 
 	// Extend things
 	ExtendMenu();
@@ -602,7 +957,7 @@ TSharedRef<SDockTab> FTileSetEditor::SpawnTab_SingleTileEditor(const FSpawnTabAr
 	return SNew(SDockTab)
 		.Label(LOCTEXT("SingleTileEditTabLabel", "Tile Editor"))
 		[
-			SNew(SSingleTileEditorViewport, TileEditorViewportClient)
+			TileEditorViewport.ToSharedRef()
 		];
 }
 
@@ -616,30 +971,31 @@ void FTileSetEditor::ExtendMenu()
 
 void FTileSetEditor::ExtendToolbar()
 {
-//@TODO: Add geometry creation commands
-// 	struct Local
-// 	{
-// 		static void FillToolbar(FToolBarBuilder& ToolbarBuilder)
-// 		{
-// 			ToolbarBuilder.BeginSection("Tools");
-// 			{
-// 				ToolbarBuilder.AddToolBarButton(FSpriteEditorCommands::Get().AddBoxShape);
-// 				ToolbarBuilder.AddToolBarButton(FSpriteEditorCommands::Get().ToggleAddPolygonMode);
-// 				ToolbarBuilder.AddToolBarButton(FSpriteEditorCommands::Get().AddCircleShape);
-// 				ToolbarBuilder.AddToolBarButton(FSpriteEditorCommands::Get().SnapAllVertices);
-// 			}
-// 			ToolbarBuilder.EndSection();
-// 		}
-// 	};
-// 
-// 	TSharedPtr<FExtender> ToolbarExtender = MakeShareable(new FExtender);
-// 
-// 	ToolbarExtender->AddToolBarExtension(
-// 		"Asset",
-// 		EExtensionHook::After,
-// 		ViewportPtr->GetCommandList(),
-// 		FToolBarExtensionDelegate::CreateStatic(&Local::FillToolbar)
-// 		);
+	struct Local
+	{
+		static void FillToolbar(FToolBarBuilder& ToolbarBuilder)
+		{
+			ToolbarBuilder.BeginSection("Tools");
+			{
+				ToolbarBuilder.AddToolBarButton(FSpriteGeometryEditCommands::Get().AddBoxShape);
+				ToolbarBuilder.AddToolBarButton(FSpriteGeometryEditCommands::Get().ToggleAddPolygonMode);
+				ToolbarBuilder.AddToolBarButton(FSpriteGeometryEditCommands::Get().AddCircleShape);
+				ToolbarBuilder.AddToolBarButton(FSpriteGeometryEditCommands::Get().SnapAllVertices);
+			}
+			ToolbarBuilder.EndSection();
+		}
+	};
+
+	TSharedPtr<FExtender> ToolbarExtender = MakeShareable(new FExtender);
+
+	ToolbarExtender->AddToolBarExtension(
+		"Asset",
+		EExtensionHook::After,
+		TileEditorViewport->GetCommandList(),
+		FToolBarExtensionDelegate::CreateStatic(&Local::FillToolbar)
+		);
+
+	AddToolbarExtender(ToolbarExtender);
 }
 
 void FTileSetEditor::OnPropertyChanged(UObject* ObjectBeingModified, FPropertyChangedEvent& PropertyChangedEvent)
@@ -686,6 +1042,22 @@ FString FTileSetEditor::GetDocumentationLink() const
 	return TEXT("Engine/Paper2D/TileSetEditor");
 }
 
+void FTileSetEditor::OnToolkitHostingStarted(const TSharedRef<class IToolkit>& Toolkit)
+{
+	//@TODO: MODETOOLS: Need to be able to register the widget in the toolbox panel with ToolkitHost, so it can instance the ed mode widgets into it
+	// 	TSharedPtr<SWidget> InlineContent = Toolkit->GetInlineContent();
+	// 	if (InlineContent.IsValid())
+	// 	{
+	// 		ToolboxPtr->SetContent(InlineContent.ToSharedRef());
+	// 	}
+}
+
+void FTileSetEditor::OnToolkitHostingFinished(const TSharedRef<class IToolkit>& Toolkit)
+{
+	//ToolboxPtr->SetContent(SNullWidget::NullWidget);
+	//@TODO: MODETOOLS: How to handle multiple ed modes at once in a standalone asset editor?
+}
+
 void FTileSetEditor::AddReferencedObjects(FReferenceCollector& Collector)
 {
 	Collector.AddReferencedObject(TileSetBeingEdited);
diff --git a/Engine/Plugins/2D/Paper2D/Source/Paper2DEditor/Private/TileSetEditor.h b/Engine/Plugins/2D/Paper2D/Source/Paper2DEditor/Private/TileSetEditor.h
--- a/Engine/Plugins/2D/Paper2D/Source/Paper2DEditor/Private/TileSetEditor.h
+++ b/Engine/Plugins/2D/Paper2D/Source/Paper2DEditor/Private/TileSetEditor.h
@@ -27,6 +27,8 @@ public:
 	virtual FLinearColor GetWorldCentricTabColorScale() const override;
 	virtual FString GetWorldCentricTabPrefix() const override;
 	virtual FString GetDocumentationLink() const override;
+	virtual void OnToolkitHostingStarted(const TSharedRef<class IToolkit>& Toolkit) override;
+	virtual void OnToolkitHostingFinished(const TSharedRef<class IToolkit>& Toolkit) override;
 	// End of FAssetEditorToolkit
 
 	// FSerializableObject interface
@@ -53,6 +55,7 @@ protected:
 	UPaperTileSet* TileSetBeingEdited;
 
 	TSharedPtr<class STileSetSelectorViewport> TileSetViewport;
+	TSharedPtr<class SSingleTileEditorViewport> TileEditorViewport;
 	TSharedPtr<class FSingleTileEditorViewportClient> TileEditorViewportClient;
 
 	FDelegateHandle OnPropertyChangedHandle;
