Paper2D: Removed implementation of Draw(FViewport* Viewport, FCanvas* Canvas) in FPaperEditorViewportClient since most children were skipping it to call the FEditorViewportClient version

diff --git a/Engine/Plugins/2D/Paper2D/Source/Paper2DEditor/Private/ExtractSprites/SPaperExtractSpritesDialog.cpp b/Engine/Plugins/2D/Paper2D/Source/Paper2DEditor/Private/ExtractSprites/SPaperExtractSpritesDialog.cpp
--- a/Engine/Plugins/2D/Paper2D/Source/Paper2DEditor/Private/ExtractSprites/SPaperExtractSpritesDialog.cpp
+++ b/Engine/Plugins/2D/Paper2D/Source/Paper2DEditor/Private/ExtractSprites/SPaperExtractSpritesDialog.cpp
@@ -41,11 +41,10 @@ private:
 
 void FPaperExtractSpritesViewportClient::Draw(FViewport* Viewport, FCanvas* Canvas)
 {
-	// Super will clear the viewport 
-	FPaperEditorViewportClient::Draw(Viewport, Canvas);
+	// Clear the viewport 
+	Canvas->Clear(GetBackgroundColor());
 
-	UTexture2D* Texture = TextureBeingExtracted.Get();
-	if (Texture != nullptr)
+	if (UTexture2D* Texture = TextureBeingExtracted.Get())
 	{
 		const bool bUseTranslucentBlend = Texture->HasAlphaChannel();
 
diff --git a/Engine/Plugins/2D/Paper2D/Source/Paper2DEditor/Private/FlipbookEditor/FlipbookEditorViewportClient.cpp b/Engine/Plugins/2D/Paper2D/Source/Paper2DEditor/Private/FlipbookEditor/FlipbookEditorViewportClient.cpp
--- a/Engine/Plugins/2D/Paper2D/Source/Paper2DEditor/Private/FlipbookEditor/FlipbookEditorViewportClient.cpp
+++ b/Engine/Plugins/2D/Paper2D/Source/Paper2DEditor/Private/FlipbookEditor/FlipbookEditorViewportClient.cpp
@@ -65,11 +65,6 @@ void FFlipbookEditorViewportClient::DrawCanvas(FViewport& Viewport, FSceneView&
 	}
 }
 
-void FFlipbookEditorViewportClient::Draw(FViewport* Viewport, FCanvas* Canvas)
-{
-	FEditorViewportClient::Draw(Viewport, Canvas);
-}
-
 void FFlipbookEditorViewportClient::Draw(const FSceneView* View, FPrimitiveDrawInterface* PDI)
 {
 	FEditorViewportClient::Draw(View, PDI);
diff --git a/Engine/Plugins/2D/Paper2D/Source/Paper2DEditor/Private/FlipbookEditor/FlipbookEditorViewportClient.h b/Engine/Plugins/2D/Paper2D/Source/Paper2DEditor/Private/FlipbookEditor/FlipbookEditorViewportClient.h
--- a/Engine/Plugins/2D/Paper2D/Source/Paper2DEditor/Private/FlipbookEditor/FlipbookEditorViewportClient.h
+++ b/Engine/Plugins/2D/Paper2D/Source/Paper2DEditor/Private/FlipbookEditor/FlipbookEditorViewportClient.h
@@ -20,7 +20,6 @@ public:
 	FFlipbookEditorViewportClient(const TAttribute<class UPaperFlipbook*>& InFlipbookBeingEdited);
 
 	// FViewportClient interface
-	virtual void Draw(FViewport* Viewport, FCanvas* Canvas) override;
 	virtual void Draw(const FSceneView* View, FPrimitiveDrawInterface* PDI) override;
 	virtual void DrawCanvas(FViewport& InViewport, FSceneView& View, FCanvas& Canvas) override;
 	virtual void Tick(float DeltaSeconds) override;
diff --git a/Engine/Plugins/2D/Paper2D/Source/Paper2DEditor/Private/PaperEditorViewportClient.cpp b/Engine/Plugins/2D/Paper2D/Source/Paper2DEditor/Private/PaperEditorViewportClient.cpp
--- a/Engine/Plugins/2D/Paper2D/Source/Paper2DEditor/Private/PaperEditorViewportClient.cpp
+++ b/Engine/Plugins/2D/Paper2D/Source/Paper2DEditor/Private/PaperEditorViewportClient.cpp
@@ -65,11 +65,6 @@ FLinearColor FPaperEditorViewportClient::GetBackgroundColor() const
 	return FLinearColor(0, 0, 127, 0);
 }
 
-void FPaperEditorViewportClient::Draw(FViewport* Viewport, FCanvas* Canvas)
-{
-	Canvas->Clear(GetBackgroundColor());
-}
-
 void FPaperEditorViewportClient::Tick(float DeltaSeconds)
 {
 	// Zoom in on the sprite
diff --git a/Engine/Plugins/2D/Paper2D/Source/Paper2DEditor/Private/PaperEditorViewportClient.h b/Engine/Plugins/2D/Paper2D/Source/Paper2DEditor/Private/PaperEditorViewportClient.h
--- a/Engine/Plugins/2D/Paper2D/Source/Paper2DEditor/Private/PaperEditorViewportClient.h
+++ b/Engine/Plugins/2D/Paper2D/Source/Paper2DEditor/Private/PaperEditorViewportClient.h
@@ -23,7 +23,6 @@ public:
 	~FPaperEditorViewportClient();
 
 	// FViewportClient interface
-	virtual void Draw(FViewport* Viewport, FCanvas* Canvas) override;
 	virtual void Tick(float DeltaSeconds) override;
 	// End of FViewportClient interface
 
diff --git a/Engine/Plugins/2D/Paper2D/Source/Paper2DEditor/Private/SpriteEditor/SpriteEditorViewportClient.cpp b/Engine/Plugins/2D/Paper2D/Source/Paper2DEditor/Private/SpriteEditor/SpriteEditorViewportClient.cpp
--- a/Engine/Plugins/2D/Paper2D/Source/Paper2DEditor/Private/SpriteEditor/SpriteEditorViewportClient.cpp
+++ b/Engine/Plugins/2D/Paper2D/Source/Paper2DEditor/Private/SpriteEditor/SpriteEditorViewportClient.cpp
@@ -662,12 +662,6 @@ void FSpriteEditorViewportClient::Draw(const FSceneView* View, FPrimitiveDrawInt
 	}
 }
 
-void FSpriteEditorViewportClient::Draw(FViewport* Viewport, FCanvas* Canvas)
-{
-	// Skipping the parent on purpose
-	FEditorViewportClient::Draw(Viewport, Canvas);
-}
-
 FBox FSpriteEditorViewportClient::GetDesiredFocusBounds() const
 {
 	UPaperSpriteComponent* ComponentToFocusOn = SourceTextureViewComponent->IsVisible() ? SourceTextureViewComponent : RenderSpriteComponent;
diff --git a/Engine/Plugins/2D/Paper2D/Source/Paper2DEditor/Private/SpriteEditor/SpriteEditorViewportClient.h b/Engine/Plugins/2D/Paper2D/Source/Paper2DEditor/Private/SpriteEditor/SpriteEditorViewportClient.h
--- a/Engine/Plugins/2D/Paper2D/Source/Paper2DEditor/Private/SpriteEditor/SpriteEditorViewportClient.h
+++ b/Engine/Plugins/2D/Paper2D/Source/Paper2DEditor/Private/SpriteEditor/SpriteEditorViewportClient.h
@@ -33,7 +33,6 @@ public:
 	FSpriteEditorViewportClient(TWeakPtr<FSpriteEditor> InSpriteEditor, TWeakPtr<class SEditorViewport> InSpriteEditorViewportPtr);
 
 	// FViewportClient interface
-	virtual void Draw(FViewport* Viewport, FCanvas* Canvas) override;
 	virtual void Draw(const FSceneView* View, FPrimitiveDrawInterface* PDI) override;
 	virtual void DrawCanvas(FViewport& InViewport, FSceneView& View, FCanvas& Canvas) override;
 	virtual void Tick(float DeltaSeconds) override;
diff --git a/Engine/Plugins/2D/Paper2D/Source/Paper2DEditor/Private/TileMapEditing/TileMapEditorViewportClient.cpp b/Engine/Plugins/2D/Paper2D/Source/Paper2DEditor/Private/TileMapEditing/TileMapEditorViewportClient.cpp
--- a/Engine/Plugins/2D/Paper2D/Source/Paper2DEditor/Private/TileMapEditing/TileMapEditorViewportClient.cpp
+++ b/Engine/Plugins/2D/Paper2D/Source/Paper2DEditor/Private/TileMapEditing/TileMapEditorViewportClient.cpp
@@ -108,21 +108,6 @@ void FTileMapEditorViewportClient::DrawCanvas(FViewport& Viewport, FSceneView& V
 	DrawBoundsAsText(Viewport, View, Canvas, /*inout*/ YPos);
 }
 
-void FTileMapEditorViewportClient::Draw(const FSceneView* View, FPrimitiveDrawInterface* PDI)
-{
-	FEditorViewportClient::Draw(View, PDI);
-
-	if (bShowPivot)
-	{
-		FUnrealEdUtils::DrawWidget(View, PDI, RenderTileMapComponent->ComponentToWorld.ToMatrixWithScale(), 0, 0, EAxisList::XZ, EWidgetMovementMode::WMM_Translate);
-	}
-}
-
-void FTileMapEditorViewportClient::Draw(FViewport* Viewport, FCanvas* Canvas)
-{
-	FEditorViewportClient::Draw(Viewport, Canvas);
-}
-
 void FTileMapEditorViewportClient::Tick(float DeltaSeconds)
 {
 	FPaperEditorViewportClient::Tick(DeltaSeconds);
diff --git a/Engine/Plugins/2D/Paper2D/Source/Paper2DEditor/Private/TileMapEditing/TileMapEditorViewportClient.h b/Engine/Plugins/2D/Paper2D/Source/Paper2DEditor/Private/TileMapEditing/TileMapEditorViewportClient.h
--- a/Engine/Plugins/2D/Paper2D/Source/Paper2DEditor/Private/TileMapEditing/TileMapEditorViewportClient.h
+++ b/Engine/Plugins/2D/Paper2D/Source/Paper2DEditor/Private/TileMapEditing/TileMapEditorViewportClient.h
@@ -20,8 +20,6 @@ public:
 	FTileMapEditorViewportClient(TWeakPtr<FTileMapEditor> InTileMapEditor, TWeakPtr<class STileMapEditorViewport> InTileMapEditorViewportPtr);
 
 	// FViewportClient interface
-	virtual void Draw(FViewport* Viewport, FCanvas* Canvas) override;
-	virtual void Draw(const FSceneView* View, FPrimitiveDrawInterface* PDI) override;
 	virtual void DrawCanvas(FViewport& InViewport, FSceneView& View, FCanvas& Canvas) override;
 	virtual void Tick(float DeltaSeconds) override;
 	// End of FViewportClient interface
