Tile Set Editor: Added an explicit clear back in

diff --git a/Engine/Plugins/2D/Paper2D/Source/Paper2DEditor/Private/TileSetEditor.cpp b/Engine/Plugins/2D/Paper2D/Source/Paper2DEditor/Private/TileSetEditor.cpp
--- a/Engine/Plugins/2D/Paper2D/Source/Paper2DEditor/Private/TileSetEditor.cpp
+++ b/Engine/Plugins/2D/Paper2D/Source/Paper2DEditor/Private/TileSetEditor.cpp
@@ -66,8 +66,8 @@ public:
 
 void FTileSetEditorViewportClient::Draw(FViewport* Viewport, FCanvas* Canvas)
 {
-	// Super will clear the viewport
-	FPaperEditorViewportClient::Draw(Viewport, Canvas);
+	// Clear the viewport
+	Canvas->Clear(GetBackgroundColor());
 
 	// Can only proceed if we have a valid tile set
 	UPaperTileSet* TileSet = TileSetBeingEdited.Get();
@@ -374,7 +374,6 @@ public:
 
 	// FViewportClient interface
 	virtual void Tick(float DeltaSeconds) override;
-	virtual void Draw(FViewport* Viewport, FCanvas* Canvas) override;
 	// End of FViewportClient interface
 
 	// FEditorViewportClient interface
@@ -482,12 +481,6 @@ void FSingleTileEditorViewportClient::Tick(float DeltaSeconds)
 	}
 }
 
-void FSingleTileEditorViewportClient::Draw(FViewport* Viewport, FCanvas* Canvas)
-{
-	// Skipping the parent on purpose
-	FEditorViewportClient::Draw(Viewport, Canvas);
-}
-
 FLinearColor FSingleTileEditorViewportClient::GetBackgroundColor() const
 {
 	return TileSet->BackgroundColor;
