Paper2D: Removed inconsistent icons in the "Sprite Actions" submenu

diff --git a/Engine/Plugins/2D/Paper2D/Source/Paper2DEditor/Private/ContentBrowserExtensions/ContentBrowserExtensions.cpp b/Engine/Plugins/2D/Paper2D/Source/Paper2DEditor/Private/ContentBrowserExtensions/ContentBrowserExtensions.cpp
--- a/Engine/Plugins/2D/Paper2D/Source/Paper2DEditor/Private/ContentBrowserExtensions/ContentBrowserExtensions.cpp
+++ b/Engine/Plugins/2D/Paper2D/Source/Paper2DEditor/Private/ContentBrowserExtensions/ContentBrowserExtensions.cpp
@@ -213,7 +213,7 @@ public:
 		MenuBuilder.AddMenuEntry(
 			LOCTEXT("CB_Extension_Texture_CreateSprite", "Create Sprite"),
 			LOCTEXT("CB_Extension_Texture_CreateSprite_Tooltip", "Create sprites from selected textures"),
-			FSlateIcon(FEditorStyle::GetStyleSetName(), "ClassIcon.PaperSprite"),
+			FSlateIcon(),
 			Action_CreateSpritesFromTextures,
 			NAME_None,
 			EUserInterfaceActionType::Button);
@@ -259,7 +259,7 @@ public:
 		MenuBuilder.AddMenuEntry(
 			LOCTEXT("CB_Extension_Texture_CreateTileSet", "Create Tile Set"),
 			LOCTEXT("CB_Extension_Texture_CreateTileSet_Tooltip", "Create tile set from selected texture"),
-			FSlateIcon(FEditorStyle::GetStyleSetName(), "ClassIcon.PaperTileSet"),
+			FSlateIcon(),
 			Action_CreateTileSetFromTextures,
 			NAME_None,
 			EUserInterfaceActionType::Button);
