Sprite Editor: Added a notification message when switching which sprite is being edited in Source Region edit mode

diff --git a/Engine/Plugins/2D/Paper2D/Source/Paper2DEditor/Private/SpriteEditor/SpriteEditorViewportClient.cpp b/Engine/Plugins/2D/Paper2D/Source/Paper2DEditor/Private/SpriteEditor/SpriteEditorViewportClient.cpp
--- a/Engine/Plugins/2D/Paper2D/Source/Paper2DEditor/Private/SpriteEditor/SpriteEditorViewportClient.cpp
+++ b/Engine/Plugins/2D/Paper2D/Source/Paper2DEditor/Private/SpriteEditor/SpriteEditorViewportClient.cpp
@@ -20,6 +20,9 @@
 #include "PhysicsEngine/BodySetup2D.h"
 #include "SEditorViewport.h"
 
+#include "SNotificationList.h"
+#include "NotificationManager.h"
+
 #define LOCTEXT_NAMESPACE "SpriteEditor"
 
 //////////////////////////////////////////////////////////////////////////
@@ -874,7 +877,7 @@ UPaperSprite* FSpriteEditorViewportClient::CreateNewSprite(FVector2D TopLeft, FV
 		ContentBrowserModule.Get().SyncBrowserToAssets(Objects);
 
 		UPaperSprite* NewSprite = Cast<UPaperSprite>(NewAsset);
-		if (SpriteEditorPtr.IsValid() && NewSprite != nullptr)
+		if (SpriteEditorPtr.IsValid() && (NewSprite != nullptr))
 		{
 			SpriteEditorPtr.Pin()->SetSpriteBeingEdited(NewSprite);
 		}
@@ -1024,6 +1027,9 @@ void FSpriteEditorViewportClient::NotifySpriteBeingEditedHasChanged()
 	//@TODO: Ideally we do this before switching
 	EndTransaction();
 
+	// Refresh the viewport in case we were not in realtime mode
+	Invalidate();
+
 	// Update components to know about the new sprite being edited
 	UPaperSprite* Sprite = GetSpriteBeingEdited();
 
@@ -1034,6 +1040,15 @@ void FSpriteEditorViewportClient::NotifySpriteBeingEditedHasChanged()
 
 	//@TODO: Only do this if the sprite isn't visible (may consider doing a flashing pulse around the source region rect?)
 	RequestFocusOnSelection(/*bInstant=*/ true);
+
+	if (Sprite != nullptr)
+	{
+		// Create and display a notification about the new sprite being edited
+		const FText NotificationErrorText = FText::Format(LOCTEXT("SwitchingToSprite", "Editing {0}"), FText::AsCultureInvariant(Sprite->GetName()));
+		FNotificationInfo Info(NotificationErrorText);
+		Info.ExpireDuration = 2.0f;
+		FSlateNotificationManager::Get().AddNotification(Info);
+	}
 }
 
 void FSpriteEditorViewportClient::InternalActivateNewMode(ESpriteEditorMode::Type NewMode)
