Sprite Editor: Removed a redundant help line in the Source Region Edit Mode
[UE-12891]

diff --git a/Engine/Plugins/2D/Paper2D/Source/Paper2DEditor/Private/SpriteEditor/SpriteEditorViewportClient.cpp b/Engine/Plugins/2D/Paper2D/Source/Paper2DEditor/Private/SpriteEditor/SpriteEditorViewportClient.cpp
--- a/Engine/Plugins/2D/Paper2D/Source/Paper2DEditor/Private/SpriteEditor/SpriteEditorViewportClient.cpp
+++ b/Engine/Plugins/2D/Paper2D/Source/Paper2DEditor/Private/SpriteEditor/SpriteEditorViewportClient.cpp
@@ -559,7 +559,7 @@ void FSpriteEditorViewportClient::DrawCanvas(FViewport& Viewport, FSceneView& Vi
 
 	int32 YPos = 42;
 
-	static const FText SourceRegionHelpStr = LOCTEXT("SourceRegionHelp", "Drag handles to adjust source region\nDouble-click on an image region to select all connected pixels\nHold down Ctrl and drag a rectangle to create a new sprite at that position\nDouble-click to create a new sprite under the cursor\nClick on other sprite rectangles to change the active sprite");
+	static const FText SourceRegionHelpStr = LOCTEXT("SourceRegionHelp", "Drag handles to adjust source region\nDouble-click on an image region to create a new sprite from all connected pixels\nHold down Ctrl and drag a rectangle to create a new sprite at that position\nClick on other sprite rectangles to change the active sprite");
 
 	switch (CurrentMode)
 	{
