Paper2D: Moved some brush and font references into FPaperStyle

diff --git a/Engine/Plugins/2D/Paper2D/Source/Paper2DEditor/Private/PaperStyle.cpp b/Engine/Plugins/2D/Paper2D/Source/Paper2DEditor/Private/PaperStyle.cpp
--- a/Engine/Plugins/2D/Paper2D/Source/Paper2DEditor/Private/PaperStyle.cpp
+++ b/Engine/Plugins/2D/Paper2D/Source/Paper2DEditor/Private/PaperStyle.cpp
@@ -36,6 +36,22 @@ void FPaperStyle::Initialize()
 	StyleSet->SetContentRoot(FPaths::EngineContentDir() / TEXT("Editor/Slate"));
 	StyleSet->SetCoreContentRoot(FPaths::EngineContentDir() / TEXT("Slate"));
 
+	const FTextBlockStyle& NormalText = FEditorStyle::Get().GetWidgetStyle<FTextBlockStyle>("NormalText");
+
+	// Shared editors
+	{
+		StyleSet->Set("Paper2D.Common.ViewportZoomTextStyle", FTextBlockStyle(NormalText)
+			.SetFont(TTF_FONT("Fonts/Roboto-BoldCondensed", 16))
+			);
+
+		StyleSet->Set("Paper2D.Common.ViewportTitleTextStyle", FTextBlockStyle(NormalText)
+			.SetFont(TTF_FONT("Fonts/Roboto-Regular", 18))
+			.SetColorAndOpacity(FLinearColor(1.0, 1.0f, 1.0f, 0.5f))
+			);
+
+		StyleSet->Set("Paper2D.Common.ViewportTitleBackground", new BOX_BRUSH("Old/Graph/GraphTitleBackground", FMargin(0)));
+	}
+
 	// Tile map editor
 	{
 		StyleSet->Set("TileMapEditor.EnterTileMapEditMode", new IMAGE_PLUGIN_BRUSH(TEXT("Icons/icon_TileMapEdModeIcon_40x"), Icon40x40));
@@ -83,7 +99,7 @@ void FPaperStyle::Initialize()
 
 		StyleSet->Set("TileMapEditor.LayerBrowser.SelectionColor", LayerSelectionColor);
 
-		StyleSet->Set("TileMapEditor.TileSetPalette.NothingSelectedText", FTextBlockStyle(FEditorStyle::Get().GetWidgetStyle<FTextBlockStyle>("NormalText"))
+		StyleSet->Set("TileMapEditor.TileSetPalette.NothingSelectedText", FTextBlockStyle(NormalText)
 			.SetFont(TTF_FONT("Fonts/Roboto-BoldCondensed", 18))
 			.SetColorAndOpacity(FLinearColor(0.8, 0.8f, 0.0f, 0.8f))
 			.SetShadowOffset(FVector2D(1.0f, 1.0f))
diff --git a/Engine/Plugins/2D/Paper2D/Source/Paper2DEditor/Private/SPaperEditorViewport.cpp b/Engine/Plugins/2D/Paper2D/Source/Paper2DEditor/Private/SPaperEditorViewport.cpp
--- a/Engine/Plugins/2D/Paper2D/Source/Paper2DEditor/Private/SPaperEditorViewport.cpp
+++ b/Engine/Plugins/2D/Paper2D/Source/Paper2DEditor/Private/SPaperEditorViewport.cpp
@@ -3,6 +3,7 @@
 #include "Paper2DEditorPrivatePCH.h"
 #include "SPaperEditorViewport.h"
 #include "PaperEditorViewportClient.h"
+#include "PaperStyle.h"
 
 const int32 DefaultZoomLevel = 5;
 const int32 NumZoomLevels = 13;
@@ -61,7 +62,6 @@ void SPaperEditorViewport::Construct(const FArguments& InArgs, TSharedRef<FPaper
 		+SOverlay::Slot()
 		[
 			SAssignNew(ViewportWidget, SViewport)
-			//.Visibility(EVisibility::HitTestInvisible)
 			.EnableGammaCorrection(false)
 			.IsEnabled(FSlateApplication::Get().GetNormalExecutionAttribute())
 			.ShowEffectWhenDisabled(false)
@@ -72,7 +72,7 @@ void SPaperEditorViewport::Construct(const FArguments& InArgs, TSharedRef<FPaper
 		.VAlign(VAlign_Top)
 		[
 			SNew(STextBlock)
-			.Font(FSlateFontInfo( FPaths::EngineContentDir() / TEXT("Slate/Fonts/Roboto-BoldCondensed"), 16))
+			.TextStyle(FPaperStyle::Get(), "Paper2D.Common.ViewportZoomTextStyle")
 			.Text(this, &SPaperEditorViewport::GetZoomText)
 			.ColorAndOpacity(this, &SPaperEditorViewport::GetZoomTextColorAndOpacity)
 		]
@@ -80,7 +80,7 @@ void SPaperEditorViewport::Construct(const FArguments& InArgs, TSharedRef<FPaper
 		.VAlign(VAlign_Top)
 		[
 			SNew(SBorder)
-			.BorderImage( FEditorStyle::GetBrush( TEXT("Graph.TitleBackground") ) )
+			.BorderImage(FPaperStyle::Get()->GetBrush("Paper2D.Common.ViewportTitleBackground"))
 			.HAlign(HAlign_Fill)
 			.Visibility(EVisibility::HitTestInvisible)
 			[
@@ -95,9 +95,8 @@ void SPaperEditorViewport::Construct(const FArguments& InArgs, TSharedRef<FPaper
 					.FillWidth(1.f)
 					[
 						SNew(STextBlock)
-						.Font( FSlateFontInfo( FPaths::EngineContentDir() / TEXT("Slate/Fonts/Roboto-Regular.ttf"), 18 ) )
-						.ColorAndOpacity( FLinearColor(1,1,1,0.5) )
-						.Text( this, &SPaperEditorViewport::GetTitleText )
+						.TextStyle(FPaperStyle::Get(), "Paper2D.Common.ViewportTitleTextStyle")
+						.Text(this, &SPaperEditorViewport::GetTitleText)
 					]
 				]
 			]
diff --git a/Engine/Plugins/2D/Paper2D/Source/Paper2DEditor/Private/TileSetEditor.cpp b/Engine/Plugins/2D/Paper2D/Source/Paper2DEditor/Private/TileSetEditor.cpp
--- a/Engine/Plugins/2D/Paper2D/Source/Paper2DEditor/Private/TileSetEditor.cpp
+++ b/Engine/Plugins/2D/Paper2D/Source/Paper2DEditor/Private/TileSetEditor.cpp
@@ -11,6 +11,7 @@
 #include "CanvasTypes.h"
 #include "CanvasItem.h"
 #include "SDockTab.h"
+#include "PaperStyle.h"
 
 #define LOCTEXT_NAMESPACE "TileSetEditor"
 
@@ -726,7 +727,7 @@ void SSingleTileEditorViewport::Construct(const FArguments& InArgs, TSharedPtr<c
 		.VAlign(VAlign_Bottom)
 		[
 			SNew(SBorder)
-			.BorderImage( FEditorStyle::GetBrush( TEXT("Graph.TitleBackground") ) ) //@TODO: Switch this over to the Paper2D style
+			.BorderImage(FPaperStyle::Get()->GetBrush("Paper2D.Common.ViewportTitleBackground"))
 			.HAlign(HAlign_Fill)
 			.Visibility(EVisibility::HitTestInvisible)
 			[
@@ -741,8 +742,7 @@ void SSingleTileEditorViewport::Construct(const FArguments& InArgs, TSharedPtr<c
 					.FillWidth(1.f)
 					[
 						SNew(STextBlock)
-						.Font( FSlateFontInfo( FPaths::EngineContentDir() / TEXT("Slate/Fonts/Roboto-Regular.ttf"), 18 ) ) //@TODO: Switch this over to the Paper2D style
-						.ColorAndOpacity( FLinearColor(1,1,1,0.5) )
+						.TextStyle(FPaperStyle::Get(), "Paper2D.Common.ViewportTitleTextStyle")
 						.Text(this, &SSingleTileEditorViewport::GetTitleText)
 					]
 				]
