Paper2D: Changed various background color settings to not allow specifying an unused alpha value [UE-13341]

diff --git a/Engine/Plugins/2D/Paper2D/Source/Paper2D/Classes/PaperTileMap.h b/Engine/Plugins/2D/Paper2D/Source/Paper2D/Classes/PaperTileMap.h
--- a/Engine/Plugins/2D/Paper2D/Source/Paper2D/Classes/PaperTileMap.h
+++ b/Engine/Plugins/2D/Paper2D/Source/Paper2D/Classes/PaperTileMap.h
@@ -106,7 +106,7 @@ public:
 	int32 SelectedLayerIndex;
 
 	/** The background color displayed in the tile map editor */
-	UPROPERTY(Category=Setup, EditAnywhere)
+	UPROPERTY(Category=Setup, EditAnywhere, meta=(HideAlphaChannel))
 	FLinearColor BackgroundColor;
 #endif
 
diff --git a/Engine/Plugins/2D/Paper2D/Source/Paper2D/Classes/PaperTileSet.h b/Engine/Plugins/2D/Paper2D/Source/Paper2D/Classes/PaperTileSet.h
--- a/Engine/Plugins/2D/Paper2D/Source/Paper2D/Classes/PaperTileSet.h
+++ b/Engine/Plugins/2D/Paper2D/Source/Paper2D/Classes/PaperTileSet.h
@@ -68,7 +68,7 @@ class PAPER2D_API UPaperTileSet : public UObject
 
 #if WITH_EDITORONLY_DATA
 	/** The background color displayed in the tile set viewer */
-	UPROPERTY(Category=TileSet, EditAnywhere)
+	UPROPERTY(Category=TileSet, EditAnywhere, meta=(HideAlphaChannel))
 	FLinearColor BackgroundColor;
 #endif
 
diff --git a/Engine/Plugins/2D/Paper2D/Source/Paper2DEditor/Private/ExtractSprites/PaperExtractSpritesSettings.h b/Engine/Plugins/2D/Paper2D/Source/Paper2DEditor/Private/ExtractSprites/PaperExtractSpritesSettings.h
--- a/Engine/Plugins/2D/Paper2D/Source/Paper2DEditor/Private/ExtractSprites/PaperExtractSpritesSettings.h
+++ b/Engine/Plugins/2D/Paper2D/Source/Paper2DEditor/Private/ExtractSprites/PaperExtractSpritesSettings.h
@@ -38,7 +38,7 @@ public:
 	FLinearColor TextureTint;
 
 	// The viewport background color
-	UPROPERTY(Category = Settings, EditAnywhere)
+	UPROPERTY(Category = Settings, EditAnywhere, meta=(HideAlphaChannel))
 	FLinearColor BackgroundColor;
 
 	UPROPERTY()
