Tile Maps: Removed support for the old binary collision layers

diff --git a/Engine/Plugins/2D/Paper2D/Source/Paper2D/Classes/PaperTileLayer.h b/Engine/Plugins/2D/Paper2D/Source/Paper2D/Classes/PaperTileLayer.h
--- a/Engine/Plugins/2D/Paper2D/Source/Paper2D/Classes/PaperTileLayer.h
+++ b/Engine/Plugins/2D/Paper2D/Source/Paper2D/Classes/PaperTileLayer.h
@@ -129,10 +129,6 @@ class PAPER2D_API UPaperTileLayer : public UObject
 	UPROPERTY(BlueprintReadOnly, Category=Sprite)
 	bool bHiddenInGame;
 
-	// Is this layer a collision layer?
-	UPROPERTY()
-	bool bCollisionLayer;
-
 protected:
 	// The allocated width of the tile data (used to handle resizing without data loss)
 	UPROPERTY()
diff --git a/Engine/Plugins/2D/Paper2D/Source/Paper2D/Classes/PaperTileMap.h b/Engine/Plugins/2D/Paper2D/Source/Paper2D/Classes/PaperTileMap.h
--- a/Engine/Plugins/2D/Paper2D/Source/Paper2D/Classes/PaperTileMap.h
+++ b/Engine/Plugins/2D/Paper2D/Source/Paper2D/Classes/PaperTileMap.h
@@ -154,7 +154,7 @@ public:
 	FBoxSphereBounds GetRenderBounds() const;
 
 	// Creates and adds a new layer and returns it
-	class UPaperTileLayer* AddNewLayer(bool bCollisionLayer = false, int32 InsertionIndex = INDEX_NONE);
+	class UPaperTileLayer* AddNewLayer(int32 InsertionIndex = INDEX_NONE);
 
 	// Creates a reasonable new layer name
 	static FText GenerateNewLayerName(UPaperTileMap* TileMap);
diff --git a/Engine/Plugins/2D/Paper2D/Source/Paper2D/Private/PaperTileLayer.cpp b/Engine/Plugins/2D/Paper2D/Source/Paper2D/Private/PaperTileLayer.cpp
--- a/Engine/Plugins/2D/Paper2D/Source/Paper2D/Private/PaperTileLayer.cpp
+++ b/Engine/Plugins/2D/Paper2D/Source/Paper2D/Private/PaperTileLayer.cpp
@@ -233,39 +233,6 @@ void UPaperTileLayer::AugmentBodySetup(UBodySetup* ShapeBodySetup)
 			}
 		}
 	}
-
-	//@TODO: Remove this legacy code path; there is no way to author new layers with bCollisionLayer and they're superceeded by per-tile collision data
-	// Create a box element for every non-zero value in the layer
-	if (bCollisionLayer)
-	{
-		//@TODO: Tile pivot issue
-		const float TileThickness = TileMap->GetCollisionThickness();
-
-		//@TODO: When the origin of the component changes, this logic will need to be adjusted as well
-		// The origin is currently the top left
-		const float XOrigin = 0;
-		const float YOrigin = 0;
-
-		// Create a box element for every non-zero value in the layer
-		for (int32 XValue = 0; XValue < LayerWidth; ++XValue)
-		{
-			for (int32 YValue = 0; YValue < LayerHeight; ++YValue)
-			{
-				if (GetCell(XValue, YValue).PackedTileIndex != 0)
-				{
-					//@TODO: BOX2D: Add support for 2D physics on tile maps if this code isn't deleted entirely (see TODO above)
-					FKBoxElem* NewBox = new (ShapeBodySetup->AggGeom.BoxElems) FKBoxElem(TileWidth, TileThickness, TileHeight);
-
-					FVector BoxPosition;
-					BoxPosition.X = XOrigin + XValue * TileWidth;
-					BoxPosition.Y = TileThickness * 0.5f;
-					BoxPosition.Z = YOrigin - YValue * TileHeight;
-
-					NewBox->Center = BoxPosition;
-				}
-			}
-		}
-	}
 }
 
 void UPaperTileLayer::ConvertToTileSetPerCell()
diff --git a/Engine/Plugins/2D/Paper2D/Source/Paper2D/Private/PaperTileMap.cpp b/Engine/Plugins/2D/Paper2D/Source/Paper2D/Private/PaperTileMap.cpp
--- a/Engine/Plugins/2D/Paper2D/Source/Paper2D/Private/PaperTileMap.cpp
+++ b/Engine/Plugins/2D/Paper2D/Source/Paper2D/Private/PaperTileMap.cpp
@@ -332,7 +332,7 @@ FBoxSphereBounds UPaperTileMap::GetRenderBounds() const
 	}
 }
 
-UPaperTileLayer* UPaperTileMap::AddNewLayer(bool bCollisionLayer, int32 InsertionIndex)
+UPaperTileLayer* UPaperTileMap::AddNewLayer(int32 InsertionIndex)
 {
 	// Create the new layer
 	UPaperTileLayer* NewLayer = NewObject<UPaperTileLayer>(this);
@@ -342,7 +342,6 @@ UPaperTileLayer* UPaperTileMap::AddNewLayer(bool bCollisionLayer, int32 Insertio
 	NewLayer->LayerHeight = MapHeight;
 	NewLayer->DestructiveAllocateMap(NewLayer->LayerWidth, NewLayer->LayerHeight);
 	NewLayer->LayerName = GenerateNewLayerName(this);
-	NewLayer->bCollisionLayer = bCollisionLayer;
 
 	// Insert the new layer
 	if (TileLayers.IsValidIndex(InsertionIndex))
diff --git a/Engine/Plugins/2D/Paper2D/Source/Paper2D/Private/PaperTileMapComponent.cpp b/Engine/Plugins/2D/Paper2D/Source/Paper2D/Private/PaperTileMapComponent.cpp
--- a/Engine/Plugins/2D/Paper2D/Source/Paper2D/Private/PaperTileMapComponent.cpp
+++ b/Engine/Plugins/2D/Paper2D/Source/Paper2D/Private/PaperTileMapComponent.cpp
@@ -271,31 +271,21 @@ void UPaperTileMapComponent::RebuildRenderData(FPaperTileMapRenderSceneProxy* Pr
 					UTexture2D* SourceTexture = nullptr;
 
 					FVector2D SourceUV = FVector2D::ZeroVector;
-					if (Layer->bCollisionLayer)
+
+					if (TileInfo.TileSet == nullptr)
 					{
-						if (TileInfo.PackedTileIndex == 0)
-						{
-							continue;
-						}
-						SourceTexture = UCanvas::StaticClass()->GetDefaultObject<UCanvas>()->DefaultTexture;
+						continue;
 					}
-					else
-					{
-						if (TileInfo.TileSet == nullptr)
-						{
-							continue;
-						}
 
-						if (!TileInfo.TileSet->GetTileUV(TileInfo.GetTileIndex(), /*out*/ SourceUV))
-						{
-							continue;
-						}
+					if (!TileInfo.TileSet->GetTileUV(TileInfo.GetTileIndex(), /*out*/ SourceUV))
+					{
+						continue;
+					}
 
-						SourceTexture = TileInfo.TileSet->TileSheet;
-						if (SourceTexture == nullptr)
-						{
-							continue;
-						}
+					SourceTexture = TileInfo.TileSet->TileSheet;
+					if (SourceTexture == nullptr)
+					{
+						continue;
 					}
 
 					if ((SourceTexture != LastSourceTexture) || (CurrentBatch == nullptr))
diff --git a/Engine/Plugins/2D/Paper2D/Source/Paper2DEditor/Private/TileMapEditing/STileLayerList.cpp b/Engine/Plugins/2D/Paper2D/Source/Paper2DEditor/Private/TileMapEditing/STileLayerList.cpp
--- a/Engine/Plugins/2D/Paper2D/Source/Paper2DEditor/Private/TileMapEditing/STileLayerList.cpp
+++ b/Engine/Plugins/2D/Paper2D/Source/Paper2DEditor/Private/TileMapEditing/STileLayerList.cpp
@@ -208,7 +208,7 @@ FText STileLayerList::GenerateDuplicatedLayerName(const FString& InputNameRaw, U
 	return FText::FromString(TestLayerName);
 }
 
-UPaperTileLayer* STileLayerList::AddLayer(bool bCollisionLayer, int32 InsertionIndex)
+UPaperTileLayer* STileLayerList::AddLayer(int32 InsertionIndex)
 {
 	UPaperTileLayer* NewLayer = nullptr;
 
@@ -218,7 +218,7 @@ UPaperTileLayer* STileLayerList::AddLayer(bool bCollisionLayer, int32 InsertionI
 		TileMap->SetFlags(RF_Transactional);
 		TileMap->Modify();
 
-		NewLayer = TileMap->AddNewLayer(bCollisionLayer, InsertionIndex);
+		NewLayer = TileMap->AddNewLayer(InsertionIndex);
 
 		PostEditNotfications();
 
@@ -256,12 +256,12 @@ void STileLayerList::ChangeLayerOrdering(int32 OldIndex, int32 NewIndex)
 
 void STileLayerList::AddNewLayerAbove()
 {
-	AddLayer(/*bCollisionLayer=*/ false, GetSelectionIndex());
+	AddLayer(GetSelectionIndex());
 }
 
 void STileLayerList::AddNewLayerBelow()
 {
-	AddLayer(/*bCollisionLayer=*/ false, GetSelectionIndex() + 1);
+	AddLayer(GetSelectionIndex() + 1);
 }
 
 int32 STileLayerList::GetSelectionIndex() const
diff --git a/Engine/Plugins/2D/Paper2D/Source/Paper2DEditor/Private/TileMapEditing/STileLayerList.h b/Engine/Plugins/2D/Paper2D/Source/Paper2DEditor/Private/TileMapEditing/STileLayerList.h
--- a/Engine/Plugins/2D/Paper2D/Source/Paper2DEditor/Private/TileMapEditing/STileLayerList.h
+++ b/Engine/Plugins/2D/Paper2D/Source/Paper2DEditor/Private/TileMapEditing/STileLayerList.h
@@ -34,7 +34,7 @@ protected:
 
 	static FText GenerateDuplicatedLayerName(const FString& InputNameRaw, UPaperTileMap* TileMap);
 
-	class UPaperTileLayer* AddLayer(bool bCollisionLayer, int32 InsertionIndex = INDEX_NONE);
+	class UPaperTileLayer* AddLayer(int32 InsertionIndex = INDEX_NONE);
 
 	// Moves a layer from OldIndex to NewIndex if both are valid, otherwise it does nothing silently
 	void ChangeLayerOrdering(int32 OldIndex, int32 NewIndex);
