Paper2D: Limited maximum tile map size to 1024x1024 (still far larger than would be recommended in a single component) [UE-13625]

diff --git a/Engine/Plugins/2D/Paper2D/Source/Paper2D/Classes/PaperTileMap.h b/Engine/Plugins/2D/Paper2D/Source/Paper2D/Classes/PaperTileMap.h
--- a/Engine/Plugins/2D/Paper2D/Source/Paper2D/Classes/PaperTileMap.h
+++ b/Engine/Plugins/2D/Paper2D/Source/Paper2D/Classes/PaperTileMap.h
@@ -34,11 +34,11 @@ class PAPER2D_API UPaperTileMap : public UObject
 	GENERATED_UCLASS_BODY()
 
 	// Width of map (in tiles)
-	UPROPERTY(Category=Setup, EditAnywhere, BlueprintReadOnly, meta=(UIMin=1, ClampMin=1))
+	UPROPERTY(Category=Setup, EditAnywhere, BlueprintReadOnly, meta=(UIMin=1, ClampMin=1, ClampMax=1024))
 	int32 MapWidth;
 
 	// Height of map (in tiles)
-	UPROPERTY(Category=Setup, EditAnywhere, BlueprintReadOnly, meta=(UIMin=1, ClampMin=1))
+	UPROPERTY(Category=Setup, EditAnywhere, BlueprintReadOnly, meta=(UIMin=1, ClampMin=1, ClampMax=1024))
 	int32 MapHeight;
 
 	// Width of one tile (in pixels)
