Paper2D: Added a help prompt to the tile map edit mode [UE-13639]

diff --git a/Engine/Plugins/2D/Paper2D/Source/Paper2DEditor/Private/TileMapEditing/EdModeTileMap.cpp b/Engine/Plugins/2D/Paper2D/Source/Paper2DEditor/Private/TileMapEditing/EdModeTileMap.cpp
--- a/Engine/Plugins/2D/Paper2D/Source/Paper2DEditor/Private/TileMapEditing/EdModeTileMap.cpp
+++ b/Engine/Plugins/2D/Paper2D/Source/Paper2DEditor/Private/TileMapEditing/EdModeTileMap.cpp
@@ -341,6 +341,23 @@ void FEdModeTileMap::Render(const FSceneView* View, FViewport* Viewport, FPrimit
 
 void FEdModeTileMap::DrawHUD(FEditorViewportClient* ViewportClient, FViewport* Viewport, const FSceneView* View, FCanvas* Canvas)
 {
+	const FIntRect CanvasRect = Canvas->GetViewRect();
+
+	// Display a help message to exit the editing mode (but only when in the world, not in individual asset editors)
+	if (GetModeManager() == &GLevelEditorModeTools())
+	{
+		static const FText EdModeHelp = LOCTEXT("TileMapEditorModeHelp", "Editing a tile map, press Escape to exit this mode");
+		static const FString EdModeHelpAsString = EdModeHelp.ToString();
+
+		int32 XL;
+		int32 YL;
+		StringSize(GEngine->GetLargeFont(), XL, YL, *EdModeHelpAsString);
+
+		const float DrawX = FMath::Floor(CanvasRect.Min.X + (CanvasRect.Width() - XL) * 0.5f);
+		const float DrawY = 30.0f;
+		Canvas->DrawShadowedString(DrawX, DrawY, *EdModeHelpAsString, GEngine->GetLargeFont(), FLinearColor::White);
+	}
+
 	bool bDrawToolDescription = false;
 
 	static const FText UnknownTool = LOCTEXT("NoTool", "No tool selected");
@@ -427,7 +444,6 @@ void FEdModeTileMap::DrawHUD(FEditorViewportClient* ViewportClient, FViewport* V
 			int32 YL;
 			StringSize(GEngine->GetLargeFont(), XL, YL, *CursorDescriptionString);
 
-			FIntRect CanvasRect = Canvas->GetViewRect();
 			const float DrawX = CanvasRect.Min.X + (CanvasRect.Width() - XL) * 0.5f;
 			const float DrawY = CanvasRect.Max.Y - 10.0f - YL;
 			Canvas->DrawShadowedString(DrawX, DrawY, *CursorDescriptionString, GEngine->GetLargeFont(), FLinearColor::White);
