Paper2D: Snapped some canvas text drawing to integers to reduce the percieved change in font 'weight' as the length of a string changes

diff --git a/Engine/Plugins/2D/Paper2D/Source/Paper2DEditor/Private/TileMapEditing/EdModeTileMap.cpp b/Engine/Plugins/2D/Paper2D/Source/Paper2DEditor/Private/TileMapEditing/EdModeTileMap.cpp
--- a/Engine/Plugins/2D/Paper2D/Source/Paper2DEditor/Private/TileMapEditing/EdModeTileMap.cpp
+++ b/Engine/Plugins/2D/Paper2D/Source/Paper2DEditor/Private/TileMapEditing/EdModeTileMap.cpp
@@ -393,7 +393,9 @@ void FEdModeTileMap::DrawHUD(FEditorViewportClient* ViewportClient, FViewport* V
 			int32 XL;
 			int32 YL;
 			StringSize(GEngine->GetLargeFont(), XL, YL, *ToolDescriptionString);
-			Canvas->DrawShadowedString(ScreenSpacePreviewLocation.X, ScreenSpacePreviewLocation.Y - YL, *ToolDescriptionString, GEngine->GetLargeFont(), ToolPromptColor);
+			const float DrawX = FMath::Floor(ScreenSpacePreviewLocation.X);
+			const float DrawY = FMath::Floor(ScreenSpacePreviewLocation.Y - YL);
+			Canvas->DrawShadowedString(DrawX, DrawY, *ToolDescriptionString, GEngine->GetLargeFont(), ToolPromptColor);
 		}
 	}
 
@@ -444,8 +446,8 @@ void FEdModeTileMap::DrawHUD(FEditorViewportClient* ViewportClient, FViewport* V
 			int32 YL;
 			StringSize(GEngine->GetLargeFont(), XL, YL, *CursorDescriptionString);
 
-			const float DrawX = CanvasRect.Min.X + (CanvasRect.Width() - XL) * 0.5f;
-			const float DrawY = CanvasRect.Max.Y - 10.0f - YL;
+			const float DrawX = FMath::Floor(CanvasRect.Min.X + (CanvasRect.Width() - XL) * 0.5f);
+			const float DrawY = FMath::Floor(CanvasRect.Max.Y - 10.0f - YL);
 			Canvas->DrawShadowedString(DrawX, DrawY, *CursorDescriptionString, GEngine->GetLargeFont(), FLinearColor::White);
 		}
 	}
diff --git a/Engine/Plugins/2D/Paper2D/Source/Paper2DEditor/Private/TileSetEditor.cpp b/Engine/Plugins/2D/Paper2D/Source/Paper2DEditor/Private/TileSetEditor.cpp
--- a/Engine/Plugins/2D/Paper2D/Source/Paper2DEditor/Private/TileSetEditor.cpp
+++ b/Engine/Plugins/2D/Paper2D/Source/Paper2DEditor/Private/TileSetEditor.cpp
@@ -121,7 +121,9 @@ void FTileSetEditorViewportClient::Draw(FViewport* Viewport, FCanvas* Canvas)
 		int32 XL;
 		int32 YL;
 		StringSize(GEngine->GetLargeFont(), XL, YL, *TileIndexString);
-		Canvas->DrawShadowedString(4, Viewport->GetSizeXY().Y - YL - 4, *TileIndexString, GEngine->GetLargeFont(), FLinearColor::White);
+		const float DrawX = 4.0f;
+		const float DrawY = FMath::Floor(Viewport->GetSizeXY().Y - YL - 4.0f);
+		Canvas->DrawShadowedString(DrawX, DrawY, *TileIndexString, GEngine->GetLargeFont(), FLinearColor::White);
 	}
 }
 
