Sprite Editor: Added support for deleting selected sockets [UE-13221]

diff --git a/Engine/Plugins/2D/Paper2D/Source/Paper2D/Classes/PaperSprite.h b/Engine/Plugins/2D/Paper2D/Source/Paper2D/Classes/PaperSprite.h
--- a/Engine/Plugins/2D/Paper2D/Source/Paper2D/Classes/PaperSprite.h
+++ b/Engine/Plugins/2D/Paper2D/Source/Paper2D/Classes/PaperSprite.h
@@ -259,6 +259,9 @@ public:
 	// Make sure all socket names are valid
 	// All duplicate / empty names will be made unique
 	void ValidateSocketNames();
+
+	// Removes the specified socket
+	void RemoveSocket(FName SocketName);
 #endif
 
 	// Return the scaling factor between pixels and Unreal units (cm)
diff --git a/Engine/Plugins/2D/Paper2D/Source/Paper2D/Private/PaperSprite.cpp b/Engine/Plugins/2D/Paper2D/Source/Paper2D/Private/PaperSprite.cpp
--- a/Engine/Plugins/2D/Paper2D/Source/Paper2D/Private/PaperSprite.cpp
+++ b/Engine/Plugins/2D/Paper2D/Source/Paper2D/Private/PaperSprite.cpp
@@ -1711,6 +1711,13 @@ void UPaperSprite::ValidateSocketNames()
 }
 #endif
 
+#if WITH_EDITOR
+void UPaperSprite::RemoveSocket(FName SocketNameToDelete)
+{
+	Sockets.RemoveAll([=](const FPaperSpriteSocket& Socket){ return Socket.SocketName == SocketNameToDelete; });
+}
+#endif
+
 #if WITH_EDITOR
 void UPaperSprite::Serialize(FArchive& Ar)
 {
diff --git a/Engine/Plugins/2D/Paper2D/Source/Paper2DEditor/Private/PaperEditorShared/AssetEditorSelectedItem.h b/Engine/Plugins/2D/Paper2D/Source/Paper2DEditor/Private/PaperEditorShared/AssetEditorSelectedItem.h
--- a/Engine/Plugins/2D/Paper2D/Source/Paper2DEditor/Private/PaperEditorShared/AssetEditorSelectedItem.h
+++ b/Engine/Plugins/2D/Paper2D/Source/Paper2DEditor/Private/PaperEditorShared/AssetEditorSelectedItem.h
@@ -57,6 +57,17 @@ public:
 		return false;
 	}
 
+	// Can this item be deleted?
+	virtual bool CanBeDeleted() const
+	{
+		return true;
+	}
+
+	// Delete this item
+	virtual void DeleteThisItem()
+	{
+	}
+
 	virtual ~FSelectedItem() {}
 };
 
diff --git a/Engine/Plugins/2D/Paper2D/Source/Paper2DEditor/Private/PaperEditorShared/SocketEditing.cpp b/Engine/Plugins/2D/Paper2D/Source/Paper2DEditor/Private/PaperEditorShared/SocketEditing.cpp
--- a/Engine/Plugins/2D/Paper2D/Source/Paper2DEditor/Private/PaperEditorShared/SocketEditing.cpp
+++ b/Engine/Plugins/2D/Paper2D/Source/Paper2DEditor/Private/PaperEditorShared/SocketEditing.cpp
@@ -93,6 +93,23 @@ FVector FSpriteSelectedSocket::GetWorldPos() const
 	return FVector::ZeroVector;
 }
 
+void FSpriteSelectedSocket::DeleteThisItem()
+{
+	if (UPrimitiveComponent* PreviewComponent = PreviewComponentPtr.Get())
+	{
+		if (PreviewComponent->DoesSocketExist(SocketName))
+		{
+			if (UPaperSpriteComponent* SpriteComponent = Cast<UPaperSpriteComponent>(PreviewComponent))
+			{
+				if (UPaperSprite* SpriteAsset = SpriteComponent->GetSprite())
+				{
+					SpriteAsset->RemoveSocket(SocketName);
+				}
+			}
+		}
+	}
+}
+
 //////////////////////////////////////////////////////////////////////////
 // FSocketEditingHelper
 
diff --git a/Engine/Plugins/2D/Paper2D/Source/Paper2DEditor/Private/PaperEditorShared/SocketEditing.h b/Engine/Plugins/2D/Paper2D/Source/Paper2DEditor/Private/PaperEditorShared/SocketEditing.h
--- a/Engine/Plugins/2D/Paper2D/Source/Paper2DEditor/Private/PaperEditorShared/SocketEditing.h
+++ b/Engine/Plugins/2D/Paper2D/Source/Paper2DEditor/Private/PaperEditorShared/SocketEditing.h
@@ -22,6 +22,7 @@ public:
 	virtual bool Equals(const FSelectedItem& OtherItem) const override;
 	virtual void ApplyDelta(const FVector2D& Delta, const FRotator& Rotation, const FVector& Scale3D, FWidget::EWidgetMode MoveMode) override;
 	FVector GetWorldPos() const override;
+	virtual void DeleteThisItem() override;
 	// End of FSelectedItem interface
 };
 
diff --git a/Engine/Plugins/2D/Paper2D/Source/Paper2DEditor/Private/PaperEditorShared/SpriteGeometryEditMode.cpp b/Engine/Plugins/2D/Paper2D/Source/Paper2DEditor/Private/PaperEditorShared/SpriteGeometryEditMode.cpp
--- a/Engine/Plugins/2D/Paper2D/Source/Paper2DEditor/Private/PaperEditorShared/SpriteGeometryEditMode.cpp
+++ b/Engine/Plugins/2D/Paper2D/Source/Paper2DEditor/Private/PaperEditorShared/SpriteGeometryEditMode.cpp
@@ -380,7 +380,7 @@ void FSpriteGeometryEditMode::BindCommands(TSharedPtr<FUICommandList> CommandLis
 	// Geometry editing commands
 	CommandList->MapAction(
 		Commands.DeleteSelection,
-		FExecuteAction::CreateRaw(&SpriteGeometryHelper, &FSpriteGeometryEditingHelper::DeleteSelectedItems, /*bShouldTransact=*/ true),
+		FExecuteAction::CreateRaw(&SpriteGeometryHelper, &FSpriteGeometryEditingHelper::DeleteSelectedItems),
 		FCanExecuteAction::CreateRaw(&SpriteGeometryHelper, &FSpriteGeometryEditingHelper::CanDeleteSelection));
 	CommandList->MapAction(
 		Commands.AddBoxShape,
diff --git a/Engine/Plugins/2D/Paper2D/Source/Paper2DEditor/Private/PaperEditorShared/SpriteGeometryEditing.cpp b/Engine/Plugins/2D/Paper2D/Source/Paper2DEditor/Private/PaperEditorShared/SpriteGeometryEditing.cpp
--- a/Engine/Plugins/2D/Paper2D/Source/Paper2DEditor/Private/PaperEditorShared/SpriteGeometryEditing.cpp
+++ b/Engine/Plugins/2D/Paper2D/Source/Paper2DEditor/Private/PaperEditorShared/SpriteGeometryEditing.cpp
@@ -31,11 +31,6 @@ void FSpriteSelectionHelper::ClearSelectionSet()
 	SelectedItemSet.Empty();
 }
 
-void FSpriteSelectionHelper::DeleteSelectedItems(bool bShouldTransact)
-{
-	//@TODO: Allow deletion of sockets, etc... here
-}
-
 void FSpriteSelectionHelper::SelectItem(TSharedPtr<FSelectedItem> NewItem)
 {
 	SelectedItemSet.Add(NewItem);
@@ -207,75 +202,102 @@ void FSpriteGeometryEditingHelper::ClearSelectionSet()
 	EditorContext->InvalidateViewportAndHitProxies();
 }
 
-void FSpriteGeometryEditingHelper::DeleteSelectedItems(bool bShouldTransact)
+void FSpriteGeometryEditingHelper::DeleteSelectedItems()
 {
-	// Make a set of polygon-vertices "pairs" to delete
+	// Determine which vertices or entire shapes should be deleted
 	TSet<FShapeVertexPair> CompositeIndicesSet;
 	TSet<int32> ShapesToDeleteSet;
 
-	FSpriteGeometryCollection& Geometry = GetGeometryChecked();
-
-	for (TSharedPtr<FSelectedItem> SelectionIt : GetSelectionSet())
+	if (IsEditingGeometry())
 	{
-		if (const FSpriteSelectedVertex* SelectedVertex = SelectionIt->CastTo<const FSpriteSelectedVertex>(FSelectionTypes::Vertex))
+		FSpriteGeometryCollection& Geometry = GetGeometryChecked();
+
+		for (TSharedPtr<FSelectedItem> SelectionIt : GetSelectionSet())
 		{
-			CompositeIndicesSet.Add(FShapeVertexPair(SelectedVertex->ShapeIndex, SelectedVertex->VertexIndex));
+			if (const FSpriteSelectedVertex* SelectedVertex = SelectionIt->CastTo<const FSpriteSelectedVertex>(FSelectionTypes::Vertex))
+			{
+				CompositeIndicesSet.Add(FShapeVertexPair(SelectedVertex->ShapeIndex, SelectedVertex->VertexIndex));
 
-			if (SelectedVertex->IsA(FSelectionTypes::Edge)) // add the "next" point for the edge
+				if (SelectedVertex->IsA(FSelectionTypes::Edge)) // add the "next" point for the edge
+				{
+					const int32 NextIndex = (SelectedVertex->VertexIndex + 1) % Geometry.Shapes[SelectedVertex->ShapeIndex].Vertices.Num();
+					CompositeIndicesSet.Add(FShapeVertexPair(SelectedVertex->ShapeIndex, NextIndex));
+				}
+			}
+			else if (const FSpriteSelectedShape* SelectedShape = SelectionIt->CastTo<const FSpriteSelectedShape>(FSelectionTypes::GeometryShape))
 			{
-				const int32 NextIndex = (SelectedVertex->VertexIndex + 1) % Geometry.Shapes[SelectedVertex->ShapeIndex].Vertices.Num();
-				CompositeIndicesSet.Add(FShapeVertexPair(SelectedVertex->ShapeIndex, NextIndex));
+				ShapesToDeleteSet.Add(SelectedShape->ShapeIndex);
 			}
 		}
-		else if (const FSpriteSelectedShape* SelectedShape = SelectionIt->CastTo<const FSpriteSelectedShape>(FSelectionTypes::GeometryShape))
+	}
+
+	// See if anything else can be deleted
+	bool bCanDeleteNonGeometry = false;
+	for (const TSharedPtr<FSelectedItem> SelectedItem : GetSelectionSet())
+	{
+		if (SelectedItem->CanBeDeleted())
 		{
-			ShapesToDeleteSet.Add(SelectedShape->ShapeIndex);
+			bCanDeleteNonGeometry = true;
+			break;
 		}
 	}
 
-	if ((CompositeIndicesSet.Num() > 0) || (ShapesToDeleteSet.Num() > 0))
+	// Now delete the stuff that was selected in the correct order so that indices aren't messed up
+	const bool bDeletingGeometry = (CompositeIndicesSet.Num() > 0) || (ShapesToDeleteSet.Num() > 0);
+	if (bDeletingGeometry || bCanDeleteNonGeometry)
 	{
 		EditorContext->BeginTransaction(LOCTEXT("DeleteSelectionTransaction", "Delete Selection"));
+		EditorContext->MarkTransactionAsDirty();
 
-		// Delete the selected vertices first, as they may cause entire shapes to need to be deleted (sort so we delete from the back first)
-		TArray<FShapeVertexPair> CompositeIndices = CompositeIndicesSet.Array();
-		CompositeIndices.Sort([](const FShapeVertexPair& A, const FShapeVertexPair& B) { return (A.VertexIndex > B.VertexIndex); });
-		for (const FShapeVertexPair& Composite : CompositeIndices)
+		if (bDeletingGeometry)
 		{
-			const int32 ShapeIndex = Composite.ShapeIndex;
-			const int32 VertexIndex = Composite.VertexIndex;
-			if (DeleteVertexInPolygonInternal(Geometry, ShapeIndex, VertexIndex))
+			FSpriteGeometryCollection& Geometry = GetGeometryChecked();
+
+			// Delete the selected vertices first, as they may cause entire shapes to need to be deleted (sort so we delete from the back first)
+			TArray<FShapeVertexPair> CompositeIndices = CompositeIndicesSet.Array();
+			CompositeIndices.Sort([](const FShapeVertexPair& A, const FShapeVertexPair& B) { return (A.VertexIndex > B.VertexIndex); });
+			for (const FShapeVertexPair& Composite : CompositeIndices)
+			{
+				const int32 ShapeIndex = Composite.ShapeIndex;
+				const int32 VertexIndex = Composite.VertexIndex;
+				if (DeleteVertexInPolygonInternal(Geometry, ShapeIndex, VertexIndex))
+				{
+					ShapesToDeleteSet.Add(ShapeIndex);
+				}
+			}
+
+			// Delete the selected shapes (plus any shapes that became empty due to selected vertices)
+			if (ShapesToDeleteSet.Num() > 0)
 			{
-				ShapesToDeleteSet.Add(ShapeIndex);
+				// Sort so we delete from the back first
+				TArray<int32> ShapesToDeleteIndicies = ShapesToDeleteSet.Array();
+				ShapesToDeleteIndicies.Sort([](const int32& A, const int32& B) { return (A > B); });
+				for (const int32 ShapeToDeleteIndex : ShapesToDeleteIndicies)
+				{
+					Geometry.Shapes.RemoveAt(ShapeToDeleteIndex);
+				}
 			}
+
+			Geometry.GeometryType = ESpritePolygonMode::FullyCustom;
 		}
 
-		// Delete the selected shapes (plus any shapes that became empty due to selected vertices)
-		if (ShapesToDeleteSet.Num() > 0)
+		// Delete everything else
+		if (bCanDeleteNonGeometry)
 		{
-			// Sort so we delete from the back first
-			TArray<int32> ShapesToDeleteIndicies = ShapesToDeleteSet.Array();
-			ShapesToDeleteIndicies.Sort([](const int32& A, const int32& B) { return (A > B); });
-			for (const int32 ShapeToDeleteIndex : ShapesToDeleteIndicies)
+			for (TSharedPtr<FSelectedItem> SelectedItem : GetSelectionSet())
 			{
-				Geometry.Shapes.RemoveAt(ShapeToDeleteIndex);
+				if (SelectedItem->CanBeDeleted())
+				{
+					SelectedItem->DeleteThisItem();
+				}
 			}
 		}
 
-		Geometry.GeometryType = ESpritePolygonMode::FullyCustom;
-		EditorContext->MarkTransactionAsDirty();
-
-		//@TODO: Should allow deleting other things when geometry is also selected!
-		//FSpriteGeometryEditingHelper::DeleteSelectedItems(/*bShouldTransact=*/ false);
-
-		ClearSelectionSet();
 		EditorContext->EndTransaction();
-		ResetAddPolygonMode();
-	}
-	else
-	{
-		FSpriteGeometryEditingHelper::DeleteSelectedItems(/*bShouldTransact=*/ true);
 	}
+
+	ClearSelectionSet();
+	ResetAddPolygonMode();
 }
 
 void FSpriteGeometryEditingHelper::AddReferencedObjects(FReferenceCollector& Collector)
diff --git a/Engine/Plugins/2D/Paper2D/Source/Paper2DEditor/Private/PaperEditorShared/SpriteGeometryEditing.h b/Engine/Plugins/2D/Paper2D/Source/Paper2DEditor/Private/PaperEditorShared/SpriteGeometryEditing.h
--- a/Engine/Plugins/2D/Paper2D/Source/Paper2DEditor/Private/PaperEditorShared/SpriteGeometryEditing.h
+++ b/Engine/Plugins/2D/Paper2D/Source/Paper2DEditor/Private/PaperEditorShared/SpriteGeometryEditing.h
@@ -58,8 +58,6 @@ public:
 
 	virtual void ClearSelectionSet();
 
-	virtual void DeleteSelectedItems(bool bShouldTransact = true);
-
 	virtual bool CanDeleteSelection() const;
 
 	~FSpriteSelectionHelper()
@@ -81,13 +79,13 @@ public:
 
 	// FSpriteSelectionHelper interface
 	virtual void ClearSelectionSet() override;
-	virtual void DeleteSelectedItems(bool bShouldTransact = true) override;
 	// End of FSpriteSelectionHelper interface
 
 	// FGCObject interface
 	virtual void AddReferencedObjects(FReferenceCollector& Collector) override;
 	// End of FGCObject interface
 
+	void DeleteSelectedItems();
 	void SetEditorContext(class ISpriteSelectionContext* InNewEditorContext);
 	class ISpriteSelectionContext* GetEditorContext() const;
 
