Sprite Editor: Added a selection highlight to sockets [UE-13652]

diff --git a/Engine/Plugins/2D/Paper2D/Source/Paper2DEditor/Private/FlipbookEditor/FlipbookEditorViewportClient.cpp b/Engine/Plugins/2D/Paper2D/Source/Paper2DEditor/Private/FlipbookEditor/FlipbookEditorViewportClient.cpp
--- a/Engine/Plugins/2D/Paper2D/Source/Paper2DEditor/Private/FlipbookEditor/FlipbookEditorViewportClient.cpp
+++ b/Engine/Plugins/2D/Paper2D/Source/Paper2DEditor/Private/FlipbookEditor/FlipbookEditorViewportClient.cpp
@@ -61,7 +61,7 @@ void FFlipbookEditorViewportClient::DrawCanvas(FViewport& Viewport, FSceneView&
 
 	if (bShowSockets)
 	{
-		FSocketEditingHelper::DrawSocketNames(AnimatedRenderComponent.Get(), Viewport, View, Canvas);
+		FSocketEditingHelper::DrawSocketNames(nullptr, AnimatedRenderComponent.Get(), Viewport, View, Canvas);
 	}
 }
 
@@ -76,7 +76,7 @@ void FFlipbookEditorViewportClient::Draw(const FSceneView* View, FPrimitiveDrawI
 
 	if (bShowSockets)
 	{
-		FSocketEditingHelper::DrawSockets(AnimatedRenderComponent.Get(), View, PDI);
+		FSocketEditingHelper::DrawSockets(nullptr, AnimatedRenderComponent.Get(), View, PDI);
 	}
 }
 
diff --git a/Engine/Plugins/2D/Paper2D/Source/Paper2DEditor/Private/PaperEditorShared/SocketEditing.cpp b/Engine/Plugins/2D/Paper2D/Source/Paper2DEditor/Private/PaperEditorShared/SocketEditing.cpp
--- a/Engine/Plugins/2D/Paper2D/Source/Paper2DEditor/Private/PaperEditorShared/SocketEditing.cpp
+++ b/Engine/Plugins/2D/Paper2D/Source/Paper2DEditor/Private/PaperEditorShared/SocketEditing.cpp
@@ -2,6 +2,7 @@
 
 #include "Paper2DEditorPrivatePCH.h"
 #include "SocketEditing.h"
+#include "SpriteGeometryEditMode.h"
 
 //////////////////////////////////////////////////////////////////////////
 // HSpriteSelectableObjectHitProxy
@@ -113,14 +114,15 @@ void FSpriteSelectedSocket::DeleteThisItem()
 //////////////////////////////////////////////////////////////////////////
 // FSocketEditingHelper
 
-void FSocketEditingHelper::DrawSockets(UPrimitiveComponent* PreviewComponent, const FSceneView* View, FPrimitiveDrawInterface* PDI)
+void FSocketEditingHelper::DrawSockets(FSpriteGeometryEditMode* GeometryEditMode, UPrimitiveComponent* PreviewComponent, const FSceneView* View, FPrimitiveDrawInterface* PDI)
 {
 	if (PreviewComponent != nullptr)
 	{
 		const bool bIsHitTesting = PDI->IsHitTesting();
 
 		const float DiamondSize = 5.0f;
-		const FColor DiamondColor(255, 128, 128);
+		const FColor UnselectedDiamondColor(255, 128, 128);
+		const FColor SelectedDiamondColor(FColor::White);
 
 		TArray<FComponentSocketDescription> SocketList;
 		PreviewComponent->QuerySupportedSockets(/*out*/ SocketList);
@@ -135,6 +137,9 @@ void FSocketEditingHelper::DrawSockets(UPrimitiveComponent* PreviewComponent, co
 				PDI->SetHitProxy(new HSpriteSelectableObjectHitProxy(Data));
 			}
 
+			const bool bIsSelected = (GeometryEditMode != nullptr) && GeometryEditMode->IsSocketSelected(Socket.Name);
+			const FColor& DiamondColor = bIsSelected ? SelectedDiamondColor : UnselectedDiamondColor;
+
 			const FMatrix SocketTM = PreviewComponent->GetSocketTransform(Socket.Name, RTS_World).ToMatrixWithScale();
 			DrawWireDiamond(PDI, SocketTM, DiamondSize, DiamondColor, SDPG_Foreground);
 
@@ -146,14 +151,15 @@ void FSocketEditingHelper::DrawSockets(UPrimitiveComponent* PreviewComponent, co
 	}
 }
 
-void FSocketEditingHelper::DrawSocketNames(UPrimitiveComponent* PreviewComponent, FViewport& Viewport, FSceneView& View, FCanvas& Canvas)
+void FSocketEditingHelper::DrawSocketNames(FSpriteGeometryEditMode* GeometryEditMode, UPrimitiveComponent* PreviewComponent, FViewport& Viewport, FSceneView& View, FCanvas& Canvas)
 {
 	if (PreviewComponent != nullptr)
 	{
 		const int32 HalfX = Viewport.GetSizeXY().X / 2;
 		const int32 HalfY = Viewport.GetSizeXY().Y / 2;
 
-		const FColor SocketNameColor(255, 196, 196);
+		const FColor UnselectedSocketNameColor(255, 196, 196);
+		const FColor SelectedSocketNameColor(FColor::White);
 
 		TArray<FComponentSocketDescription> SocketList;
 		PreviewComponent->QuerySupportedSockets(/*out*/ SocketList);
@@ -168,7 +174,10 @@ void FSocketEditingHelper::DrawSocketNames(UPrimitiveComponent* PreviewComponent
 				const int32 XPos = HalfX + (HalfX * Proj.X);
 				const int32 YPos = HalfY + (HalfY * (-Proj.Y));
 
-				FCanvasTextItem Msg(FVector2D(XPos, YPos), FText::FromString(Socket.Name.ToString()), GEngine->GetMediumFont(), SocketNameColor);
+				const bool bIsSelected = (GeometryEditMode != nullptr) && GeometryEditMode->IsSocketSelected(Socket.Name);
+				const FColor& SocketColor = bIsSelected ? SelectedSocketNameColor : UnselectedSocketNameColor;
+
+				FCanvasTextItem Msg(FVector2D(XPos, YPos), FText::FromString(Socket.Name.ToString()), GEngine->GetMediumFont(), SocketColor);
 				Msg.EnableShadow(FLinearColor::Black);
 				Canvas.DrawItem(Msg);
 
diff --git a/Engine/Plugins/2D/Paper2D/Source/Paper2DEditor/Private/PaperEditorShared/SocketEditing.h b/Engine/Plugins/2D/Paper2D/Source/Paper2DEditor/Private/PaperEditorShared/SocketEditing.h
--- a/Engine/Plugins/2D/Paper2D/Source/Paper2DEditor/Private/PaperEditorShared/SocketEditing.h
+++ b/Engine/Plugins/2D/Paper2D/Source/Paper2DEditor/Private/PaperEditorShared/SocketEditing.h
@@ -32,6 +32,6 @@ public:
 class FSocketEditingHelper
 {
 public:
-	static void DrawSockets(UPrimitiveComponent* PreviewComponent, const FSceneView* View, FPrimitiveDrawInterface* PDI);
-	static void DrawSocketNames(UPrimitiveComponent* PreviewComponent, FViewport& Viewport, FSceneView& View, FCanvas& Canvas);
+	static void DrawSockets(class FSpriteGeometryEditMode* GeometryEditMode, UPrimitiveComponent* PreviewComponent, const FSceneView* View, FPrimitiveDrawInterface* PDI);
+	static void DrawSocketNames(class FSpriteGeometryEditMode* GeometryEditMode, UPrimitiveComponent* PreviewComponent, FViewport& Viewport, FSceneView& View, FCanvas& Canvas);
 };
diff --git a/Engine/Plugins/2D/Paper2D/Source/Paper2DEditor/Private/PaperEditorShared/SpriteGeometryEditMode.cpp b/Engine/Plugins/2D/Paper2D/Source/Paper2DEditor/Private/PaperEditorShared/SpriteGeometryEditMode.cpp
--- a/Engine/Plugins/2D/Paper2D/Source/Paper2DEditor/Private/PaperEditorShared/SpriteGeometryEditMode.cpp
+++ b/Engine/Plugins/2D/Paper2D/Source/Paper2DEditor/Private/PaperEditorShared/SpriteGeometryEditMode.cpp
@@ -3,6 +3,7 @@
 #include "Paper2DEditorPrivatePCH.h"
 #include "SpriteGeometryEditMode.h"
 #include "../SpriteEditor/SpriteEditorCommands.h"
+#include "SocketEditing.h"
 
 #define LOCTEXT_NAMESPACE "PaperGeometryEditing"
 
@@ -313,58 +314,6 @@ void FSpriteGeometryEditMode::SetGeometryBeingEdited(FSpriteGeometryCollection*
 	SpriteGeometryHelper.SetGeometryBeingEdited(NewGeometryBeingEdited, bInAllowCircles, bInAllowSubtractivePolygons);
 	bIsMarqueeTracking = false;
 }
-// 
-// FVector2D FSpriteGeometryEditMode::SelectedItemConvertWorldSpaceDeltaToLocalSpace(const FVector& WorldSpaceDelta) const
-// {
-// 	const FVector ProjectionX = WorldSpaceDelta.ProjectOnTo(PaperAxisX);
-// 	const FVector ProjectionY = WorldSpaceDelta.ProjectOnTo(PaperAxisY);
-// 
-// 	const float XValue = FMath::Sign(ProjectionX | PaperAxisX) * ProjectionX.Size();
-// 	const float YValue = FMath::Sign(ProjectionY | PaperAxisY) * ProjectionY.Size();
-// 
-// 	return FVector2D(XValue, YValue);
-// }
-// 
-// FVector2D FSpriteGeometryEditMode::WorldSpaceToTextureSpace(const FVector& SourcePoint) const
-// {
-// 	const FVector ProjectionX = SourcePoint.ProjectOnTo(PaperAxisX);
-// 	const FVector ProjectionY = SourcePoint.ProjectOnTo(PaperAxisY);
-// 
-// 	const float XValue = FMath::Sign(ProjectionX | PaperAxisX) * ProjectionX.Size();
-// 	const float YValue = FMath::Sign(ProjectionY | PaperAxisY) * ProjectionY.Size();
-// 
-// 	return FVector2D(XValue, YValue);
-// }
-// 
-// FVector FSpriteGeometryEditMode::TextureSpaceToWorldSpace(const FVector2D& SourcePoint) const
-// {
-// 	return (SourcePoint.X * PaperAxisX) + (SourcePoint.Y * PaperAxisY);
-// }
-// 
-// float FSpriteGeometryEditMode::SelectedItemGetUnitsPerPixel() const
-// {
-// 	return 1.0f;
-// }
-// 
-// void FSpriteGeometryEditMode::BeginTransaction(const FText& SessionName)
-// {
-// 	// Dummy implementation does nothing
-// }
-// 
-// void FSpriteGeometryEditMode::MarkTransactionAsDirty()
-// {
-// 	// Dummy implementation does nothing
-// }
-// 
-// void FSpriteGeometryEditMode::EndTransaction()
-// {
-// 	// Dummy implementation does nothing
-// }
-// 
-// void FSpriteGeometryEditMode::InvalidateViewportAndHitProxies()
-// {
-// 	// Dummy implementation does nothing
-// }
 
 void FSpriteGeometryEditMode::BindCommands(TSharedPtr<FUICommandList> CommandList)
 {
@@ -456,6 +405,7 @@ void FSpriteGeometryEditMode::SelectVerticesInMarquee(FEditorViewportClient* Vie
 
 		const FBox2D QueryBounds(TextureSpaceStartPos, TextureSpaceEndPos);
 
+		// Check geometry
 		if (FSpriteGeometryCollection* Geometry = SpriteGeometryHelper.GetGeometryBeingEdited())
 		{
 			for (int32 ShapeIndex = 0; ShapeIndex < Geometry->Shapes.Num(); ++ShapeIndex)
@@ -494,6 +444,8 @@ void FSpriteGeometryEditMode::SelectVerticesInMarquee(FEditorViewportClient* Vie
 				}
 			}
 		}
+
+		//@TODO: Check other items (sockets/etc...)
 	}
 }
 
@@ -532,6 +484,22 @@ bool FSpriteGeometryEditMode::ProcessMarquee(FViewport* Viewport, FKey Key, EInp
 	return bMarqueeReady;
 }
 
+bool FSpriteGeometryEditMode::IsSocketSelected(FName SocketName) const
+{
+	for (TSharedPtr<FSelectedItem> SelectedItemPtr : SpriteGeometryHelper.GetSelectionSet())
+	{
+		if (const FSpriteSelectedSocket* SelectedSocket = SelectedItemPtr->CastTo<const FSpriteSelectedSocket>(FSpriteSelectedSocket::SocketTypeID))
+		{
+			if (SelectedSocket->SocketName == SocketName)
+			{
+				return true;
+			}
+		}
+	}
+
+	return false;
+}
+
 void FSpriteGeometryEditMode::DrawMarquee(FViewport& InViewport, const FSceneView& View, FCanvas& Canvas, const FLinearColor& Color)
 {
 	FVector2D MarqueeDelta = MarqueeEndPos - MarqueeStartPos;
diff --git a/Engine/Plugins/2D/Paper2D/Source/Paper2DEditor/Private/PaperEditorShared/SpriteGeometryEditMode.h b/Engine/Plugins/2D/Paper2D/Source/Paper2DEditor/Private/PaperEditorShared/SpriteGeometryEditMode.h
--- a/Engine/Plugins/2D/Paper2D/Source/Paper2DEditor/Private/PaperEditorShared/SpriteGeometryEditMode.h
+++ b/Engine/Plugins/2D/Paper2D/Source/Paper2DEditor/Private/PaperEditorShared/SpriteGeometryEditMode.h
@@ -8,7 +8,7 @@
 //////////////////////////////////////////////////////////////////////////
 // FSpriteGeometryEditMode
 
-class FSpriteGeometryEditMode : public FEdMode//, public ISpriteSelectionContext
+class FSpriteGeometryEditMode : public FEdMode
 {
 public:
 	static const FEditorModeID EM_SpriteGeometry;
@@ -41,23 +41,15 @@ public:
 	// Changes the geometry being edited (clears the selection set in the process)
 	void SetGeometryBeingEdited(FSpriteGeometryCollection* NewGeometryBeingEdited, bool bInAllowCircles, bool bInAllowSubtractivePolygons);
 
-// 	// Dummy implementation of ISpriteSelectionContext; users of this mode should provide their own and register via SetEditorContext
-// 	virtual FVector2D SelectedItemConvertWorldSpaceDeltaToLocalSpace(const FVector& WorldSpaceDelta) const override;
-// 	virtual FVector2D WorldSpaceToTextureSpace(const FVector& SourcePoint) const override;
-// 	virtual FVector TextureSpaceToWorldSpace(const FVector2D& SourcePoint) const override;
-// 	virtual float SelectedItemGetUnitsPerPixel() const override;
-// 	virtual void BeginTransaction(const FText& SessionName) override;
-// 	virtual void MarkTransactionAsDirty() override;
-// 	virtual void EndTransaction() override;
-// 	virtual void InvalidateViewportAndHitProxies() override;
-// 	// End of ISpriteSelectionContext interface
-// 
 	void BindCommands(TSharedPtr<FUICommandList> InCommandList);
 
 	FVector2D GetMarqueeStartPos() const { return MarqueeStartPos; }
 	FVector2D GetMarqueeEndPos() const { return MarqueeEndPos; }
 	bool ProcessMarquee(FViewport* Viewport, FKey Key, EInputEvent Event, bool bMarqueeStartModifierPressed);
 
+	// Returns true if the specified socket is selected
+	bool IsSocketSelected(FName SocketName) const;
+
 protected:
 	FBox2D BoundsForNewShapes;
 	FLinearColor GeometryVertexColor;
diff --git a/Engine/Plugins/2D/Paper2D/Source/Paper2DEditor/Private/SpriteEditor/SpriteEditorViewportClient.cpp b/Engine/Plugins/2D/Paper2D/Source/Paper2DEditor/Private/SpriteEditor/SpriteEditorViewportClient.cpp
--- a/Engine/Plugins/2D/Paper2D/Source/Paper2DEditor/Private/SpriteEditor/SpriteEditorViewportClient.cpp
+++ b/Engine/Plugins/2D/Paper2D/Source/Paper2DEditor/Private/SpriteEditor/SpriteEditorViewportClient.cpp
@@ -640,7 +640,8 @@ void FSpriteEditorViewportClient::DrawCanvas(FViewport& Viewport, FSceneView& Vi
 
 	if (bShowSockets && !IsInSourceRegionEditMode())
 	{
-		FSocketEditingHelper::DrawSocketNames(RenderSpriteComponent, Viewport, View, Canvas);
+		FSpriteGeometryEditMode* GeometryEditMode = ModeTools->GetActiveModeTyped<FSpriteGeometryEditMode>(FSpriteGeometryEditMode::EM_SpriteGeometry);
+		FSocketEditingHelper::DrawSocketNames(GeometryEditMode, RenderSpriteComponent, Viewport, View, Canvas);
 	}
 
 	FEditorViewportClient::DrawCanvas(Viewport, View, Canvas);
@@ -661,7 +662,8 @@ void FSpriteEditorViewportClient::Draw(const FSceneView* View, FPrimitiveDrawInt
 
 	if (bShowSockets && !IsInSourceRegionEditMode())
 	{
-		FSocketEditingHelper::DrawSockets(RenderSpriteComponent, View, PDI);
+		FSpriteGeometryEditMode* GeometryEditMode = ModeTools->GetActiveModeTyped<FSpriteGeometryEditMode>(FSpriteGeometryEditMode::EM_SpriteGeometry);
+		FSocketEditingHelper::DrawSockets(GeometryEditMode, RenderSpriteComponent, View, PDI);
 	}
 }
 
