Paper2D: Fixed some deprecation warnings about uses of FMath::Floor instead of FMath::FloorToFloat

diff --git a/Engine/Plugins/2D/Paper2D/Source/Paper2DEditor/Private/TileMapEditing/EdModeTileMap.cpp b/Engine/Plugins/2D/Paper2D/Source/Paper2DEditor/Private/TileMapEditing/EdModeTileMap.cpp
--- a/Engine/Plugins/2D/Paper2D/Source/Paper2DEditor/Private/TileMapEditing/EdModeTileMap.cpp
+++ b/Engine/Plugins/2D/Paper2D/Source/Paper2DEditor/Private/TileMapEditing/EdModeTileMap.cpp
@@ -353,7 +353,7 @@ void FEdModeTileMap::DrawHUD(FEditorViewportClient* ViewportClient, FViewport* V
 		int32 YL;
 		StringSize(GEngine->GetLargeFont(), XL, YL, *EdModeHelpAsString);
 
-		const float DrawX = FMath::Floor(CanvasRect.Min.X + (CanvasRect.Width() - XL) * 0.5f);
+		const float DrawX = FMath::FloorToFloat(CanvasRect.Min.X + (CanvasRect.Width() - XL) * 0.5f);
 		const float DrawY = 30.0f;
 		Canvas->DrawShadowedString(DrawX, DrawY, *EdModeHelpAsString, GEngine->GetLargeFont(), FLinearColor::White);
 	}
@@ -393,8 +393,8 @@ void FEdModeTileMap::DrawHUD(FEditorViewportClient* ViewportClient, FViewport* V
 			int32 XL;
 			int32 YL;
 			StringSize(GEngine->GetLargeFont(), XL, YL, *ToolDescriptionString);
-			const float DrawX = FMath::Floor(ScreenSpacePreviewLocation.X);
-			const float DrawY = FMath::Floor(ScreenSpacePreviewLocation.Y - YL);
+			const float DrawX = FMath::FloorToFloat(ScreenSpacePreviewLocation.X);
+			const float DrawY = FMath::FloorToFloat(ScreenSpacePreviewLocation.Y - YL);
 			Canvas->DrawShadowedString(DrawX, DrawY, *ToolDescriptionString, GEngine->GetLargeFont(), ToolPromptColor);
 		}
 	}
@@ -446,8 +446,8 @@ void FEdModeTileMap::DrawHUD(FEditorViewportClient* ViewportClient, FViewport* V
 			int32 YL;
 			StringSize(GEngine->GetLargeFont(), XL, YL, *CursorDescriptionString);
 
-			const float DrawX = FMath::Floor(CanvasRect.Min.X + (CanvasRect.Width() - XL) * 0.5f);
-			const float DrawY = FMath::Floor(CanvasRect.Max.Y - 10.0f - YL);
+			const float DrawX = FMath::FloorToFloat(CanvasRect.Min.X + (CanvasRect.Width() - XL) * 0.5f);
+			const float DrawY = FMath::FloorToFloat(CanvasRect.Max.Y - 10.0f - YL);
 			Canvas->DrawShadowedString(DrawX, DrawY, *CursorDescriptionString, GEngine->GetLargeFont(), FLinearColor::White);
 		}
 	}
diff --git a/Engine/Plugins/2D/Paper2D/Source/Paper2DEditor/Private/TileSetEditor.cpp b/Engine/Plugins/2D/Paper2D/Source/Paper2DEditor/Private/TileSetEditor.cpp
--- a/Engine/Plugins/2D/Paper2D/Source/Paper2DEditor/Private/TileSetEditor.cpp
+++ b/Engine/Plugins/2D/Paper2D/Source/Paper2DEditor/Private/TileSetEditor.cpp
@@ -122,7 +122,7 @@ void FTileSetEditorViewportClient::Draw(FViewport* Viewport, FCanvas* Canvas)
 		int32 YL;
 		StringSize(GEngine->GetLargeFont(), XL, YL, *TileIndexString);
 		const float DrawX = 4.0f;
-		const float DrawY = FMath::Floor(Viewport->GetSizeXY().Y - YL - 4.0f);
+		const float DrawY = FMath::FloorToFloat(Viewport->GetSizeXY().Y - YL - 4.0f);
 		Canvas->DrawShadowedString(DrawX, DrawY, *TileIndexString, GEngine->GetLargeFont(), FLinearColor::White);
 	}
 }
