Editor: Passed in a host command list to the details view in all uses of SSingleObjectDetailsPanel

diff --git a/Engine/Plugins/2D/Paper2D/Source/Paper2DEditor/Private/FlipbookEditor/FlipbookEditor.cpp b/Engine/Plugins/2D/Paper2D/Source/Paper2DEditor/Private/FlipbookEditor/FlipbookEditor.cpp
--- a/Engine/Plugins/2D/Paper2D/Source/Paper2DEditor/Private/FlipbookEditor/FlipbookEditor.cpp
+++ b/Engine/Plugins/2D/Paper2D/Source/Paper2DEditor/Private/FlipbookEditor/FlipbookEditor.cpp
@@ -191,7 +191,7 @@ public:
 	{
 		FlipbookEditorPtr = InFlipbookEditor;
 
-		SSingleObjectDetailsPanel::Construct(SSingleObjectDetailsPanel::FArguments(), /*bAutomaticallyObserveViaGetObjectToObserve=*/ true, /*bAllowSearch=*/ true);
+		SSingleObjectDetailsPanel::Construct(SSingleObjectDetailsPanel::FArguments().HostCommandList(InFlipbookEditor->GetToolkitCommands()), /*bAutomaticallyObserveViaGetObjectToObserve=*/ true, /*bAllowSearch=*/ true);
 	}
 
 	// SSingleObjectDetailsPanel interface
diff --git a/Engine/Plugins/2D/Paper2D/Source/Paper2DEditor/Private/SpriteEditor/SpriteEditor.cpp b/Engine/Plugins/2D/Paper2D/Source/Paper2DEditor/Private/SpriteEditor/SpriteEditor.cpp
--- a/Engine/Plugins/2D/Paper2D/Source/Paper2DEditor/Private/SpriteEditor/SpriteEditor.cpp
+++ b/Engine/Plugins/2D/Paper2D/Source/Paper2DEditor/Private/SpriteEditor/SpriteEditor.cpp
@@ -234,7 +234,7 @@ public:
 	{
 		SpriteEditorPtr = InSpriteEditor;
 
-		SSingleObjectDetailsPanel::Construct(SSingleObjectDetailsPanel::FArguments(), /*bAutomaticallyObserveViaGetObjectToObserve=*/ true, /*bAllowSearch=*/ true);
+		SSingleObjectDetailsPanel::Construct(SSingleObjectDetailsPanel::FArguments().HostCommandList(InSpriteEditor->GetToolkitCommands()), /*bAutomaticallyObserveViaGetObjectToObserve=*/ true, /*bAllowSearch=*/ true);
 
 		TAttribute<ESpriteEditorMode::Type> SpriteEditorMode = TAttribute<ESpriteEditorMode::Type>::Create(
 			TAttribute<ESpriteEditorMode::Type>::FGetter::CreateSP(InSpriteEditor.ToSharedRef(), &FSpriteEditor::GetCurrentMode));
diff --git a/Engine/Plugins/2D/Paper2D/Source/Paper2DEditor/Private/TileMapEditing/TileMapEditor.cpp b/Engine/Plugins/2D/Paper2D/Source/Paper2DEditor/Private/TileMapEditing/TileMapEditor.cpp
--- a/Engine/Plugins/2D/Paper2D/Source/Paper2DEditor/Private/TileMapEditing/TileMapEditor.cpp
+++ b/Engine/Plugins/2D/Paper2D/Source/Paper2DEditor/Private/TileMapEditing/TileMapEditor.cpp
@@ -175,7 +175,7 @@ public:
 	{
 		TileMapEditorPtr = InTileMapEditor;
 
-		SSingleObjectDetailsPanel::Construct(SSingleObjectDetailsPanel::FArguments());
+		SSingleObjectDetailsPanel::Construct(SSingleObjectDetailsPanel::FArguments().HostCommandList(InTileMapEditor->GetToolkitCommands()));
 	}
 
 	// SSingleObjectDetailsPanel interface
diff --git a/Engine/Plugins/2D/Paper2D/Source/Paper2DEditor/Private/TileSetEditor.cpp b/Engine/Plugins/2D/Paper2D/Source/Paper2DEditor/Private/TileSetEditor.cpp
--- a/Engine/Plugins/2D/Paper2D/Source/Paper2DEditor/Private/TileSetEditor.cpp
+++ b/Engine/Plugins/2D/Paper2D/Source/Paper2DEditor/Private/TileSetEditor.cpp
@@ -284,7 +284,7 @@ public:
 	{
 		TileSetEditorPtr = InTileSetEditor;
 
-		SSingleObjectDetailsPanel::Construct(SSingleObjectDetailsPanel::FArguments());
+		SSingleObjectDetailsPanel::Construct(SSingleObjectDetailsPanel::FArguments().HostCommandList(InTileSetEditor->GetToolkitCommands()));
 	}
 
 	// SSingleObjectDetailsPanel interface
diff --git a/Engine/Source/Editor/KismetWidgets/Private/SSingleObjectDetailsPanel.cpp b/Engine/Source/Editor/KismetWidgets/Private/SSingleObjectDetailsPanel.cpp
--- a/Engine/Source/Editor/KismetWidgets/Private/SSingleObjectDetailsPanel.cpp
+++ b/Engine/Source/Editor/KismetWidgets/Private/SSingleObjectDetailsPanel.cpp
@@ -15,6 +15,7 @@ void SSingleObjectDetailsPanel::Construct(const FArguments& InArgs, bool bAutoma
 	FPropertyEditorModule& EditModule = FModuleManager::Get().GetModuleChecked<FPropertyEditorModule>("PropertyEditor");
 
 	FDetailsViewArgs DetailsViewArgs(/*bUpdateFromSelection=*/ false, /*bLockable=*/ false, bAllowSearch, FDetailsViewArgs::HideNameArea, /*bHideSelectionTip=*/ true);
+	DetailsViewArgs.HostCommandList = InArgs._HostCommandList;
 
 	PropertyView = EditModule.CreateDetailView(DetailsViewArgs);
 	
diff --git a/Engine/Source/Editor/KismetWidgets/Public/SSingleObjectDetailsPanel.h b/Engine/Source/Editor/KismetWidgets/Public/SSingleObjectDetailsPanel.h
--- a/Engine/Source/Editor/KismetWidgets/Public/SSingleObjectDetailsPanel.h
+++ b/Engine/Source/Editor/KismetWidgets/Public/SSingleObjectDetailsPanel.h
@@ -9,6 +9,7 @@ class KISMETWIDGETS_API SSingleObjectDetailsPanel : public SCompoundWidget
 {
 public:
 	SLATE_BEGIN_ARGS(SSingleObjectDetailsPanel) {}
+		SLATE_ARGUMENT(TSharedPtr<class FUICommandList>, HostCommandList)
 	SLATE_END_ARGS()
 
 	void Construct(const FArguments& InArgs, bool bAutomaticallyObserveViaGetObjectToObserve = true, bool bAllowSearch = false);
diff --git a/Engine/Source/Editor/Persona/Private/AnimationMode/AnimationMode.cpp b/Engine/Source/Editor/Persona/Private/AnimationMode/AnimationMode.cpp
--- a/Engine/Source/Editor/Persona/Private/AnimationMode/AnimationMode.cpp
+++ b/Engine/Source/Editor/Persona/Private/AnimationMode/AnimationMode.cpp
@@ -28,7 +28,7 @@ public:
 	{
 		PersonaPtr = InPersona;
 
-		SSingleObjectDetailsPanel::Construct(SSingleObjectDetailsPanel::FArguments(), true, true);
+		SSingleObjectDetailsPanel::Construct(SSingleObjectDetailsPanel::FArguments().HostCommandList(InPersona->GetToolkitCommands()), true, true);
 	}
 
 	virtual EVisibility GetAssetDisplayNameVisibility() const
diff --git a/Engine/Source/Editor/Persona/Private/MeshMode/MeshMode.cpp b/Engine/Source/Editor/Persona/Private/MeshMode/MeshMode.cpp
--- a/Engine/Source/Editor/Persona/Private/MeshMode/MeshMode.cpp
+++ b/Engine/Source/Editor/Persona/Private/MeshMode/MeshMode.cpp
@@ -31,7 +31,7 @@ public:
 	{
 		PersonaPtr = InPersona;
 
-		SSingleObjectDetailsPanel::Construct(SSingleObjectDetailsPanel::FArguments(), true, true);
+		SSingleObjectDetailsPanel::Construct(SSingleObjectDetailsPanel::FArguments().HostCommandList(InPersona->GetToolkitCommands()), true, true);
 
 		PropertyView->SetGenericLayoutDetailsDelegate( FOnGetDetailCustomizationInstance::CreateStatic( &FPersonaMeshDetails::MakeInstance, InPersona ) );
 	}
diff --git a/Engine/Source/Editor/Persona/Private/Persona.cpp b/Engine/Source/Editor/Persona/Private/Persona.cpp
--- a/Engine/Source/Editor/Persona/Private/Persona.cpp
+++ b/Engine/Source/Editor/Persona/Private/Persona.cpp
@@ -100,7 +100,7 @@ public:
 	{
 		PersonaPtr = InPersona;
 
-		SSingleObjectDetailsPanel::Construct(SSingleObjectDetailsPanel::FArguments(), /*bAutomaticallyObserveViaGetObjectToObserve*/ true, /*bAllowSearch*/ true);
+		SSingleObjectDetailsPanel::Construct(SSingleObjectDetailsPanel::FArguments().HostCommandList(InPersona->GetToolkitCommands()), /*bAutomaticallyObserveViaGetObjectToObserve*/ true, /*bAllowSearch*/ true);
 
 		PropertyView->SetIsPropertyEditingEnabledDelegate(FIsPropertyEditingEnabled::CreateStatic([] { return !GIntraFrameDebuggingGameThread; }));
 	}
