Tile Map Editor: Appended the command list for the tile layer list to the host command list from the details panel, allowing key bindings to work
- Added some more commands to the context menu
[UE-13176]

diff --git a/Engine/Plugins/2D/Paper2D/Source/Paper2DEditor/Private/TileMapEditing/PaperTileMapDetailsCustomization.cpp b/Engine/Plugins/2D/Paper2D/Source/Paper2DEditor/Private/TileMapEditing/PaperTileMapDetailsCustomization.cpp
--- a/Engine/Plugins/2D/Paper2D/Source/Paper2DEditor/Private/TileMapEditing/PaperTileMapDetailsCustomization.cpp
+++ b/Engine/Plugins/2D/Paper2D/Source/Paper2DEditor/Private/TileMapEditing/PaperTileMapDetailsCustomization.cpp
@@ -133,10 +133,12 @@ void FPaperTileMapDetailsCustomization::CustomizeDetails(IDetailLayoutBuilder& D
 		TileMapCategory.AddProperty(GET_MEMBER_NAME_CHECKED(UPaperTileMapComponent, TileMap));
 	}
 
-	//@TODO: No clue where to route this to a toolkit one, and anywhere I can put an OnKeyDown around here
-	// or in the STileLayerList to route it myself isn't good enough, since the editor viewport will have focus
-	// 99% of the time, not the details panel.
-	TSharedPtr<FUICommandList> CommandList = MakeShareable(new FUICommandList);
+	// Try to get the hosting command list from the details view
+	TSharedPtr<FUICommandList> CommandList = DetailLayout.GetDetailsView().GetHostCommandList();
+	if (!CommandList.IsValid())
+	{
+		CommandList = MakeShareable(new FUICommandList);
+	}
 
 	// Add the layer browser
 	if (TileMap != nullptr)
diff --git a/Engine/Plugins/2D/Paper2D/Source/Paper2DEditor/Private/TileMapEditing/STileLayerList.cpp b/Engine/Plugins/2D/Paper2D/Source/Paper2DEditor/Private/TileMapEditing/STileLayerList.cpp
--- a/Engine/Plugins/2D/Paper2D/Source/Paper2DEditor/Private/TileMapEditing/STileLayerList.cpp
+++ b/Engine/Plugins/2D/Paper2D/Source/Paper2DEditor/Private/TileMapEditing/STileLayerList.cpp
@@ -14,7 +14,7 @@
 //////////////////////////////////////////////////////////////////////////
 // STileLayerList
 
-void STileLayerList::Construct(const FArguments& InArgs, UPaperTileMap* InTileMap, FNotifyHook* InNotifyHook, TSharedPtr<class FUICommandList> InCommandList)
+void STileLayerList::Construct(const FArguments& InArgs, UPaperTileMap* InTileMap, FNotifyHook* InNotifyHook, TSharedPtr<class FUICommandList> InParentCommandList)
 {
 	TileMapPtr = InTileMap;
 	NotifyHook = InNotifyHook;
@@ -22,7 +22,8 @@ void STileLayerList::Construct(const FArguments& InArgs, UPaperTileMap* InTileMa
 	FTileMapEditorCommands::Register();
 	const FTileMapEditorCommands& Commands = FTileMapEditorCommands::Get();
 
-	CommandList = InCommandList;
+	CommandList = MakeShareable(new FUICommandList);
+	InParentCommandList->Append(CommandList.ToSharedRef());
 
 	CommandList->MapAction(
 		Commands.AddNewLayerAbove,
@@ -475,7 +476,19 @@ TSharedPtr<SWidget> STileLayerList::OnConstructContextMenu()
 		MenuBuilder.AddMenuEntry(Commands.DuplicateLayer);
 		MenuBuilder.AddMenuEntry(Commands.DeleteLayer);
 		MenuBuilder.AddMenuEntry(Commands.MergeLayerDown);
- 	}
+		MenuBuilder.AddMenuSeparator();
+		MenuBuilder.AddMenuEntry(Commands.SelectLayerAbove);
+		MenuBuilder.AddMenuEntry(Commands.SelectLayerBelow);
+	}
+	MenuBuilder.EndSection();
+
+	MenuBuilder.BeginSection("OrderingOperations", LOCTEXT("OrderingOperationsHeader", "Order actions"));
+	{
+		MenuBuilder.AddMenuEntry(Commands.MoveLayerToTop);
+		MenuBuilder.AddMenuEntry(Commands.MoveLayerUp);
+		MenuBuilder.AddMenuEntry(Commands.MoveLayerDown);
+		MenuBuilder.AddMenuEntry(Commands.MoveLayerToBottom);
+	}
 	MenuBuilder.EndSection();
 
 	return MenuBuilder.MakeWidget();
