Paper2D: Improved Texture Tint tooltip in extract sprites window [UE-13168]

diff --git a/Engine/Plugins/2D/Paper2D/Source/Paper2DEditor/Private/ExtractSprites/PaperExtractSpritesSettings.cpp b/Engine/Plugins/2D/Paper2D/Source/Paper2DEditor/Private/ExtractSprites/PaperExtractSpritesSettings.cpp
--- a/Engine/Plugins/2D/Paper2D/Source/Paper2DEditor/Private/ExtractSprites/PaperExtractSpritesSettings.cpp
+++ b/Engine/Plugins/2D/Paper2D/Source/Paper2DEditor/Private/ExtractSprites/PaperExtractSpritesSettings.cpp
@@ -7,7 +7,7 @@ UPaperExtractSpritesSettings::UPaperExtractSpritesSettings(const FObjectInitiali
 	: Super(ObjectInitializer)
 {
 	OutlineColor = FLinearColor::Yellow;
-	TextureTint = FLinearColor::Gray;
+	ViewportTextureTint = FLinearColor::Gray;
 	BackgroundColor = FLinearColor(0.1f, 0.1f, 0.1f);
 }
 
diff --git a/Engine/Plugins/2D/Paper2D/Source/Paper2DEditor/Private/ExtractSprites/PaperExtractSpritesSettings.h b/Engine/Plugins/2D/Paper2D/Source/Paper2DEditor/Private/ExtractSprites/PaperExtractSpritesSettings.h
--- a/Engine/Plugins/2D/Paper2D/Source/Paper2DEditor/Private/ExtractSprites/PaperExtractSpritesSettings.h
+++ b/Engine/Plugins/2D/Paper2D/Source/Paper2DEditor/Private/ExtractSprites/PaperExtractSpritesSettings.h
@@ -33,9 +33,9 @@ public:
 	UPROPERTY(Category = Settings, EditAnywhere)
 	FLinearColor OutlineColor;
 
-	// Tint the texture to help increase outline visibility
+	// Apply a tint to the texture in the viewport to improve outline visibility in this editor
 	UPROPERTY(Category = Settings, EditAnywhere)
-	FLinearColor TextureTint;
+	FLinearColor ViewportTextureTint;
 
 	// The viewport background color
 	UPROPERTY(Category = Settings, EditAnywhere, meta=(HideAlphaChannel))
