Paper2D: Missed file in CL2509836

diff --git a/Engine/Plugins/2D/Paper2D/Source/Paper2DEditor/Private/ExtractSprites/SPaperExtractSpritesDialog.cpp b/Engine/Plugins/2D/Paper2D/Source/Paper2DEditor/Private/ExtractSprites/SPaperExtractSpritesDialog.cpp
--- a/Engine/Plugins/2D/Paper2D/Source/Paper2DEditor/Private/ExtractSprites/SPaperExtractSpritesDialog.cpp
+++ b/Engine/Plugins/2D/Paper2D/Source/Paper2DEditor/Private/ExtractSprites/SPaperExtractSpritesDialog.cpp
@@ -53,7 +53,7 @@ void FPaperExtractSpritesViewportClient::Draw(FViewport* Viewport, FCanvas* Canv
 		Texture->SetForceMipLevelsToBeResident(30.0f);
 		Texture->WaitForStreaming();
 
-		FLinearColor TextureDrawColor = Settings->TextureTint;
+		FLinearColor TextureDrawColor = Settings->ViewportTextureTint;
 		//FLinearColor RectOutlineColor = FLinearColor::Yellow;
 		const FLinearColor RectOutlineColor = Settings->OutlineColor;
 
@@ -380,7 +380,7 @@ void SPaperExtractSpritesDialog::OnFinishedChangingProperties(const FPropertyCha
 		
 		if (PropertyName != GET_MEMBER_NAME_CHECKED(UPaperExtractSpritesSettings, OutlineColor) &&
 			PropertyName != GET_MEMBER_NAME_CHECKED(UPaperExtractSpritesSettings, BackgroundColor) &&
-			PropertyName != GET_MEMBER_NAME_CHECKED(UPaperExtractSpritesSettings, TextureTint))
+			PropertyName != GET_MEMBER_NAME_CHECKED(UPaperExtractSpritesSettings, ViewportTextureTint))
 		{
 			PreviewExtractedSprites();
 		}
