Paper2D: Terrain spline points are constrained to the XZ plane. [UE-9061]

diff --git a/Engine/Plugins/2D/Paper2D/Source/Paper2D/Classes/PaperTerrainComponent.h b/Engine/Plugins/2D/Paper2D/Source/Paper2D/Classes/PaperTerrainComponent.h
--- a/Engine/Plugins/2D/Paper2D/Source/Paper2D/Classes/PaperTerrainComponent.h
+++ b/Engine/Plugins/2D/Paper2D/Source/Paper2D/Classes/PaperTerrainComponent.h
@@ -114,6 +114,7 @@ protected:
 	void GenerateFillRenderDataFromPolygon(const class UPaperSprite* NewSprite, FSpriteDrawCallRecord& FillDrawCall, const FVector2D& TextureSize, const TArray<FVector2D>& TriangulatedPolygonVertices);
 	void GenerateCollisionDataFromPolygon(const TArray<FVector2D>& SplinePolyVertices2D, const TArray<float>& TerrainOffsets, const TArray<FVector2D>& TriangulatedPolygonVertices);
 	void InsertConvexCollisionDataFromPolygon(const TArray<FVector2D>& ClosedPolyVertices2D);
+	void ConstrainSplinePointsToXZ();
 
 	void OnSplineEdited();
 
diff --git a/Engine/Plugins/2D/Paper2D/Source/Paper2D/Private/Terrain/PaperTerrainComponent.cpp b/Engine/Plugins/2D/Paper2D/Source/Paper2D/Private/Terrain/PaperTerrainComponent.cpp
--- a/Engine/Plugins/2D/Paper2D/Source/Paper2D/Private/Terrain/PaperTerrainComponent.cpp
+++ b/Engine/Plugins/2D/Paper2D/Source/Paper2D/Private/Terrain/PaperTerrainComponent.cpp
@@ -328,8 +328,37 @@ static void SimplifyPolygon(TArray<FVector2D> &SplinePolyVertices2D, TArray<floa
 	}
 }
 
+// Makes sure all spline points are constrained to the XZ plane
+void UPaperTerrainComponent::ConstrainSplinePointsToXZ()
+{
+	if (AssociatedSpline != nullptr)
+	{
+		bool bSplineChanged = false;
+		auto& Points = AssociatedSpline->SplineInfo.Points;
+		int32 NumPoints = Points.Num();
+		for (int PointIndex = 0; PointIndex < NumPoints; ++PointIndex)
+		{
+			auto& CurrentPoint = Points[PointIndex];
+			if (CurrentPoint.ArriveTangent.Y != 0 || CurrentPoint.LeaveTangent.Y != 0 || CurrentPoint.OutVal.Y != 0)
+			{
+				CurrentPoint.ArriveTangent.Y = 0;
+				CurrentPoint.LeaveTangent.Y = 0;
+				CurrentPoint.OutVal.Y = 0;
+				bSplineChanged = true;
+			}
+		}
+
+		if (bSplineChanged)
+		{
+			AssociatedSpline->UpdateSpline();
+		}
+	}
+}
+
 void UPaperTerrainComponent::OnSplineEdited()
 {
+	ConstrainSplinePointsToXZ();
+
 	// Ensure we have the data structure for the desired collision method
 	if (SpriteCollisionDomain == ESpriteCollisionMode::Use3DPhysics)
 	{
