Paper2D: Added includes necessary for compiling the Paper2D runtime without including Paper2DClasses.h in the PCH

diff --git a/Engine/Plugins/2D/Paper2D/Source/Paper2D/Classes/PaperTerrainMaterial.h b/Engine/Plugins/2D/Paper2D/Source/Paper2D/Classes/PaperTerrainMaterial.h
--- a/Engine/Plugins/2D/Paper2D/Source/Paper2D/Classes/PaperTerrainMaterial.h
+++ b/Engine/Plugins/2D/Paper2D/Source/Paper2D/Classes/PaperTerrainMaterial.h
@@ -12,15 +12,15 @@ struct FPaperTerrainMaterialRule
 
 	// The sprite to use at the 'left' (closest to spline start) edge of the terrain segment
 	UPROPERTY(Category=Sprite, EditAnywhere, meta=(DisplayThumbnail="true"))
-	UPaperSprite* StartCap;
+	class UPaperSprite* StartCap;
 
 	// A set of sprites to randomly choose to fill up the interior space between the caps in a terrain segment
 	UPROPERTY(Category=Sprite, EditAnywhere, meta=(DisplayThumbnail="true"))
-	TArray<UPaperSprite*> Body;
+	TArray<class UPaperSprite*> Body;
 
 	// The sprite to use at the 'right' (closest to spline end) edge of the terrain segment
 	UPROPERTY(Category=Sprite, EditAnywhere, meta=(DisplayThumbnail="true"))
-	UPaperSprite* EndCap;
+	class UPaperSprite* EndCap;
 
 	// Minimum slope angle (in degrees) to apply this rule
 	UPROPERTY(Category=Sprite, EditAnywhere)
@@ -75,7 +75,7 @@ class PAPER2D_API UPaperTerrainMaterial : public UDataAsset
 
 	// The sprite to use for an interior region fill
 	UPROPERTY(Category = Sprite, EditAnywhere, meta=(DisplayThumbnail="true"))
-	UPaperSprite* InteriorFill;
+	class UPaperSprite* InteriorFill;
 
 #if WITH_EDITOR
 	// UObject interface
diff --git a/Engine/Plugins/2D/Paper2D/Source/Paper2D/Classes/PaperTileMapComponent.h b/Engine/Plugins/2D/Paper2D/Source/Paper2D/Classes/PaperTileMapComponent.h
--- a/Engine/Plugins/2D/Paper2D/Source/Paper2D/Classes/PaperTileMapComponent.h
+++ b/Engine/Plugins/2D/Paper2D/Source/Paper2D/Classes/PaperTileMapComponent.h
@@ -3,7 +3,7 @@
 #pragma once
 
 #include "PaperSprite.h"
-
+#include "PaperTileLayer.h"
 #include "PaperTileMapComponent.generated.h"
 
 /**
@@ -36,13 +36,13 @@ private:
 	int32 TileHeight_DEPRECATED;
 
 	UPROPERTY()
-	UPaperTileSet* DefaultLayerTileSet_DEPRECATED;
+	class UPaperTileSet* DefaultLayerTileSet_DEPRECATED;
 
 	UPROPERTY()
 	UMaterialInterface* Material_DEPRECATED;
 
 	UPROPERTY()
-	TArray<UPaperTileLayer*> TileLayers_DEPRECATED;
+	TArray<class UPaperTileLayer*> TileLayers_DEPRECATED;
 
 public:
 	// The tile map used by this component
diff --git a/Engine/Plugins/2D/Paper2D/Source/Paper2D/Classes/SpriteDrawCall.h b/Engine/Plugins/2D/Paper2D/Source/Paper2D/Classes/SpriteDrawCall.h
--- a/Engine/Plugins/2D/Paper2D/Source/Paper2D/Classes/SpriteDrawCall.h
+++ b/Engine/Plugins/2D/Paper2D/Source/Paper2D/Classes/SpriteDrawCall.h
@@ -2,6 +2,7 @@
 
 #pragma once
 
+#include "PaperSprite.h"
 #include "SpriteDrawCall.generated.h"
 
 //
diff --git a/Engine/Plugins/2D/Paper2D/Source/Paper2D/Private/PaperCharacter.cpp b/Engine/Plugins/2D/Paper2D/Source/Paper2D/Private/PaperCharacter.cpp
--- a/Engine/Plugins/2D/Paper2D/Source/Paper2D/Private/PaperCharacter.cpp
+++ b/Engine/Plugins/2D/Paper2D/Source/Paper2D/Private/PaperCharacter.cpp
@@ -1,6 +1,8 @@
 // Copyright 1998-2015 Epic Games, Inc. All Rights Reserved.
 
 #include "Paper2DPrivatePCH.h"
+#include "PaperCharacter.h"
+#include "PaperFlipbookComponent.h"
 
 FName APaperCharacter::SpriteComponentName(TEXT("Sprite0"));
 
diff --git a/Engine/Plugins/2D/Paper2D/Source/Paper2D/Private/PaperFlipbook.cpp b/Engine/Plugins/2D/Paper2D/Source/Paper2D/Private/PaperFlipbook.cpp
--- a/Engine/Plugins/2D/Paper2D/Source/Paper2D/Private/PaperFlipbook.cpp
+++ b/Engine/Plugins/2D/Paper2D/Source/Paper2D/Private/PaperFlipbook.cpp
@@ -3,6 +3,7 @@
 #include "Paper2DPrivatePCH.h"
 #include "PaperFlipbook.h"
 #include "PaperCustomVersion.h"
+#include "PaperFlipbookComponent.h"
 
 #if WITH_EDITORONLY_DATA
 #include "Runtime/Engine/Public/ComponentReregisterContext.h"
diff --git a/Engine/Plugins/2D/Paper2D/Source/Paper2D/Private/PaperFlipbookComponent.cpp b/Engine/Plugins/2D/Paper2D/Source/Paper2D/Private/PaperFlipbookComponent.cpp
--- a/Engine/Plugins/2D/Paper2D/Source/Paper2D/Private/PaperFlipbookComponent.cpp
+++ b/Engine/Plugins/2D/Paper2D/Source/Paper2D/Private/PaperFlipbookComponent.cpp
@@ -10,6 +10,7 @@
 #include "Runtime/Core/Public/Misc/MapErrors.h"
 #include "Runtime/CoreUObject/Public/Misc/UObjectToken.h"
 #include "PhysicsEngine/BodySetup.h"
+#include "SpriteDrawCall.h"
 
 #define LOCTEXT_NAMESPACE "Paper2D"
 
diff --git a/Engine/Plugins/2D/Paper2D/Source/Paper2D/Private/PaperFlipbookSceneProxy.cpp b/Engine/Plugins/2D/Paper2D/Source/Paper2D/Private/PaperFlipbookSceneProxy.cpp
--- a/Engine/Plugins/2D/Paper2D/Source/Paper2D/Private/PaperFlipbookSceneProxy.cpp
+++ b/Engine/Plugins/2D/Paper2D/Source/Paper2D/Private/PaperFlipbookSceneProxy.cpp
@@ -2,6 +2,7 @@
 
 #include "Paper2DPrivatePCH.h"
 #include "PaperFlipbookSceneProxy.h"
+#include "PaperFlipbookComponent.h"
 
 //////////////////////////////////////////////////////////////////////////
 // FPaperFlipbookSceneProxy
diff --git a/Engine/Plugins/2D/Paper2D/Source/Paper2D/Private/PaperFlipbookSceneProxy.h b/Engine/Plugins/2D/Paper2D/Source/Paper2D/Private/PaperFlipbookSceneProxy.h
--- a/Engine/Plugins/2D/Paper2D/Source/Paper2D/Private/PaperFlipbookSceneProxy.h
+++ b/Engine/Plugins/2D/Paper2D/Source/Paper2D/Private/PaperFlipbookSceneProxy.h
@@ -10,5 +10,5 @@
 class FPaperFlipbookSceneProxy : public FPaperRenderSceneProxy
 {
 public:
-	FPaperFlipbookSceneProxy(const UPaperFlipbookComponent* InComponent);
+	FPaperFlipbookSceneProxy(const class UPaperFlipbookComponent* InComponent);
 };
diff --git a/Engine/Plugins/2D/Paper2D/Source/Paper2D/Private/PaperRenderSceneProxy.cpp b/Engine/Plugins/2D/Paper2D/Source/Paper2D/Private/PaperRenderSceneProxy.cpp
--- a/Engine/Plugins/2D/Paper2D/Source/Paper2D/Private/PaperRenderSceneProxy.cpp
+++ b/Engine/Plugins/2D/Paper2D/Source/Paper2D/Private/PaperRenderSceneProxy.cpp
@@ -11,6 +11,7 @@
 #include "EngineGlobals.h"
 #include "SceneManagement.h"
 #include "Engine/Engine.h"
+#include "SpriteDrawCall.h"
 
 //////////////////////////////////////////////////////////////////////////
 // FPaperSpriteVertex
diff --git a/Engine/Plugins/2D/Paper2D/Source/Paper2D/Private/PaperRenderSceneProxy.h b/Engine/Plugins/2D/Paper2D/Source/Paper2D/Private/PaperRenderSceneProxy.h
--- a/Engine/Plugins/2D/Paper2D/Source/Paper2D/Private/PaperRenderSceneProxy.h
+++ b/Engine/Plugins/2D/Paper2D/Source/Paper2D/Private/PaperRenderSceneProxy.h
@@ -2,6 +2,9 @@
 
 #pragma once
 
+#include "PaperSprite.h"
+#include "SpriteDrawCall.h"
+
 #define TEST_BATCHING 0
 
 //////////////////////////////////////////////////////////////////////////
diff --git a/Engine/Plugins/2D/Paper2D/Source/Paper2D/Private/PaperSprite.cpp b/Engine/Plugins/2D/Paper2D/Source/Paper2D/Private/PaperSprite.cpp
--- a/Engine/Plugins/2D/Paper2D/Source/Paper2D/Private/PaperSprite.cpp
+++ b/Engine/Plugins/2D/Paper2D/Source/Paper2D/Private/PaperSprite.cpp
@@ -7,6 +7,10 @@
 
 #include "PaperCustomVersion.h"
 #include "PaperGeomTools.h"
+#include "PaperRuntimeSettings.h"
+#include "PaperSpriteComponent.h"
+#include "PaperFlipbookComponent.h"
+#include "SpriteDrawCall.h"
 
 #if WITH_EDITOR
 #include "UnrealEd.h"
diff --git a/Engine/Plugins/2D/Paper2D/Source/Paper2D/Private/PaperSpriteAtlas.cpp b/Engine/Plugins/2D/Paper2D/Source/Paper2D/Private/PaperSpriteAtlas.cpp
--- a/Engine/Plugins/2D/Paper2D/Source/Paper2D/Private/PaperSpriteAtlas.cpp
+++ b/Engine/Plugins/2D/Paper2D/Source/Paper2D/Private/PaperSpriteAtlas.cpp
@@ -1,6 +1,7 @@
 // Copyright 1998-2015 Epic Games, Inc. All Rights Reserved.
 
 #include "Paper2DPrivatePCH.h"
+#include "PaperSpriteAtlas.h"
 
 //////////////////////////////////////////////////////////////////////////
 // UPaperSpriteAtlas
diff --git a/Engine/Plugins/2D/Paper2D/Source/Paper2D/Private/PaperSpriteComponent.cpp b/Engine/Plugins/2D/Paper2D/Source/Paper2D/Private/PaperSpriteComponent.cpp
--- a/Engine/Plugins/2D/Paper2D/Source/Paper2D/Private/PaperSpriteComponent.cpp
+++ b/Engine/Plugins/2D/Paper2D/Source/Paper2D/Private/PaperSpriteComponent.cpp
@@ -9,6 +9,7 @@
 #include "Runtime/Core/Public/Logging/MessageLog.h"
 #include "Runtime/Core/Public/Misc/MapErrors.h"
 #include "Runtime/CoreUObject/Public/Misc/UObjectToken.h"
+#include "SpriteDrawCall.h"
 
 #define LOCTEXT_NAMESPACE "Paper2D"
 
diff --git a/Engine/Plugins/2D/Paper2D/Source/Paper2D/Private/PaperSpriteSceneProxy.cpp b/Engine/Plugins/2D/Paper2D/Source/Paper2D/Private/PaperSpriteSceneProxy.cpp
--- a/Engine/Plugins/2D/Paper2D/Source/Paper2D/Private/PaperSpriteSceneProxy.cpp
+++ b/Engine/Plugins/2D/Paper2D/Source/Paper2D/Private/PaperSpriteSceneProxy.cpp
@@ -6,6 +6,8 @@
 #include "SceneManagement.h"
 #include "EngineGlobals.h"
 #include "Engine/Engine.h"
+#include "PaperSpriteComponent.h"
+#include "SpriteDrawCall.h"
 
 //////////////////////////////////////////////////////////////////////////
 // FPaperSpriteSceneProxy
diff --git a/Engine/Plugins/2D/Paper2D/Source/Paper2D/Private/PaperSpriteSceneProxy.h b/Engine/Plugins/2D/Paper2D/Source/Paper2D/Private/PaperSpriteSceneProxy.h
--- a/Engine/Plugins/2D/Paper2D/Source/Paper2D/Private/PaperSpriteSceneProxy.h
+++ b/Engine/Plugins/2D/Paper2D/Source/Paper2D/Private/PaperSpriteSceneProxy.h
@@ -10,13 +10,13 @@
 class FPaperSpriteSceneProxy : public FPaperRenderSceneProxy
 {
 public:
-	FPaperSpriteSceneProxy(const UPaperSpriteComponent* InComponent);
+	FPaperSpriteSceneProxy(const class UPaperSpriteComponent* InComponent);
 
 	// FPrimitiveSceneProxy interface
 	virtual void GetDynamicMeshElements(const TArray<const FSceneView*>& Views, const FSceneViewFamily& ViewFamily, uint32 VisibilityMap, FMeshElementCollector& Collector) const override;
 	// End of FPrimitiveSceneProxy interface
 
-	void SetSprite_RenderThread(const FSpriteDrawCallRecord& NewDynamicData, int32 SplitIndex);
+	void SetSprite_RenderThread(const struct FSpriteDrawCallRecord& NewDynamicData, int32 SplitIndex);
 
 protected:
 
@@ -26,6 +26,6 @@ protected:
 
 	UMaterialInterface* AlternateMaterial;
 	int32 MaterialSplitIndex;
-	const UPaperSprite* SourceSprite;
-	TArray<FSpriteDrawCallRecord> AlternateBatchedSprites;
+	const class UPaperSprite* SourceSprite;
+	TArray<struct FSpriteDrawCallRecord> AlternateBatchedSprites;
 };
diff --git a/Engine/Plugins/2D/Paper2D/Source/Paper2D/Private/PaperTileLayer.cpp b/Engine/Plugins/2D/Paper2D/Source/Paper2D/Private/PaperTileLayer.cpp
--- a/Engine/Plugins/2D/Paper2D/Source/Paper2D/Private/PaperTileLayer.cpp
+++ b/Engine/Plugins/2D/Paper2D/Source/Paper2D/Private/PaperTileLayer.cpp
@@ -5,6 +5,9 @@
 #include "ComponentReregisterContext.h"
 #include "PhysicsEngine/BoxElem.h"
 #include "PhysicsEngine/BodySetup.h"
+#include "PaperTileMap.h"
+#include "PaperTileLayer.h"
+#include "PaperTileMapComponent.h"
 
 //////////////////////////////////////////////////////////////////////////
 // FTileMapLayerReregisterContext
diff --git a/Engine/Plugins/2D/Paper2D/Source/Paper2D/Private/PaperTileMap.cpp b/Engine/Plugins/2D/Paper2D/Source/Paper2D/Private/PaperTileMap.cpp
--- a/Engine/Plugins/2D/Paper2D/Source/Paper2D/Private/PaperTileMap.cpp
+++ b/Engine/Plugins/2D/Paper2D/Source/Paper2D/Private/PaperTileMap.cpp
@@ -3,6 +3,9 @@
 #include "Paper2DPrivatePCH.h"
 #include "PhysicsEngine/BodySetup.h"
 #include "PhysicsEngine/BodySetup2D.h"
+#include "PaperTileMap.h"
+#include "PaperTileLayer.h"
+#include "PaperRuntimeSettings.h"
 
 #define LOCTEXT_NAMESPACE "Paper2D"
 
diff --git a/Engine/Plugins/2D/Paper2D/Source/Paper2D/Private/PaperTileMapActor.cpp b/Engine/Plugins/2D/Paper2D/Source/Paper2D/Private/PaperTileMapActor.cpp
--- a/Engine/Plugins/2D/Paper2D/Source/Paper2D/Private/PaperTileMapActor.cpp
+++ b/Engine/Plugins/2D/Paper2D/Source/Paper2D/Private/PaperTileMapActor.cpp
@@ -1,6 +1,8 @@
 // Copyright 1998-2015 Epic Games, Inc. All Rights Reserved.
 
 #include "Paper2DPrivatePCH.h"
+#include "PaperTileMapActor.h"
+#include "PaperTileMapComponent.h"
 
 //////////////////////////////////////////////////////////////////////////
 // APaperTileMapActor
diff --git a/Engine/Plugins/2D/Paper2D/Source/Paper2D/Private/PaperTileMapComponent.cpp b/Engine/Plugins/2D/Paper2D/Source/Paper2D/Private/PaperTileMapComponent.cpp
--- a/Engine/Plugins/2D/Paper2D/Source/Paper2D/Private/PaperTileMapComponent.cpp
+++ b/Engine/Plugins/2D/Paper2D/Source/Paper2D/Private/PaperTileMapComponent.cpp
@@ -5,6 +5,8 @@
 #include "PhysicsEngine/BodySetup.h"
 #include "PaperCustomVersion.h"
 #include "Engine/Canvas.h"
+#include "PaperTileMapComponent.h"
+#include "PaperTileMap.h"
 
 #define LOCTEXT_NAMESPACE "Paper2D"
 
diff --git a/Engine/Plugins/2D/Paper2D/Source/Paper2D/Private/PaperTileMapRenderSceneProxy.cpp b/Engine/Plugins/2D/Paper2D/Source/Paper2D/Private/PaperTileMapRenderSceneProxy.cpp
--- a/Engine/Plugins/2D/Paper2D/Source/Paper2D/Private/PaperTileMapRenderSceneProxy.cpp
+++ b/Engine/Plugins/2D/Paper2D/Source/Paper2D/Private/PaperTileMapRenderSceneProxy.cpp
@@ -7,6 +7,9 @@
 #include "EngineGlobals.h"
 #include "SceneManagement.h"
 #include "Engine/Engine.h"
+#include "PaperTileMap.h"
+#include "PaperTileLayer.h"
+#include "PaperTileMapComponent.h"
 
 //////////////////////////////////////////////////////////////////////////
 // FPaperTileMapRenderSceneProxy
diff --git a/Engine/Plugins/2D/Paper2D/Source/Paper2D/Private/PaperTileMapRenderSceneProxy.h b/Engine/Plugins/2D/Paper2D/Source/Paper2D/Private/PaperTileMapRenderSceneProxy.h
--- a/Engine/Plugins/2D/Paper2D/Source/Paper2D/Private/PaperTileMapRenderSceneProxy.h
+++ b/Engine/Plugins/2D/Paper2D/Source/Paper2D/Private/PaperTileMapRenderSceneProxy.h
@@ -10,17 +10,17 @@
 class FPaperTileMapRenderSceneProxy : public FPaperRenderSceneProxy
 {
 public:
-	FPaperTileMapRenderSceneProxy(const UPaperTileMapComponent* InComponent);
+	FPaperTileMapRenderSceneProxy(const class UPaperTileMapComponent* InComponent);
 
 	// FPrimitiveSceneProxy interface.
 	virtual void GetDynamicMeshElements(const TArray<const FSceneView*>& Views, const FSceneViewFamily& ViewFamily, uint32 VisibilityMap, FMeshElementCollector& Collector) const override;
 	// End of FPrimitiveSceneProxy interface.
 
-	void SetBatchesHack(TArray<FSpriteDrawCallRecord>& InBatchedSprites)
+	void SetBatchesHack(TArray<struct FSpriteDrawCallRecord>& InBatchedSprites)
 	{
 		BatchedSprites = InBatchedSprites;
 	}
 protected:
 	//@TODO: Not thread safe
-	const UPaperTileMap* TileMap;
+	const class UPaperTileMap* TileMap;
 };
diff --git a/Engine/Plugins/2D/Paper2D/Source/Paper2D/Private/PaperTileSet.cpp b/Engine/Plugins/2D/Paper2D/Source/Paper2D/Private/PaperTileSet.cpp
--- a/Engine/Plugins/2D/Paper2D/Source/Paper2D/Private/PaperTileSet.cpp
+++ b/Engine/Plugins/2D/Paper2D/Source/Paper2D/Private/PaperTileSet.cpp
@@ -1,6 +1,7 @@
 // Copyright 1998-2015 Epic Games, Inc. All Rights Reserved.
 
 #include "Paper2DPrivatePCH.h"
+#include "PaperTileSet.h"
 
 //////////////////////////////////////////////////////////////////////////
 // UPaperTileSet
diff --git a/Engine/Plugins/2D/Paper2D/Source/Paper2D/Private/Terrain/PaperTerrainActor.cpp b/Engine/Plugins/2D/Paper2D/Source/Paper2D/Private/Terrain/PaperTerrainActor.cpp
--- a/Engine/Plugins/2D/Paper2D/Source/Paper2D/Private/Terrain/PaperTerrainActor.cpp
+++ b/Engine/Plugins/2D/Paper2D/Source/Paper2D/Private/Terrain/PaperTerrainActor.cpp
@@ -1,6 +1,10 @@
 // Copyright 1998-2015 Epic Games, Inc. All Rights Reserved.
 
 #include "Paper2DPrivatePCH.h"
+#include "PaperTerrainActor.h"
+#include "PaperTerrainMaterial.h"
+#include "PaperTerrainComponent.h"
+#include "PaperTerrainSplineComponent.h"
 
 //////////////////////////////////////////////////////////////////////////
 // APaperTerrainActor
diff --git a/Engine/Plugins/2D/Paper2D/Source/Paper2D/Private/Terrain/PaperTerrainComponent.cpp b/Engine/Plugins/2D/Paper2D/Source/Paper2D/Private/Terrain/PaperTerrainComponent.cpp
--- a/Engine/Plugins/2D/Paper2D/Source/Paper2D/Private/Terrain/PaperTerrainComponent.cpp
+++ b/Engine/Plugins/2D/Paper2D/Source/Paper2D/Private/Terrain/PaperTerrainComponent.cpp
@@ -7,6 +7,9 @@
 #include "PaperGeomTools.h"
 #include "PhysicsEngine/BodySetup.h"
 #include "PhysicsEngine/ConvexElem.h"
+#include "PaperTerrainComponent.h"
+#include "PaperTerrainSplineComponent.h"
+#include "PaperTerrainMaterial.h"
 
 #define PAPER_USE_MATERIAL_SLOPES 1
 
@@ -27,10 +30,11 @@ static FBox2D GetSpriteRenderDataBounds2D(const TArray<FVector4>& Data)
 
 //////////////////////////////////////////////////////////////////////////
 
-FTerrainSpriteStamp::FTerrainSpriteStamp(const UPaperSprite* InSprite, float InTime, bool bIsEndCap) : Sprite(InSprite)
-, Time(InTime)
-, Scale(1.0f)
-, bCanStretch(!bIsEndCap)
+FTerrainSpriteStamp::FTerrainSpriteStamp(const UPaperSprite* InSprite, float InTime, bool bIsEndCap)
+	: Sprite(InSprite)
+	, Time(InTime)
+	, Scale(1.0f)
+	, bCanStretch(!bIsEndCap)
 {
 	const FBox2D Bounds2D = GetSpriteRenderDataBounds2D(InSprite->BakedRenderData);
 	NominalWidth = FMath::Max<float>(Bounds2D.GetSize().X, 1.0f);
diff --git a/Engine/Plugins/2D/Paper2D/Source/Paper2D/Private/Terrain/PaperTerrainMaterial.cpp b/Engine/Plugins/2D/Paper2D/Source/Paper2D/Private/Terrain/PaperTerrainMaterial.cpp
--- a/Engine/Plugins/2D/Paper2D/Source/Paper2D/Private/Terrain/PaperTerrainMaterial.cpp
+++ b/Engine/Plugins/2D/Paper2D/Source/Paper2D/Private/Terrain/PaperTerrainMaterial.cpp
@@ -1,6 +1,7 @@
 // Copyright 1998-2015 Epic Games, Inc. All Rights Reserved.
 
 #include "Paper2DPrivatePCH.h"
+#include "PaperTerrainMaterial.h"
 
 #define LOCTEXT_NAMESPACE "Paper2D"
 
diff --git a/Engine/Plugins/2D/Paper2D/Source/Paper2D/Private/Terrain/PaperTerrainSplineComponent.cpp b/Engine/Plugins/2D/Paper2D/Source/Paper2D/Private/Terrain/PaperTerrainSplineComponent.cpp
--- a/Engine/Plugins/2D/Paper2D/Source/Paper2D/Private/Terrain/PaperTerrainSplineComponent.cpp
+++ b/Engine/Plugins/2D/Paper2D/Source/Paper2D/Private/Terrain/PaperTerrainSplineComponent.cpp
@@ -1,6 +1,7 @@
 // Copyright 1998-2015 Epic Games, Inc. All Rights Reserved.
 
 #include "Paper2DPrivatePCH.h"
+#include "PaperTerrainSplineComponent.h"
 
 //////////////////////////////////////////////////////////////////////////
 // UPaperTerrainSplineComponent
