Sprite Editor: Added a corner text display for the current mode when not in View mode

diff --git a/Engine/Plugins/2D/Paper2D/Source/Paper2DEditor/Private/FlipbookEditor/FlipbookEditor.cpp b/Engine/Plugins/2D/Paper2D/Source/Paper2DEditor/Private/FlipbookEditor/FlipbookEditor.cpp
--- a/Engine/Plugins/2D/Paper2D/Source/Paper2DEditor/Private/FlipbookEditor/FlipbookEditor.cpp
+++ b/Engine/Plugins/2D/Paper2D/Source/Paper2DEditor/Private/FlipbookEditor/FlipbookEditor.cpp
@@ -259,7 +259,7 @@ TSharedRef<SDockTab> FFlipbookEditor::SpawnTab_Viewport(const FSpawnTabArgs& Arg
 			[
 				SNew(SOverlay)
 
-				// The sprite editor viewport
+				// The flipbook editor viewport
 				+SOverlay::Slot()
 				[
 					ViewportPtr.ToSharedRef()
diff --git a/Engine/Plugins/2D/Paper2D/Source/Paper2DEditor/Private/SpriteEditor/SpriteEditor.cpp b/Engine/Plugins/2D/Paper2D/Source/Paper2DEditor/Private/SpriteEditor/SpriteEditor.cpp
--- a/Engine/Plugins/2D/Paper2D/Source/Paper2DEditor/Private/SpriteEditor/SpriteEditor.cpp
+++ b/Engine/Plugins/2D/Paper2D/Source/Paper2DEditor/Private/SpriteEditor/SpriteEditor.cpp
@@ -268,7 +268,25 @@ TSharedRef<SDockTab> FSpriteEditor::SpawnTab_Viewport(const FSpawnTabArgs& Args)
 	return SNew(SDockTab)
 		.Label(LOCTEXT("ViewportTab_Title", "Viewport"))
 		[
-			ViewportPtr.ToSharedRef()
+			SNew(SOverlay)
+
+			// The sprite editor viewport
+			+SOverlay::Slot()
+			[
+				ViewportPtr.ToSharedRef()
+			]
+
+			// Bottom-right corner text indicating the preview nature of the sprite editor
+			+SOverlay::Slot()
+			.Padding(10)
+			.VAlign(VAlign_Bottom)
+			.HAlign(HAlign_Right)
+			[
+				SNew(STextBlock)
+				.Visibility(EVisibility::HitTestInvisible)
+				.TextStyle(FEditorStyle::Get(), "Graph.CornerText")
+				.Text(this, &FSpriteEditor::GetCurrentModeCornerText)
+			]
 		];
 }
 
@@ -546,6 +564,21 @@ void FSpriteEditor::CreateModeToolbarWidgets(FToolBarBuilder& IgnoredBuilder)
 	AddToolbarWidget(ToolbarBuilder.MakeWidget());
 }
 
+FText FSpriteEditor::GetCurrentModeCornerText() const
+{
+	switch (GetCurrentMode())
+	{
+	case ESpriteEditorMode::EditCollisionMode:
+		return LOCTEXT("EditCollisionGeometry_CornerText", "Edit Collision");
+	case ESpriteEditorMode::EditRenderingGeomMode:
+		return LOCTEXT("EditRenderGeometry_CornerText", "Edit Render Geometry");
+	case ESpriteEditorMode::EditSourceRegionMode:
+		return LOCTEXT("EditSourceRegion_CornerText", "Edit Source Region");
+	default:
+		return FText::GetEmpty();
+	}
+}
+
 //////////////////////////////////////////////////////////////////////////
 
 #undef LOCTEXT_NAMESPACE
diff --git a/Engine/Plugins/2D/Paper2D/Source/Paper2DEditor/Private/SpriteEditor/SpriteEditor.h b/Engine/Plugins/2D/Paper2D/Source/Paper2DEditor/Private/SpriteEditor/SpriteEditor.h
--- a/Engine/Plugins/2D/Paper2D/Source/Paper2DEditor/Private/SpriteEditor/SpriteEditor.h
+++ b/Engine/Plugins/2D/Paper2D/Source/Paper2DEditor/Private/SpriteEditor/SpriteEditor.h
@@ -69,4 +69,6 @@ protected:
 	TSharedRef<SDockTab> SpawnTab_SpriteList(const FSpawnTabArgs& Args);
 
 	void CreateModeToolbarWidgets(FToolBarBuilder& ToolbarBuilder);
+
+	FText GetCurrentModeCornerText() const;
 };
\ No newline at end of file
