Sprites: Moved sockets into their own category

diff --git a/Engine/Plugins/2D/Paper2D/Source/Paper2DEditor/Private/SpriteEditor/SpriteDetailsCustomization.cpp b/Engine/Plugins/2D/Paper2D/Source/Paper2DEditor/Private/SpriteEditor/SpriteDetailsCustomization.cpp
--- a/Engine/Plugins/2D/Paper2D/Source/Paper2DEditor/Private/SpriteEditor/SpriteDetailsCustomization.cpp
+++ b/Engine/Plugins/2D/Paper2D/Source/Paper2DEditor/Private/SpriteEditor/SpriteDetailsCustomization.cpp
@@ -71,6 +71,12 @@ void FSpriteDetailsCustomization::CustomizeDetails(IDetailLayoutBuilder& DetailL
 	IDetailCategoryBuilder& SpriteCategory = DetailLayout.EditCategory("Sprite", FText::GetEmpty(), ECategoryPriority::Important);
 	BuildSpriteSection(SpriteCategory, DetailLayout);
 
+	// Build the socket category
+	{
+		IDetailCategoryBuilder& SocketCategory = DetailLayout.EditCategory("Sockets");
+		SocketCategory.AddProperty(GET_MEMBER_NAME_CHECKED(UPaperSprite, Sockets));
+	}
+
 	// Build the collision category
 	IDetailCategoryBuilder& CollisionCategory = DetailLayout.EditCategory("Collision");
 	BuildCollisionSection(CollisionCategory, DetailLayout);
@@ -88,7 +94,6 @@ void FSpriteDetailsCustomization::BuildSpriteSection(IDetailCategoryBuilder& Spr
 	BuildTextureSection(SpriteCategory, DetailLayout);
 
 	SpriteCategory.AddProperty(GET_MEMBER_NAME_CHECKED(UPaperSprite, DefaultMaterial));
-	SpriteCategory.AddProperty(GET_MEMBER_NAME_CHECKED(UPaperSprite, Sockets));
 	SpriteCategory.AddProperty(GET_MEMBER_NAME_CHECKED(UPaperSprite, PixelsPerUnrealUnit));
 
 	// Show/hide the experimental atlas group support based on whether or not it is enabled
