Paper2D: Exposed UPaperSpriteComponent::GetSpriteColor

diff --git a/Engine/Plugins/2D/Paper2D/Source/Paper2D/Classes/PaperSpriteComponent.h b/Engine/Plugins/2D/Paper2D/Source/Paper2D/Classes/PaperSpriteComponent.h
--- a/Engine/Plugins/2D/Paper2D/Source/Paper2D/Classes/PaperSpriteComponent.h
+++ b/Engine/Plugins/2D/Paper2D/Source/Paper2D/Classes/PaperSpriteComponent.h
@@ -42,6 +42,9 @@ public:
 	UFUNCTION(BlueprintPure, Category="Sprite")
 	virtual UPaperSprite* GetSprite();
 
+	/** Returns the current color of the sprite */
+	FLinearColor GetSpriteColor() const { return SpriteColor; }
+
 	/** Set color of the sprite */
 	UFUNCTION(BlueprintCallable, Category="Sprite")
 	void SetSpriteColor(FLinearColor NewColor);
