Tile Maps: Painting with a brush that contains empty tiles now works like a brush with holes, rather than erasing the existing cells

diff --git a/Engine/Plugins/2D/Paper2D/Source/Paper2DEditor/Private/TileMapEditing/EdModeTileMap.cpp b/Engine/Plugins/2D/Paper2D/Source/Paper2DEditor/Private/TileMapEditing/EdModeTileMap.cpp
--- a/Engine/Plugins/2D/Paper2D/Source/Paper2DEditor/Private/TileMapEditing/EdModeTileMap.cpp
+++ b/Engine/Plugins/2D/Paper2D/Source/Paper2DEditor/Private/TileMapEditing/EdModeTileMap.cpp
@@ -604,7 +604,7 @@ bool FEdModeTileMap::UseActiveToolAtLocation(const FViewportCursorLocation& Ray)
 	};
 }
 
-bool BlitLayer(UPaperTileLayer* SourceLayer, UPaperTileLayer* TargetLayer, int32 OffsetX = 0, int32 OffsetY = 0)
+bool BlitLayer(UPaperTileLayer* SourceLayer, UPaperTileLayer* TargetLayer, int32 OffsetX = 0, int32 OffsetY = 0, bool bBlitEmptyTiles = false)
 {
 	FScopedTransaction Transaction(LOCTEXT("TileMapPaintAction", "Tile Painting"));
 
@@ -629,9 +629,9 @@ bool BlitLayer(UPaperTileLayer* SourceLayer, UPaperTileLayer* TargetLayer, int32
 				continue;
 			}
 
-			FPaperTileInfo Ink = SourceLayer->GetCell(SourceX, SourceY);
+			const FPaperTileInfo Ink = SourceLayer->GetCell(SourceX, SourceY);
 
-			if (TargetLayer->GetCell(TargetX, TargetY) != Ink)
+			if ((Ink.IsValid() || bBlitEmptyTiles) && (TargetLayer->GetCell(TargetX, TargetY) != Ink))
 			{
 				if (!bChangedSomething)
 				{
