Paper2D: Split various classes out from TileSetEditor.cpp

diff --git a/Engine/Plugins/2D/Paper2D/Source/Paper2DEditor/Private/TileMapEditing/TileMapEdModeToolkit.cpp b/Engine/Plugins/2D/Paper2D/Source/Paper2DEditor/Private/TileMapEditing/TileMapEdModeToolkit.cpp
--- a/Engine/Plugins/2D/Paper2D/Source/Paper2DEditor/Private/TileMapEditing/TileMapEdModeToolkit.cpp
+++ b/Engine/Plugins/2D/Paper2D/Source/Paper2DEditor/Private/TileMapEditing/TileMapEdModeToolkit.cpp
@@ -3,7 +3,7 @@
 #include "Paper2DEditorPrivatePCH.h"
 #include "EdModeTileMap.h"
 #include "TileMapEdModeToolkit.h"
-#include "../TileSetEditor.h"
+#include "TileSetEditor/TileSetSelectorViewport.h"
 #include "SContentReference.h"
 #include "TileMapEditorCommands.h"
 #include "Engine/Selection.h"
diff --git a/Engine/Plugins/2D/Paper2D/Source/Paper2DEditor/Private/TileSetEditor.cpp b/Engine/Plugins/2D/Paper2D/Source/Paper2DEditor/Private/TileSetEditor.cpp
--- a/Engine/Plugins/2D/Paper2D/Source/Paper2DEditor/Private/TileSetEditor.cpp
+++ b/Engine/Plugins/2D/Paper2D/Source/Paper2DEditor/Private/TileSetEditor.cpp
@@ -12,6 +12,11 @@
 #include "CanvasItem.h"
 #include "SDockTab.h"
 #include "PaperStyle.h"
+#include "TileSetEditor/TileSetEditorViewportClient.h"
+#include "TileSetEditor/TileSetSelectorViewport.h"
+#include "TileSetEditor/SingleTileEditorViewport.h"
+#include "TileSetEditor/SingleTileEditorViewportClient.h"
+#include "SpriteEditor/SpriteEditorCommands.h"
 
 #define LOCTEXT_NAMESPACE "TileSetEditor"
 
@@ -35,238 +40,6 @@ const FName FTileSetEditorTabs::DetailsID(TEXT("Details"));
 const FName FTileSetEditorTabs::TextureViewID(TEXT("TextureCanvas"));
 const FName FTileSetEditorTabs::SingleTileEditorID(TEXT("SingleTileEditor"));
 
-//////////////////////////////////////////////////////////////////////////
-// FTileSetEditorViewportClient 
-
-class FTileSetEditorViewportClient : public FPaperEditorViewportClient
-{
-public:
-	FTileSetEditorViewportClient(UPaperTileSet* InTileSet)
-		: TileSetBeingEdited(InTileSet)
-		, bHasValidPaintRectangle(false)
-		, TileIndex(INDEX_NONE)
-	{
-	}
-
-	// FViewportClient interface
-	virtual void Draw(FViewport* Viewport, FCanvas* Canvas) override;
-	// End of FViewportClient interface
-
-	// FEditorViewportClient interface
-	virtual FLinearColor GetBackgroundColor() const override;
-	// End of FEditorViewportClient interface
-
-public:
-	// Tile set
-	TWeakObjectPtr<UPaperTileSet> TileSetBeingEdited;
-
-	bool bHasValidPaintRectangle;
-	FViewportSelectionRectangle ValidPaintRectangle;
-	int32 TileIndex;
-};
-
-void FTileSetEditorViewportClient::Draw(FViewport* Viewport, FCanvas* Canvas)
-{
-	// Clear the viewport
-	Canvas->Clear(GetBackgroundColor());
-
-	// Can only proceed if we have a valid tile set
-	UPaperTileSet* TileSet = TileSetBeingEdited.Get();
-	if (TileSet == nullptr)
-	{
-		return;
-	}
-
-	UTexture2D* Texture = TileSet->TileSheet;
-
-	if (Texture != nullptr)
-	{
-		const bool bUseTranslucentBlend = Texture->HasAlphaChannel();
-
-		// Fully stream in the texture before drawing it.
-		Texture->SetForceMipLevelsToBeResident(30.0f);
-		Texture->WaitForStreaming();
-
-		FLinearColor TextureDrawColor = FLinearColor::White;
-
-		const float XPos = -ZoomPos.X * ZoomAmount;
-		const float YPos = -ZoomPos.Y * ZoomAmount;
-		const float Width = Texture->GetSurfaceWidth() * ZoomAmount;
-		const float Height = Texture->GetSurfaceHeight() * ZoomAmount;
-
-		Canvas->DrawTile(XPos, YPos, Width, Height, 0.0f, 0.0f, 1.0f, 1.0f, TextureDrawColor, Texture->Resource, bUseTranslucentBlend);
-	}
-
-	// Overlay the selection rectangles
-	DrawSelectionRectangles(Viewport, Canvas);
-
-	if (bHasValidPaintRectangle)
-	{
-		const FViewportSelectionRectangle& Rect = ValidPaintRectangle;
-
-		const float X = (Rect.TopLeft.X - ZoomPos.X) * ZoomAmount;
-		const float Y = (Rect.TopLeft.Y - ZoomPos.Y) * ZoomAmount;
-		const float W = Rect.Dimensions.X * ZoomAmount;
-		const float H = Rect.Dimensions.Y * ZoomAmount;
-
-		FCanvasBoxItem BoxItem(FVector2D(X, Y), FVector2D(W, H));
-		BoxItem.SetColor(Rect.Color);
-		Canvas->DrawItem(BoxItem);
-	}
-
-	if (TileIndex != INDEX_NONE)
-	{
-		const FString TileIndexString = FString::Printf(TEXT("Tile# %d"), TileIndex);
-
-		int32 XL;
-		int32 YL;
-		StringSize(GEngine->GetLargeFont(), XL, YL, *TileIndexString);
-		const float DrawX = 4.0f;
-		const float DrawY = FMath::FloorToFloat(Viewport->GetSizeXY().Y - YL - 4.0f);
-		Canvas->DrawShadowedString(DrawX, DrawY, *TileIndexString, GEngine->GetLargeFont(), FLinearColor::White);
-	}
-}
-
-FLinearColor FTileSetEditorViewportClient::GetBackgroundColor() const
-{
-	if (UPaperTileSet* TileSet = TileSetBeingEdited.Get())
-	{
-		return TileSet->BackgroundColor;
-	}
-	else
-	{
-		return FEditorViewportClient::GetBackgroundColor();
-	}
-}
-
-//////////////////////////////////////////////////////////////////////////
-// STileSetSelectorViewport
-
-STileSetSelectorViewport::~STileSetSelectorViewport()
-{
-	TypedViewportClient = nullptr;
-}
-
-void STileSetSelectorViewport::Construct(const FArguments& InArgs, UPaperTileSet* InTileSet, FEdModeTileMap* InTileMapEditor)
-{
-	SelectionTopLeft = FIntPoint::ZeroValue;
-	SelectionDimensions = FIntPoint::ZeroValue;
-
-	TileSetPtr = InTileSet;
-	TileMapEditor = InTileMapEditor;
-
-	TypedViewportClient = MakeShareable(new FTileSetEditorViewportClient(InTileSet));
-
-	SPaperEditorViewport::Construct(
-		SPaperEditorViewport::FArguments().OnSelectionChanged(this, &STileSetSelectorViewport::OnSelectionChanged),
-		TypedViewportClient.ToSharedRef());
-
-	// Make sure we get input instead of the viewport stealing it
-	ViewportWidget->SetVisibility(EVisibility::HitTestInvisible);
-}
-
-void STileSetSelectorViewport::ChangeTileSet(UPaperTileSet* InTileSet)
-{
-	if (InTileSet != TileSetPtr.Get())
-	{
-		TileSetPtr = InTileSet;
-		TypedViewportClient->TileSetBeingEdited = InTileSet;
-
-		// Update the selection rectangle
-		RefreshSelectionRectangle();
-		TypedViewportClient->Invalidate();
-	}
-}
-
-FText STileSetSelectorViewport::GetTitleText() const
-{
-	if (UPaperTileSet* TileSet = TileSetPtr.Get())
-	{
-		return FText::FromString(TileSet->GetName());
-	}
-
-	return LOCTEXT("TileSetSelectorTitle", "Tile Set Selector");
-}
-
-
-void STileSetSelectorViewport::OnSelectionChanged(FMarqueeOperation Marquee, bool bIsPreview)
-{
-	UPaperTileSet* TileSetBeingEdited = TileSetPtr.Get();
-	if (TileSetBeingEdited == nullptr)
-	{
-		return;
-	}
-
-	const FVector2D TopLeftUnrounded = Marquee.Rect.GetUpperLeft();
-	const FVector2D BottomRightUnrounded = Marquee.Rect.GetLowerRight();
-	if ((TopLeftUnrounded != FVector2D::ZeroVector) || Marquee.IsValid())
-	{
-		const int32 TileCountX = TileSetBeingEdited->GetTileCountX();
-		const int32 TileCountY = TileSetBeingEdited->GetTileCountY();
-
-		// Round down the top left corner
-		const FIntPoint TileTopLeft = TileSetBeingEdited->GetTileXYFromTextureUV(TopLeftUnrounded, false);
-		SelectionTopLeft = FIntPoint(FMath::Clamp<int32>(TileTopLeft.X, 0, TileCountX), FMath::Clamp<int32>(TileTopLeft.Y, 0, TileCountY));
-
-		// Round up the bottom right corner
-		const FIntPoint TileBottomRight = TileSetBeingEdited->GetTileXYFromTextureUV(BottomRightUnrounded, true);
-		const FIntPoint SelectionBottomRight(FMath::Clamp<int32>(TileBottomRight.X, 0, TileCountX), FMath::Clamp<int32>(TileBottomRight.Y, 0, TileCountY));
-
-		// Compute the new selection dimensions
-		SelectionDimensions = SelectionBottomRight - SelectionTopLeft;
-	}
-	else
-	{
-		SelectionTopLeft.X = 0;
-		SelectionTopLeft.Y = 0;
-		SelectionDimensions.X = 0;
-		SelectionDimensions.Y = 0;
-	}
-
-	OnTileSelectionChanged.Broadcast(SelectionTopLeft, SelectionDimensions);
-
-	const bool bHasSelection = (SelectionDimensions.X * SelectionDimensions.Y) > 0;
-	if (bIsPreview && bHasSelection)
-	{
-		if (TileMapEditor != nullptr)
-		{
-			TileMapEditor->SetActivePaint(TileSetBeingEdited, SelectionTopLeft, SelectionDimensions);
-
-			// Switch to paint brush mode if we were in the eraser mode since the user is trying to select some ink to paint with
-			if (TileMapEditor->GetActiveTool() == ETileMapEditorTool::Eraser)
-			{
-				TileMapEditor->SetActiveTool(ETileMapEditorTool::Paintbrush);
-			}
-		}
-
-		RefreshSelectionRectangle();
-	}
-}
-
-void STileSetSelectorViewport::RefreshSelectionRectangle()
-{
-	if (FTileSetEditorViewportClient* Client = TypedViewportClient.Get())
-	{
-		UPaperTileSet* TileSetBeingEdited = TileSetPtr.Get();
-
-		const bool bHasSelection = (SelectionDimensions.X * SelectionDimensions.Y) > 0;
-		Client->bHasValidPaintRectangle = bHasSelection && (TileSetBeingEdited != nullptr);
-
-		const int32 TileIndex = (bHasSelection && (TileSetBeingEdited != nullptr)) ? (SelectionTopLeft.X + SelectionTopLeft.Y * TileSetBeingEdited->GetTileCountX()) : INDEX_NONE;
-		Client->TileIndex = TileIndex;
-
-		if (bHasSelection && (TileSetBeingEdited != nullptr))
-		{
-			Client->ValidPaintRectangle.Color = FLinearColor::White;
-
-			const FIntPoint TopLeft = TileSetBeingEdited->GetTileUVFromTileXY(SelectionTopLeft);
-			const FIntPoint BottomRight = TileSetBeingEdited->GetTileUVFromTileXY(SelectionTopLeft + SelectionDimensions);
-			Client->ValidPaintRectangle.Dimensions = BottomRight - TopLeft;
-			Client->ValidPaintRectangle.TopLeft = TopLeft;
-		}
-	}
-}
-
 /////////////////////////////////////////////////////
 // STileSetPropertiesTabBody
 
@@ -305,502 +78,6 @@ public:
 	// End of SSingleObjectDetailsPanel interface
 };
 
-
-
-
-
-
-
-
-
-
-
-
-
-
-
-
-#include "Editor/UnrealEd/Public/SCommonEditorViewportToolbarBase.h"
-
-// In-viewport toolbar widget used in the tile set editor
-class STileSetEditorViewportToolbar : public SCommonEditorViewportToolbarBase
-{
-public:
-	SLATE_BEGIN_ARGS(STileSetEditorViewportToolbar) {}
-	SLATE_END_ARGS()
-
-	void Construct(const FArguments& InArgs, TSharedPtr<class ICommonEditorViewportToolbarInfoProvider> InInfoProvider);
-
-	// SCommonEditorViewportToolbarBase interface
-	virtual TSharedRef<SWidget> GenerateShowMenu() const override;
-	// End of SCommonEditorViewportToolbarBase
-};
-
-///////////////////////////////////////////////////////////
-// STileSetEditorViewportToolbar
-
-#include "SEditorViewport.h"
-#include "SpriteEditor/SpriteEditorCommands.h"
-
-void STileSetEditorViewportToolbar::Construct(const FArguments& InArgs, TSharedPtr<class ICommonEditorViewportToolbarInfoProvider> InInfoProvider)
-{
-	SCommonEditorViewportToolbarBase::Construct(SCommonEditorViewportToolbarBase::FArguments(), InInfoProvider);
-}
-
-TSharedRef<SWidget> STileSetEditorViewportToolbar::GenerateShowMenu() const
-{
-	GetInfoProvider().OnFloatingButtonClicked();
-
-	TSharedRef<SEditorViewport> ViewportRef = GetInfoProvider().GetViewportWidget();
-
-	const bool bInShouldCloseWindowAfterMenuSelection = true;
-	FMenuBuilder ShowMenuBuilder(bInShouldCloseWindowAfterMenuSelection, ViewportRef->GetCommandList());
-	{
-		ShowMenuBuilder.AddMenuEntry(FSpriteGeometryEditCommands::Get().SetShowNormals);
-	}
-
-	return ShowMenuBuilder.MakeWidget();
-}
-
-
-//////////////////////////////////////////////////////////////////////////
-// FSingleTileEditorViewportClient
-
-#include "PreviewScene.h"
-#include "SpriteEditor/SpriteEditorSelections.h"
-#include "SpriteEditor/SpriteEditorCommands.h"
-
-class FSingleTileEditorViewportClient : public FPaperEditorViewportClient, public ISpriteSelectionContext
-{
-public:
-	FSingleTileEditorViewportClient(UPaperTileSet* InTileSet);
-
-	// FViewportClient interface
-	virtual void Tick(float DeltaSeconds) override;
-	// End of FViewportClient interface
-
-	// FEditorViewportClient interface
-	virtual FLinearColor GetBackgroundColor() const override;
-	virtual void TrackingStarted(const struct FInputEventState& InInputState, bool bIsDragging, bool bNudge) override;
-	virtual void TrackingStopped() override;
-	// End of FEditorViewportClient interface
-
-	// ISpriteSelectionContext interface
-	virtual FVector2D SelectedItemConvertWorldSpaceDeltaToLocalSpace(const FVector& WorldSpaceDelta) const override;
-	virtual FVector2D WorldSpaceToTextureSpace(const FVector& SourcePoint) const override;
-	virtual FVector TextureSpaceToWorldSpace(const FVector2D& SourcePoint) const override;
-	virtual float SelectedItemGetUnitsPerPixel() const override;
-	virtual void BeginTransaction(const FText& SessionName) override;
-	virtual void MarkTransactionAsDirty() override;
-	virtual void EndTransaction() override;
-	virtual void InvalidateViewportAndHitProxies() override;
-	// End of ISpriteSelectionContext interface
-
-	void SetTileIndex(int32 InTileIndex);
-	int32 GetTileIndex() const;
-	void OnActiveTileIndexChanged(const FIntPoint& TopLeft, const FIntPoint& Dimensions);
-
-	void ActivateEditMode(TSharedPtr<FUICommandList> InCommandList);
-
-protected:
-	// FPaperEditorViewportClient interface
-	virtual FBox GetDesiredFocusBounds() const override;
-	// End of FPaperEditorViewportClient
-
-public:
-	// Tile set
-	UPaperTileSet* TileSet;
-
-	int32 TileBeingEditedIndex;
-
-	// Are we currently manipulating something?
-	bool bManipulating;
-
-	// Did we dirty something during manipulation?
-	bool bManipulationDirtiedSomething;
-
-	// Pointer back to the sprite editor viewport control that owns us
-	TWeakPtr<class SEditorViewport> SpriteEditorViewportPtr;
-
-	// The current transaction for undo/redo
-	class FScopedTransaction* ScopedTransaction;
-
-	// The preview scene
-	FPreviewScene OwnedPreviewScene;
-
-	// The preview sprite in the scene
-	UPaperSpriteComponent* PreviewTileSpriteComponent;
-};
-
-#include "PaperEditorShared/SpriteGeometryEditMode.h"
-#include "ScopedTransaction.h"
-
-FSingleTileEditorViewportClient::FSingleTileEditorViewportClient(UPaperTileSet* InTileSet)
-	: TileSet(InTileSet)
-	, TileBeingEditedIndex(INDEX_NONE)
-	, bManipulating(false)
-	, bManipulationDirtiedSomething(false)
-	, ScopedTransaction(nullptr)
-{
-	//@TODO: Should be able to set this to false eventually
-	SetRealtime(true);
-
-	// The tile map editor fully supports mode tools and isn't doing any incompatible stuff with the Widget
-	Widget->SetUsesEditorModeTools(ModeTools);
-
-	DrawHelper.bDrawGrid = true;//@TODO:
-	DrawHelper.bDrawPivot = false;
-
-	PreviewScene = &OwnedPreviewScene;
-	((FAssetEditorModeManager*)ModeTools)->SetPreviewScene(PreviewScene);
-
-	EngineShowFlags.DisableAdvancedFeatures();
-	EngineShowFlags.CompositeEditorPrimitives = true;
-
-	// Create a render component for the tile preview
-	PreviewTileSpriteComponent = NewObject<UPaperSpriteComponent>();
-
-	FStringAssetReference DefaultTranslucentMaterialName("/Paper2D/TranslucentUnlitSpriteMaterial.TranslucentUnlitSpriteMaterial");
-	UMaterialInterface* TranslucentMaterial = Cast<UMaterialInterface>(DefaultTranslucentMaterialName.TryLoad());
-	PreviewTileSpriteComponent->SetMaterial(0, TranslucentMaterial);
-
-	PreviewScene->AddComponent(PreviewTileSpriteComponent, FTransform::Identity);
-}
-
-FBox FSingleTileEditorViewportClient::GetDesiredFocusBounds() const
-{
-	UPaperSpriteComponent* ComponentToFocusOn = PreviewTileSpriteComponent;
-	return ComponentToFocusOn->Bounds.GetBox();
-}
-
-
-void FSingleTileEditorViewportClient::Tick(float DeltaSeconds)
-{
-	FPaperEditorViewportClient::Tick(DeltaSeconds);
-
-	if (!GIntraFrameDebuggingGameThread)
-	{
-		OwnedPreviewScene.GetWorld()->Tick(LEVELTICK_All, DeltaSeconds);
-	}
-}
-
-FLinearColor FSingleTileEditorViewportClient::GetBackgroundColor() const
-{
-	return TileSet->BackgroundColor;
-}
-
-void FSingleTileEditorViewportClient::TrackingStarted(const struct FInputEventState& InInputState, bool bIsDragging, bool bNudge)
-{
-	//@TODO: Should push this into FEditorViewportClient
-	// Begin transacting.  Give the current editor mode an opportunity to do the transacting.
-	const bool bTrackingHandledExternally = ModeTools->StartTracking(this, Viewport);
-
-	if (!bManipulating && bIsDragging && !bTrackingHandledExternally)
-	{
-		BeginTransaction(LOCTEXT("ModificationInViewport", "Modification in Viewport"));
-		bManipulating = true;
-		bManipulationDirtiedSomething = false;
-	}
-}
-
-void FSingleTileEditorViewportClient::TrackingStopped()
-{
-	// Stop transacting.  Give the current editor mode an opportunity to do the transacting.
-	const bool bTransactingHandledByEditorMode = ModeTools->EndTracking(this, Viewport);
-
-	if (bManipulating && !bTransactingHandledByEditorMode)
-	{
-		EndTransaction();
-		bManipulating = false;
-	}
-}
-
-FVector2D FSingleTileEditorViewportClient::SelectedItemConvertWorldSpaceDeltaToLocalSpace(const FVector& WorldSpaceDelta) const
-{
-	const FVector ProjectionX = WorldSpaceDelta.ProjectOnTo(PaperAxisX);
-	const FVector ProjectionY = WorldSpaceDelta.ProjectOnTo(PaperAxisY);
-
-	const float XValue = FMath::Sign(ProjectionX | PaperAxisX) * ProjectionX.Size();
-	const float YValue = FMath::Sign(ProjectionY | PaperAxisY) * ProjectionY.Size();
-
-	return FVector2D(XValue, YValue);
-}
-
-FVector2D FSingleTileEditorViewportClient::WorldSpaceToTextureSpace(const FVector& SourcePoint) const
-{
-	const FVector ProjectionX = SourcePoint.ProjectOnTo(PaperAxisX);
-	const FVector ProjectionY = -SourcePoint.ProjectOnTo(PaperAxisY);
-
-	const float XValue = FMath::Sign(ProjectionX | PaperAxisX) * ProjectionX.Size();
-	const float YValue = FMath::Sign(ProjectionY | PaperAxisY) * ProjectionY.Size();
-
-	return FVector2D(XValue, YValue);
-}
-
-FVector FSingleTileEditorViewportClient::TextureSpaceToWorldSpace(const FVector2D& SourcePoint) const
-{
-	return (SourcePoint.X * PaperAxisX) - (SourcePoint.Y * PaperAxisY);
-}
-
-float FSingleTileEditorViewportClient::SelectedItemGetUnitsPerPixel() const
-{
-	return 1.0f;
-}
-
-void FSingleTileEditorViewportClient::BeginTransaction(const FText& SessionName)
-{
-	if (ScopedTransaction == nullptr)
-	{
-		ScopedTransaction = new FScopedTransaction(SessionName);
-
-		TileSet->Modify();
-	}
-}
-
-void FSingleTileEditorViewportClient::MarkTransactionAsDirty()
-{
-	bManipulationDirtiedSomething = true;
-	Invalidate();
-}
-
-void FSingleTileEditorViewportClient::EndTransaction()
-{
-	if (bManipulationDirtiedSomething)
-	{
-		TileSet->PostEditChange();
-	}
-
-	bManipulationDirtiedSomething = false;
-
-	if (ScopedTransaction != nullptr)
-	{
-		delete ScopedTransaction;
-		ScopedTransaction = nullptr;
-	}
-}
-
-void FSingleTileEditorViewportClient::InvalidateViewportAndHitProxies()
-{
-	Invalidate();
-}
-
-void FSingleTileEditorViewportClient::SetTileIndex(int32 InTileIndex)
-{
-	const bool bNewIndexValid = (InTileIndex >= 0) && (InTileIndex < TileSet->GetTileCount());
-	TileBeingEditedIndex = bNewIndexValid ? InTileIndex : INDEX_NONE;
-
-	FSpriteGeometryEditMode* GeometryEditMode = ModeTools->GetActiveModeTyped<FSpriteGeometryEditMode>(FSpriteGeometryEditMode::EM_SpriteGeometry);
-	check(GeometryEditMode);
-	FSpriteGeometryCollection* GeomToEdit = nullptr;
-	
-	if (TileBeingEditedIndex != INDEX_NONE)
-	{
-		if (FPaperTileMetadata* Metadata = TileSet->GetMutableTileMetadata(InTileIndex))
-		{
-			GeomToEdit = &(Metadata->CollisionData);
-		}
-	}
-
-	// Tell the geometry editor about the new tile (if it exists)
-	GeometryEditMode->SetGeometryBeingEdited(GeomToEdit, /*bAllowCircles=*/ true, /*bAllowSubtractivePolygons=*/ false);
-
-	// Update the visual representation
-	UPaperSprite* DummySprite = nullptr;
-	if (TileBeingEditedIndex != INDEX_NONE)
-	{
-		//@TODO: Should use this to pick the correct material
-		//const bool bUseTranslucentBlend = Texture->HasAlphaChannel();
-
-		DummySprite = NewObject<UPaperSprite>();
- 		DummySprite->SpriteCollisionDomain = ESpriteCollisionMode::None;
- 		DummySprite->PivotMode = ESpritePivotMode::Center_Center;
- 		DummySprite->CollisionGeometry.GeometryType = ESpritePolygonMode::SourceBoundingBox;
- 		DummySprite->RenderGeometry.GeometryType = ESpritePolygonMode::SourceBoundingBox;
-		DummySprite->PixelsPerUnrealUnit = 1.0f;
-
-		FSpriteAssetInitParameters SpriteReinitParams;
-		SpriteReinitParams.Texture = TileSet->TileSheet;
-		TileSet->GetTileUV(TileBeingEditedIndex, /*out*/ SpriteReinitParams.Offset);
-		SpriteReinitParams.Dimension = FVector2D(TileSet->TileWidth, TileSet->TileHeight);
-		DummySprite->InitializeSprite(SpriteReinitParams);
-	}
-	PreviewTileSpriteComponent->SetSprite(DummySprite);
-
-	// Update the default geometry bounds
-	const FVector2D HalfTileSize(TileSet->TileWidth * 0.5f, TileSet->TileHeight * 0.5f);
-	FBox2D DesiredBounds(ForceInitToZero);
-	DesiredBounds.Min = -HalfTileSize;
-	DesiredBounds.Max = HalfTileSize;
-	GeometryEditMode->SetNewGeometryPreferredBounds(DesiredBounds);
-
-	// Redraw the viewport
-	Invalidate();
-}
-
-void FSingleTileEditorViewportClient::OnActiveTileIndexChanged(const FIntPoint& TopLeft, const FIntPoint& Dimensions)
-{
-	const int32 NewIndex = (TileSet->GetTileCountX() * TopLeft.Y) + TopLeft.X;
-	SetTileIndex(NewIndex);
-}
-
-int32 FSingleTileEditorViewportClient::GetTileIndex() const
-{
-	return TileBeingEditedIndex;
-}
-
-void FSingleTileEditorViewportClient::ActivateEditMode(TSharedPtr<FUICommandList> InCommandList)
-{
-	// Activate the sprite geometry edit mode
-	//@TODO: ModeTools->SetToolkitHost(SpriteEditorPtr.Pin()->GetToolkitHost());
-	ModeTools->SetDefaultMode(FSpriteGeometryEditMode::EM_SpriteGeometry);
-	ModeTools->ActivateDefaultMode();
-	ModeTools->SetWidgetMode(FWidget::WM_Translate);
-
-	FSpriteGeometryEditMode* GeometryEditMode = ModeTools->GetActiveModeTyped<FSpriteGeometryEditMode>(FSpriteGeometryEditMode::EM_SpriteGeometry);
-	check(GeometryEditMode);
-	GeometryEditMode->SetEditorContext(this);
-	GeometryEditMode->BindCommands(InCommandList /*SpriteEditorViewportPtr.Pin()->GetCommandList()*/);
-
-	const FLinearColor CollisionShapeColor(0.0f, 0.7f, 1.0f, 1.0f); //@TODO: Duplicated constant from SpriteEditingConstants
-	GeometryEditMode->SetGeometryColors(CollisionShapeColor, FLinearColor::White);
-}
-
-//////////////////////////////////////////////////////////////////////////
-// SSingleTileEditorViewport
-
-#include "SEditorViewport.h"
-
-class SSingleTileEditorViewport : public SEditorViewport, public ICommonEditorViewportToolbarInfoProvider
-{
-public:
-	SLATE_BEGIN_ARGS(SSingleTileEditorViewport) {}
-	SLATE_END_ARGS()
-
-	~SSingleTileEditorViewport();
-
-	void Construct(const FArguments& InArgs, TSharedPtr<class FSingleTileEditorViewportClient> InViewportClient);
-
-	// SEditorViewport interface
-	virtual TSharedPtr<SWidget> MakeViewportToolbar() override;
-	virtual TSharedRef<FEditorViewportClient> MakeEditorViewportClient() override;
-	// End of SEditorViewport interface
-
-	// ICommonEditorViewportToolbarInfoProvider interface
-	virtual TSharedRef<class SEditorViewport> GetViewportWidget() override;
-	virtual TSharedPtr<FExtender> GetExtenders() const override;
-	virtual void OnFloatingButtonClicked() override;
-	// End of ICommonEditorViewportToolbarInfoProvider interface
-
-protected:
-	FText GetTitleText() const;
-
-	bool IsVisible() const;
-
-private:
-	TSharedPtr<class FSingleTileEditorViewportClient> TypedViewportClient;
-	class FEdModeTileMap* TileMapEditor;
-};
-
-//////////////////////////////////////////////////////////////////////////
-// SSingleTileEditorViewport
-
-SSingleTileEditorViewport::~SSingleTileEditorViewport()
-{
-	TypedViewportClient.Reset();
-}
-
-void SSingleTileEditorViewport::Construct(const FArguments& InArgs, TSharedPtr<class FSingleTileEditorViewportClient> InViewportClient)
-{
-	TypedViewportClient = InViewportClient;
-
-	SEditorViewport::Construct(SEditorViewport::FArguments());
-
-	TSharedRef<SWidget> ParentContents = ChildSlot.GetWidget();
-
-	this->ChildSlot
-	[
-		SNew(SOverlay)
-		+SOverlay::Slot()
-		[
-			ParentContents
-		]
-
-		+SOverlay::Slot()
-		.VAlign(VAlign_Bottom)
-		[
-			SNew(SBorder)
-			.BorderImage(FPaperStyle::Get()->GetBrush("Paper2D.Common.ViewportTitleBackground"))
-			.HAlign(HAlign_Fill)
-			.Visibility(EVisibility::HitTestInvisible)
-			[
-				SNew(SVerticalBox)
-				// Title text/icon
-				+SVerticalBox::Slot()
-				.AutoHeight()
-				[
-					SNew(SHorizontalBox)
-					+SHorizontalBox::Slot()
-					.HAlign(HAlign_Center)
-					.FillWidth(1.f)
-					[
-						SNew(STextBlock)
-						.TextStyle(FPaperStyle::Get(), "Paper2D.Common.ViewportTitleTextStyle")
-						.Text(this, &SSingleTileEditorViewport::GetTitleText)
-					]
-				]
-			]
-		]
-	];
-}
-
-TSharedPtr<SWidget> SSingleTileEditorViewport::MakeViewportToolbar()
-{
-	return SNew(STileSetEditorViewportToolbar, SharedThis(this));
-}
-
-TSharedRef<FEditorViewportClient> SSingleTileEditorViewport::MakeEditorViewportClient()
-{
-	TypedViewportClient->VisibilityDelegate.BindSP(this, &SSingleTileEditorViewport::IsVisible);
-
-	return TypedViewportClient.ToSharedRef();
-}
-
-TSharedRef<class SEditorViewport> SSingleTileEditorViewport::GetViewportWidget()
-{
-	return SharedThis(this);
-}
-
-TSharedPtr<FExtender> SSingleTileEditorViewport::GetExtenders() const
-{
-	TSharedPtr<FExtender> Result(MakeShareable(new FExtender));
-	return Result;
-}
-
-void SSingleTileEditorViewport::OnFloatingButtonClicked()
-{
-}
-
-FText SSingleTileEditorViewport::GetTitleText() const
-{
-	const int32 CurrentTileIndex = TypedViewportClient->GetTileIndex();
-	if (CurrentTileIndex != INDEX_NONE)
-	{
-		FNumberFormattingOptions NoGroupingFormat;
-		NoGroupingFormat.SetUseGrouping(false);
-
-		return FText::Format(LOCTEXT("SingleTileEditorViewportTitle", "Editing tile #{0}"), FText::AsNumber(CurrentTileIndex, &NoGroupingFormat));
-	}
-	else
-	{
-		return LOCTEXT("SingleTileEditorViewportTitle_NoTile", "Tile Editor - Select a tile");
-	}
-}
-
-bool SSingleTileEditorViewport::IsVisible() const
-{
-	return true;//@TODO: Determine this better so viewport ticking optimizations can take place
-}
-
 //////////////////////////////////////////////////////////////////////////
 // FTileSetEditor
 
diff --git a/Engine/Plugins/2D/Paper2D/Source/Paper2DEditor/Private/TileSetEditor.h b/Engine/Plugins/2D/Paper2D/Source/Paper2DEditor/Private/TileSetEditor.h
--- a/Engine/Plugins/2D/Paper2D/Source/Paper2DEditor/Private/TileSetEditor.h
+++ b/Engine/Plugins/2D/Paper2D/Source/Paper2DEditor/Private/TileSetEditor.h
@@ -60,44 +60,3 @@ protected:
 
 	FDelegateHandle OnPropertyChangedHandle;
 };
-
-//////////////////////////////////////////////////////////////////////////
-// STileSetSelectorViewport
-
-class STileSetSelectorViewport : public SPaperEditorViewport
-{
-public:
-	DECLARE_MULTICAST_DELEGATE_TwoParams(FOnTileViewportSelectionChanged, const FIntPoint& /*TopLeft*/, const FIntPoint& /*Dimensions*/);
-
-	SLATE_BEGIN_ARGS(STileSetSelectorViewport) {}
-	SLATE_END_ARGS()
-
-	~STileSetSelectorViewport();
-
-	void Construct(const FArguments& InArgs, UPaperTileSet* InTileSet, class FEdModeTileMap* InTileMapEditor);
-
-	void ChangeTileSet(UPaperTileSet* InTileSet);
-
-	FOnTileViewportSelectionChanged& GetTileSelectionChanged()
-	{
-		return OnTileSelectionChanged;
-	}
-
-protected:
-	// SPaperEditorViewport interface
-	virtual FText GetTitleText() const override;
-	// End of SPaperEditorViewport interface
-
-private:
-	void OnSelectionChanged(FMarqueeOperation Marquee, bool bIsPreview);
-	void RefreshSelectionRectangle();
-
-private:
-	TWeakObjectPtr<class UPaperTileSet> TileSetPtr;
-	TSharedPtr<class FTileSetEditorViewportClient> TypedViewportClient;
-	class FEdModeTileMap* TileMapEditor;
-	FIntPoint SelectionTopLeft;
-	FIntPoint SelectionDimensions;
-
-	FOnTileViewportSelectionChanged OnTileSelectionChanged;
-};
\ No newline at end of file
diff --git a/Engine/Plugins/2D/Paper2D/Source/Paper2DEditor/Private/TileSetEditor/SingleTileEditorViewport.cpp b/Engine/Plugins/2D/Paper2D/Source/Paper2DEditor/Private/TileSetEditor/SingleTileEditorViewport.cpp
new file mode 100644
--- /dev/null
+++ b/Engine/Plugins/2D/Paper2D/Source/Paper2DEditor/Private/TileSetEditor/SingleTileEditorViewport.cpp
@@ -0,0 +1,150 @@
+// Copyright 1998-2015 Epic Games, Inc. All Rights Reserved.
+
+#include "Paper2DEditorPrivatePCH.h"
+#include "PaperStyle.h"
+#include "SingleTileEditorViewport.h"
+#include "SpriteEditor/SpriteEditorCommands.h"
+#include "SingleTileEditorViewportClient.h"
+
+#define LOCTEXT_NAMESPACE "TileSetEditor"
+
+//////////////////////////////////////////////////////////////////////////
+// STileSetEditorViewportToolbar
+
+// In-viewport toolbar widget used in the tile set editor
+class STileSetEditorViewportToolbar : public SCommonEditorViewportToolbarBase
+{
+public:
+	SLATE_BEGIN_ARGS(STileSetEditorViewportToolbar) {}
+	SLATE_END_ARGS()
+
+	void Construct(const FArguments& InArgs, TSharedPtr<class ICommonEditorViewportToolbarInfoProvider> InInfoProvider);
+
+	// SCommonEditorViewportToolbarBase interface
+	virtual TSharedRef<SWidget> GenerateShowMenu() const override;
+	// End of SCommonEditorViewportToolbarBase
+};
+
+void STileSetEditorViewportToolbar::Construct(const FArguments& InArgs, TSharedPtr<class ICommonEditorViewportToolbarInfoProvider> InInfoProvider)
+{
+	SCommonEditorViewportToolbarBase::Construct(SCommonEditorViewportToolbarBase::FArguments(), InInfoProvider);
+}
+
+TSharedRef<SWidget> STileSetEditorViewportToolbar::GenerateShowMenu() const
+{
+	GetInfoProvider().OnFloatingButtonClicked();
+
+	TSharedRef<SEditorViewport> ViewportRef = GetInfoProvider().GetViewportWidget();
+
+	const bool bInShouldCloseWindowAfterMenuSelection = true;
+	FMenuBuilder ShowMenuBuilder(bInShouldCloseWindowAfterMenuSelection, ViewportRef->GetCommandList());
+	{
+		ShowMenuBuilder.AddMenuEntry(FSpriteGeometryEditCommands::Get().SetShowNormals);
+	}
+
+	return ShowMenuBuilder.MakeWidget();
+}
+
+//////////////////////////////////////////////////////////////////////////
+// SSingleTileEditorViewport
+
+SSingleTileEditorViewport::~SSingleTileEditorViewport()
+{
+	TypedViewportClient.Reset();
+}
+
+void SSingleTileEditorViewport::Construct(const FArguments& InArgs, TSharedPtr<class FSingleTileEditorViewportClient> InViewportClient)
+{
+	TypedViewportClient = InViewportClient;
+
+	SEditorViewport::Construct(SEditorViewport::FArguments());
+
+	TSharedRef<SWidget> ParentContents = ChildSlot.GetWidget();
+
+	this->ChildSlot
+	[
+		SNew(SOverlay)
+		+SOverlay::Slot()
+		[
+			ParentContents
+		]
+
+		+SOverlay::Slot()
+		.VAlign(VAlign_Bottom)
+		[
+			SNew(SBorder)
+			.BorderImage(FPaperStyle::Get()->GetBrush("Paper2D.Common.ViewportTitleBackground"))
+			.HAlign(HAlign_Fill)
+			.Visibility(EVisibility::HitTestInvisible)
+			[
+				SNew(SVerticalBox)
+				// Title text/icon
+				+SVerticalBox::Slot()
+				.AutoHeight()
+				[
+					SNew(SHorizontalBox)
+					+SHorizontalBox::Slot()
+					.HAlign(HAlign_Center)
+					.FillWidth(1.f)
+					[
+						SNew(STextBlock)
+						.TextStyle(FPaperStyle::Get(), "Paper2D.Common.ViewportTitleTextStyle")
+						.Text(this, &SSingleTileEditorViewport::GetTitleText)
+					]
+				]
+			]
+		]
+	];
+}
+
+TSharedPtr<SWidget> SSingleTileEditorViewport::MakeViewportToolbar()
+{
+	return SNew(STileSetEditorViewportToolbar, SharedThis(this));
+}
+
+TSharedRef<FEditorViewportClient> SSingleTileEditorViewport::MakeEditorViewportClient()
+{
+	TypedViewportClient->VisibilityDelegate.BindSP(this, &SSingleTileEditorViewport::IsVisible);
+
+	return TypedViewportClient.ToSharedRef();
+}
+
+TSharedRef<class SEditorViewport> SSingleTileEditorViewport::GetViewportWidget()
+{
+	return SharedThis(this);
+}
+
+TSharedPtr<FExtender> SSingleTileEditorViewport::GetExtenders() const
+{
+	TSharedPtr<FExtender> Result(MakeShareable(new FExtender));
+	return Result;
+}
+
+void SSingleTileEditorViewport::OnFloatingButtonClicked()
+{
+}
+
+FText SSingleTileEditorViewport::GetTitleText() const
+{
+	const int32 CurrentTileIndex = TypedViewportClient->GetTileIndex();
+	if (CurrentTileIndex != INDEX_NONE)
+	{
+		FNumberFormattingOptions NoGroupingFormat;
+		NoGroupingFormat.SetUseGrouping(false);
+
+		return FText::Format(LOCTEXT("SingleTileEditorViewportTitle", "Editing tile #{0}"), FText::AsNumber(CurrentTileIndex, &NoGroupingFormat));
+	}
+	else
+	{
+		return LOCTEXT("SingleTileEditorViewportTitle_NoTile", "Tile Editor - Select a tile");
+	}
+}
+
+bool SSingleTileEditorViewport::IsVisible() const
+{
+	return true;//@TODO: Determine this better so viewport ticking optimizations can take place
+}
+
+//////////////////////////////////////////////////////////////////////////
+
+#undef LOCTEXT_NAMESPACE
\ No newline at end of file
diff --git a/Engine/Plugins/2D/Paper2D/Source/Paper2DEditor/Private/TileSetEditor/SingleTileEditorViewport.h b/Engine/Plugins/2D/Paper2D/Source/Paper2DEditor/Private/TileSetEditor/SingleTileEditorViewport.h
new file mode 100644
--- /dev/null
+++ b/Engine/Plugins/2D/Paper2D/Source/Paper2DEditor/Private/TileSetEditor/SingleTileEditorViewport.h
@@ -0,0 +1,44 @@
+// Copyright 1998-2015 Epic Games, Inc. All Rights Reserved.
+
+#pragma once
+
+#include "SEditorViewport.h"
+
+class FSingleTileEditorViewportClient;
+class FEdModeTileMap;
+
+#include "Editor/UnrealEd/Public/SCommonEditorViewportToolbarBase.h"
+
+//////////////////////////////////////////////////////////////////////////
+// SSingleTileEditorViewport
+
+class SSingleTileEditorViewport : public SEditorViewport, public ICommonEditorViewportToolbarInfoProvider
+{
+public:
+	SLATE_BEGIN_ARGS(SSingleTileEditorViewport) {}
+	SLATE_END_ARGS()
+
+	~SSingleTileEditorViewport();
+
+	void Construct(const FArguments& InArgs, TSharedPtr<FSingleTileEditorViewportClient> InViewportClient);
+
+	// SEditorViewport interface
+	virtual TSharedPtr<SWidget> MakeViewportToolbar() override;
+	virtual TSharedRef<FEditorViewportClient> MakeEditorViewportClient() override;
+	// End of SEditorViewport interface
+
+	// ICommonEditorViewportToolbarInfoProvider interface
+	virtual TSharedRef<SEditorViewport> GetViewportWidget() override;
+	virtual TSharedPtr<FExtender> GetExtenders() const override;
+	virtual void OnFloatingButtonClicked() override;
+	// End of ICommonEditorViewportToolbarInfoProvider interface
+
+protected:
+	FText GetTitleText() const;
+
+	bool IsVisible() const;
+
+private:
+	TSharedPtr<FSingleTileEditorViewportClient> TypedViewportClient;
+	FEdModeTileMap* TileMapEditor;
+};
diff --git a/Engine/Plugins/2D/Paper2D/Source/Paper2DEditor/Private/TileSetEditor/SingleTileEditorViewportClient.cpp b/Engine/Plugins/2D/Paper2D/Source/Paper2DEditor/Private/TileSetEditor/SingleTileEditorViewportClient.cpp
new file mode 100644
--- /dev/null
+++ b/Engine/Plugins/2D/Paper2D/Source/Paper2DEditor/Private/TileSetEditor/SingleTileEditorViewportClient.cpp
@@ -0,0 +1,246 @@
+// Copyright 1998-2015 Epic Games, Inc. All Rights Reserved.
+
+#include "Paper2DEditorPrivatePCH.h"
+#include "SingleTileEditorViewportClient.h"
+#include "SpriteEditor/SpriteEditorCommands.h"
+#include "PaperEditorShared/SpriteGeometryEditMode.h"
+#include "ScopedTransaction.h"
+
+#define LOCTEXT_NAMESPACE "TileSetEditor"
+
+//////////////////////////////////////////////////////////////////////////
+// FSingleTileEditorViewportClient
+
+FSingleTileEditorViewportClient::FSingleTileEditorViewportClient(UPaperTileSet* InTileSet)
+	: TileSet(InTileSet)
+	, TileBeingEditedIndex(INDEX_NONE)
+	, bManipulating(false)
+	, bManipulationDirtiedSomething(false)
+	, ScopedTransaction(nullptr)
+{
+	//@TODO: Should be able to set this to false eventually
+	SetRealtime(true);
+
+	// The tile map editor fully supports mode tools and isn't doing any incompatible stuff with the Widget
+	Widget->SetUsesEditorModeTools(ModeTools);
+
+	DrawHelper.bDrawGrid = true;//@TODO:
+	DrawHelper.bDrawPivot = false;
+
+	PreviewScene = &OwnedPreviewScene;
+	((FAssetEditorModeManager*)ModeTools)->SetPreviewScene(PreviewScene);
+
+	EngineShowFlags.DisableAdvancedFeatures();
+	EngineShowFlags.CompositeEditorPrimitives = true;
+
+	// Create a render component for the tile preview
+	PreviewTileSpriteComponent = NewObject<UPaperSpriteComponent>();
+
+	FStringAssetReference DefaultTranslucentMaterialName("/Paper2D/TranslucentUnlitSpriteMaterial.TranslucentUnlitSpriteMaterial");
+	UMaterialInterface* TranslucentMaterial = Cast<UMaterialInterface>(DefaultTranslucentMaterialName.TryLoad());
+	PreviewTileSpriteComponent->SetMaterial(0, TranslucentMaterial);
+
+	PreviewScene->AddComponent(PreviewTileSpriteComponent, FTransform::Identity);
+}
+
+FBox FSingleTileEditorViewportClient::GetDesiredFocusBounds() const
+{
+	UPaperSpriteComponent* ComponentToFocusOn = PreviewTileSpriteComponent;
+	return ComponentToFocusOn->Bounds.GetBox();
+}
+
+
+void FSingleTileEditorViewportClient::Tick(float DeltaSeconds)
+{
+	FPaperEditorViewportClient::Tick(DeltaSeconds);
+
+	if (!GIntraFrameDebuggingGameThread)
+	{
+		OwnedPreviewScene.GetWorld()->Tick(LEVELTICK_All, DeltaSeconds);
+	}
+}
+
+FLinearColor FSingleTileEditorViewportClient::GetBackgroundColor() const
+{
+	return TileSet->BackgroundColor;
+}
+
+void FSingleTileEditorViewportClient::TrackingStarted(const struct FInputEventState& InInputState, bool bIsDragging, bool bNudge)
+{
+	//@TODO: Should push this into FEditorViewportClient
+	// Begin transacting.  Give the current editor mode an opportunity to do the transacting.
+	const bool bTrackingHandledExternally = ModeTools->StartTracking(this, Viewport);
+
+	if (!bManipulating && bIsDragging && !bTrackingHandledExternally)
+	{
+		BeginTransaction(LOCTEXT("ModificationInViewport", "Modification in Viewport"));
+		bManipulating = true;
+		bManipulationDirtiedSomething = false;
+	}
+}
+
+void FSingleTileEditorViewportClient::TrackingStopped()
+{
+	// Stop transacting.  Give the current editor mode an opportunity to do the transacting.
+	const bool bTransactingHandledByEditorMode = ModeTools->EndTracking(this, Viewport);
+
+	if (bManipulating && !bTransactingHandledByEditorMode)
+	{
+		EndTransaction();
+		bManipulating = false;
+	}
+}
+
+FVector2D FSingleTileEditorViewportClient::SelectedItemConvertWorldSpaceDeltaToLocalSpace(const FVector& WorldSpaceDelta) const
+{
+	const FVector ProjectionX = WorldSpaceDelta.ProjectOnTo(PaperAxisX);
+	const FVector ProjectionY = WorldSpaceDelta.ProjectOnTo(PaperAxisY);
+
+	const float XValue = FMath::Sign(ProjectionX | PaperAxisX) * ProjectionX.Size();
+	const float YValue = FMath::Sign(ProjectionY | PaperAxisY) * ProjectionY.Size();
+
+	return FVector2D(XValue, YValue);
+}
+
+FVector2D FSingleTileEditorViewportClient::WorldSpaceToTextureSpace(const FVector& SourcePoint) const
+{
+	const FVector ProjectionX = SourcePoint.ProjectOnTo(PaperAxisX);
+	const FVector ProjectionY = -SourcePoint.ProjectOnTo(PaperAxisY);
+
+	const float XValue = FMath::Sign(ProjectionX | PaperAxisX) * ProjectionX.Size();
+	const float YValue = FMath::Sign(ProjectionY | PaperAxisY) * ProjectionY.Size();
+
+	return FVector2D(XValue, YValue);
+}
+
+FVector FSingleTileEditorViewportClient::TextureSpaceToWorldSpace(const FVector2D& SourcePoint) const
+{
+	return (SourcePoint.X * PaperAxisX) - (SourcePoint.Y * PaperAxisY);
+}
+
+float FSingleTileEditorViewportClient::SelectedItemGetUnitsPerPixel() const
+{
+	return 1.0f;
+}
+
+void FSingleTileEditorViewportClient::BeginTransaction(const FText& SessionName)
+{
+	if (ScopedTransaction == nullptr)
+	{
+		ScopedTransaction = new FScopedTransaction(SessionName);
+
+		TileSet->Modify();
+	}
+}
+
+void FSingleTileEditorViewportClient::MarkTransactionAsDirty()
+{
+	bManipulationDirtiedSomething = true;
+	Invalidate();
+}
+
+void FSingleTileEditorViewportClient::EndTransaction()
+{
+	if (bManipulationDirtiedSomething)
+	{
+		TileSet->PostEditChange();
+	}
+
+	bManipulationDirtiedSomething = false;
+
+	if (ScopedTransaction != nullptr)
+	{
+		delete ScopedTransaction;
+		ScopedTransaction = nullptr;
+	}
+}
+
+void FSingleTileEditorViewportClient::InvalidateViewportAndHitProxies()
+{
+	Invalidate();
+}
+
+void FSingleTileEditorViewportClient::SetTileIndex(int32 InTileIndex)
+{
+	const bool bNewIndexValid = (InTileIndex >= 0) && (InTileIndex < TileSet->GetTileCount());
+	TileBeingEditedIndex = bNewIndexValid ? InTileIndex : INDEX_NONE;
+
+	FSpriteGeometryEditMode* GeometryEditMode = ModeTools->GetActiveModeTyped<FSpriteGeometryEditMode>(FSpriteGeometryEditMode::EM_SpriteGeometry);
+	check(GeometryEditMode);
+	FSpriteGeometryCollection* GeomToEdit = nullptr;
+	
+	if (TileBeingEditedIndex != INDEX_NONE)
+	{
+		if (FPaperTileMetadata* Metadata = TileSet->GetMutableTileMetadata(InTileIndex))
+		{
+			GeomToEdit = &(Metadata->CollisionData);
+		}
+	}
+
+	// Tell the geometry editor about the new tile (if it exists)
+	GeometryEditMode->SetGeometryBeingEdited(GeomToEdit, /*bAllowCircles=*/ true, /*bAllowSubtractivePolygons=*/ false);
+
+	// Update the visual representation
+	UPaperSprite* DummySprite = nullptr;
+	if (TileBeingEditedIndex != INDEX_NONE)
+	{
+		//@TODO: Should use this to pick the correct material
+		//const bool bUseTranslucentBlend = Texture->HasAlphaChannel();
+
+		DummySprite = NewObject<UPaperSprite>();
+ 		DummySprite->SpriteCollisionDomain = ESpriteCollisionMode::None;
+ 		DummySprite->PivotMode = ESpritePivotMode::Center_Center;
+ 		DummySprite->CollisionGeometry.GeometryType = ESpritePolygonMode::SourceBoundingBox;
+ 		DummySprite->RenderGeometry.GeometryType = ESpritePolygonMode::SourceBoundingBox;
+		DummySprite->PixelsPerUnrealUnit = 1.0f;
+
+		FSpriteAssetInitParameters SpriteReinitParams;
+		SpriteReinitParams.Texture = TileSet->TileSheet;
+		TileSet->GetTileUV(TileBeingEditedIndex, /*out*/ SpriteReinitParams.Offset);
+		SpriteReinitParams.Dimension = FVector2D(TileSet->TileWidth, TileSet->TileHeight);
+		DummySprite->InitializeSprite(SpriteReinitParams);
+	}
+	PreviewTileSpriteComponent->SetSprite(DummySprite);
+
+	// Update the default geometry bounds
+	const FVector2D HalfTileSize(TileSet->TileWidth * 0.5f, TileSet->TileHeight * 0.5f);
+	FBox2D DesiredBounds(ForceInitToZero);
+	DesiredBounds.Min = -HalfTileSize;
+	DesiredBounds.Max = HalfTileSize;
+	GeometryEditMode->SetNewGeometryPreferredBounds(DesiredBounds);
+
+	// Redraw the viewport
+	Invalidate();
+}
+
+void FSingleTileEditorViewportClient::OnActiveTileIndexChanged(const FIntPoint& TopLeft, const FIntPoint& Dimensions)
+{
+	const int32 NewIndex = (TileSet->GetTileCountX() * TopLeft.Y) + TopLeft.X;
+	SetTileIndex(NewIndex);
+}
+
+int32 FSingleTileEditorViewportClient::GetTileIndex() const
+{
+	return TileBeingEditedIndex;
+}
+
+void FSingleTileEditorViewportClient::ActivateEditMode(TSharedPtr<FUICommandList> InCommandList)
+{
+	// Activate the sprite geometry edit mode
+	//@TODO: ModeTools->SetToolkitHost(SpriteEditorPtr.Pin()->GetToolkitHost());
+	ModeTools->SetDefaultMode(FSpriteGeometryEditMode::EM_SpriteGeometry);
+	ModeTools->ActivateDefaultMode();
+	ModeTools->SetWidgetMode(FWidget::WM_Translate);
+
+	FSpriteGeometryEditMode* GeometryEditMode = ModeTools->GetActiveModeTyped<FSpriteGeometryEditMode>(FSpriteGeometryEditMode::EM_SpriteGeometry);
+	check(GeometryEditMode);
+	GeometryEditMode->SetEditorContext(this);
+	GeometryEditMode->BindCommands(InCommandList /*SpriteEditorViewportPtr.Pin()->GetCommandList()*/);
+
+	const FLinearColor CollisionShapeColor(0.0f, 0.7f, 1.0f, 1.0f); //@TODO: Duplicated constant from SpriteEditingConstants
+	GeometryEditMode->SetGeometryColors(CollisionShapeColor, FLinearColor::White);
+}
+
+//////////////////////////////////////////////////////////////////////////
+
+#undef LOCTEXT_NAMESPACE
\ No newline at end of file
diff --git a/Engine/Plugins/2D/Paper2D/Source/Paper2DEditor/Private/TileSetEditor/SingleTileEditorViewportClient.h b/Engine/Plugins/2D/Paper2D/Source/Paper2DEditor/Private/TileSetEditor/SingleTileEditorViewportClient.h
new file mode 100644
--- /dev/null
+++ b/Engine/Plugins/2D/Paper2D/Source/Paper2DEditor/Private/TileSetEditor/SingleTileEditorViewportClient.h
@@ -0,0 +1,72 @@
+// Copyright 1998-2015 Epic Games, Inc. All Rights Reserved.
+
+#pragma once
+
+#include "PreviewScene.h"
+#include "PaperEditorViewportClient.h"
+#include "SpriteEditor/SpriteEditorSelections.h"
+
+//////////////////////////////////////////////////////////////////////////
+// FSingleTileEditorViewportClient
+
+class FSingleTileEditorViewportClient : public FPaperEditorViewportClient, public ISpriteSelectionContext
+{
+public:
+	FSingleTileEditorViewportClient(UPaperTileSet* InTileSet);
+
+	// FViewportClient interface
+	virtual void Tick(float DeltaSeconds) override;
+	// End of FViewportClient interface
+
+	// FEditorViewportClient interface
+	virtual FLinearColor GetBackgroundColor() const override;
+	virtual void TrackingStarted(const struct FInputEventState& InInputState, bool bIsDragging, bool bNudge) override;
+	virtual void TrackingStopped() override;
+	// End of FEditorViewportClient interface
+
+	// ISpriteSelectionContext interface
+	virtual FVector2D SelectedItemConvertWorldSpaceDeltaToLocalSpace(const FVector& WorldSpaceDelta) const override;
+	virtual FVector2D WorldSpaceToTextureSpace(const FVector& SourcePoint) const override;
+	virtual FVector TextureSpaceToWorldSpace(const FVector2D& SourcePoint) const override;
+	virtual float SelectedItemGetUnitsPerPixel() const override;
+	virtual void BeginTransaction(const FText& SessionName) override;
+	virtual void MarkTransactionAsDirty() override;
+	virtual void EndTransaction() override;
+	virtual void InvalidateViewportAndHitProxies() override;
+	// End of ISpriteSelectionContext interface
+
+	void SetTileIndex(int32 InTileIndex);
+	int32 GetTileIndex() const;
+	void OnActiveTileIndexChanged(const FIntPoint& TopLeft, const FIntPoint& Dimensions);
+
+	void ActivateEditMode(TSharedPtr<FUICommandList> InCommandList);
+
+protected:
+	// FPaperEditorViewportClient interface
+	virtual FBox GetDesiredFocusBounds() const override;
+	// End of FPaperEditorViewportClient
+
+public:
+	// Tile set
+	UPaperTileSet* TileSet;
+
+	int32 TileBeingEditedIndex;
+
+	// Are we currently manipulating something?
+	bool bManipulating;
+
+	// Did we dirty something during manipulation?
+	bool bManipulationDirtiedSomething;
+
+	// Pointer back to the sprite editor viewport control that owns us
+	TWeakPtr<class SEditorViewport> SpriteEditorViewportPtr;
+
+	// The current transaction for undo/redo
+	class FScopedTransaction* ScopedTransaction;
+
+	// The preview scene
+	FPreviewScene OwnedPreviewScene;
+
+	// The preview sprite in the scene
+	UPaperSpriteComponent* PreviewTileSpriteComponent;
+};
diff --git a/Engine/Plugins/2D/Paper2D/Source/Paper2DEditor/Private/TileSetEditor/TileSetEditorViewportClient.cpp b/Engine/Plugins/2D/Paper2D/Source/Paper2DEditor/Private/TileSetEditor/TileSetEditorViewportClient.cpp
new file mode 100644
--- /dev/null
+++ b/Engine/Plugins/2D/Paper2D/Source/Paper2DEditor/Private/TileSetEditor/TileSetEditorViewportClient.cpp
@@ -0,0 +1,94 @@
+// Copyright 1998-2015 Epic Games, Inc. All Rights Reserved.
+
+#include "Paper2DEditorPrivatePCH.h"
+#include "TileSetEditorViewportClient.h"
+
+#define LOCTEXT_NAMESPACE "TileSetEditor"
+
+//////////////////////////////////////////////////////////////////////////
+// FTileSetEditorViewportClient
+
+FTileSetEditorViewportClient::FTileSetEditorViewportClient(UPaperTileSet* InTileSet)
+	: TileSetBeingEdited(InTileSet)
+	, bHasValidPaintRectangle(false)
+	, TileIndex(INDEX_NONE)
+{
+}
+
+void FTileSetEditorViewportClient::Draw(FViewport* Viewport, FCanvas* Canvas)
+{
+	// Clear the viewport
+	Canvas->Clear(GetBackgroundColor());
+
+	// Can only proceed if we have a valid tile set
+	UPaperTileSet* TileSet = TileSetBeingEdited.Get();
+	if (TileSet == nullptr)
+	{
+		return;
+	}
+
+	UTexture2D* Texture = TileSet->TileSheet;
+
+	if (Texture != nullptr)
+	{
+		const bool bUseTranslucentBlend = Texture->HasAlphaChannel();
+
+		// Fully stream in the texture before drawing it.
+		Texture->SetForceMipLevelsToBeResident(30.0f);
+		Texture->WaitForStreaming();
+
+		FLinearColor TextureDrawColor = FLinearColor::White;
+
+		const float XPos = -ZoomPos.X * ZoomAmount;
+		const float YPos = -ZoomPos.Y * ZoomAmount;
+		const float Width = Texture->GetSurfaceWidth() * ZoomAmount;
+		const float Height = Texture->GetSurfaceHeight() * ZoomAmount;
+
+		Canvas->DrawTile(XPos, YPos, Width, Height, 0.0f, 0.0f, 1.0f, 1.0f, TextureDrawColor, Texture->Resource, bUseTranslucentBlend);
+	}
+
+	// Overlay the selection rectangles
+	DrawSelectionRectangles(Viewport, Canvas);
+
+	if (bHasValidPaintRectangle)
+	{
+		const FViewportSelectionRectangle& Rect = ValidPaintRectangle;
+
+		const float X = (Rect.TopLeft.X - ZoomPos.X) * ZoomAmount;
+		const float Y = (Rect.TopLeft.Y - ZoomPos.Y) * ZoomAmount;
+		const float W = Rect.Dimensions.X * ZoomAmount;
+		const float H = Rect.Dimensions.Y * ZoomAmount;
+
+		FCanvasBoxItem BoxItem(FVector2D(X, Y), FVector2D(W, H));
+		BoxItem.SetColor(Rect.Color);
+		Canvas->DrawItem(BoxItem);
+	}
+
+	if (TileIndex != INDEX_NONE)
+	{
+		const FString TileIndexString = FString::Printf(TEXT("Tile# %d"), TileIndex);
+
+		int32 XL;
+		int32 YL;
+		StringSize(GEngine->GetLargeFont(), XL, YL, *TileIndexString);
+		const float DrawX = 4.0f;
+		const float DrawY = FMath::FloorToFloat(Viewport->GetSizeXY().Y - YL - 4.0f);
+		Canvas->DrawShadowedString(DrawX, DrawY, *TileIndexString, GEngine->GetLargeFont(), FLinearColor::White);
+	}
+}
+
+FLinearColor FTileSetEditorViewportClient::GetBackgroundColor() const
+{
+	if (UPaperTileSet* TileSet = TileSetBeingEdited.Get())
+	{
+		return TileSet->BackgroundColor;
+	}
+	else
+	{
+		return FEditorViewportClient::GetBackgroundColor();
+	}
+}
+
+//////////////////////////////////////////////////////////////////////////
+
+#undef LOCTEXT_NAMESPACE
\ No newline at end of file
diff --git a/Engine/Plugins/2D/Paper2D/Source/Paper2DEditor/Private/TileSetEditor/TileSetEditorViewportClient.h b/Engine/Plugins/2D/Paper2D/Source/Paper2DEditor/Private/TileSetEditor/TileSetEditorViewportClient.h
new file mode 100644
--- /dev/null
+++ b/Engine/Plugins/2D/Paper2D/Source/Paper2DEditor/Private/TileSetEditor/TileSetEditorViewportClient.h
@@ -0,0 +1,30 @@
+// Copyright 1998-2015 Epic Games, Inc. All Rights Reserved.
+
+#pragma once
+
+#include "PaperEditorViewportClient.h"
+
+//////////////////////////////////////////////////////////////////////////
+// FTileSetEditorViewportClient
+
+class FTileSetEditorViewportClient : public FPaperEditorViewportClient
+{
+public:
+	FTileSetEditorViewportClient(UPaperTileSet* InTileSet);
+
+	// FViewportClient interface
+	virtual void Draw(FViewport* Viewport, FCanvas* Canvas) override;
+	// End of FViewportClient interface
+
+	// FEditorViewportClient interface
+	virtual FLinearColor GetBackgroundColor() const override;
+	// End of FEditorViewportClient interface
+
+public:
+	// Tile set
+	TWeakObjectPtr<UPaperTileSet> TileSetBeingEdited;
+
+	bool bHasValidPaintRectangle;
+	FViewportSelectionRectangle ValidPaintRectangle;
+	int32 TileIndex;
+};
diff --git a/Engine/Plugins/2D/Paper2D/Source/Paper2DEditor/Private/TileSetEditor/TileSetSelectorViewport.cpp b/Engine/Plugins/2D/Paper2D/Source/Paper2DEditor/Private/TileSetEditor/TileSetSelectorViewport.cpp
new file mode 100644
--- /dev/null
+++ b/Engine/Plugins/2D/Paper2D/Source/Paper2DEditor/Private/TileSetEditor/TileSetSelectorViewport.cpp
@@ -0,0 +1,140 @@
+// Copyright 1998-2015 Epic Games, Inc. All Rights Reserved.
+
+#include "Paper2DEditorPrivatePCH.h"
+#include "TileMapEditing/EdModeTileMap.h"
+#include "TileSetSelectorViewport.h"
+#include "TileSetEditorViewportClient.h"
+
+#define LOCTEXT_NAMESPACE "TileSetEditor"
+
+//////////////////////////////////////////////////////////////////////////
+// STileSetSelectorViewport
+
+STileSetSelectorViewport::~STileSetSelectorViewport()
+{
+	TypedViewportClient = nullptr;
+}
+
+void STileSetSelectorViewport::Construct(const FArguments& InArgs, UPaperTileSet* InTileSet, FEdModeTileMap* InTileMapEditor)
+{
+	SelectionTopLeft = FIntPoint::ZeroValue;
+	SelectionDimensions = FIntPoint::ZeroValue;
+
+	TileSetPtr = InTileSet;
+	TileMapEditor = InTileMapEditor;
+
+	TypedViewportClient = MakeShareable(new FTileSetEditorViewportClient(InTileSet));
+
+	SPaperEditorViewport::Construct(
+		SPaperEditorViewport::FArguments().OnSelectionChanged(this, &STileSetSelectorViewport::OnSelectionChanged),
+		TypedViewportClient.ToSharedRef());
+
+	// Make sure we get input instead of the viewport stealing it
+	ViewportWidget->SetVisibility(EVisibility::HitTestInvisible);
+}
+
+void STileSetSelectorViewport::ChangeTileSet(UPaperTileSet* InTileSet)
+{
+	if (InTileSet != TileSetPtr.Get())
+	{
+		TileSetPtr = InTileSet;
+		TypedViewportClient->TileSetBeingEdited = InTileSet;
+
+		// Update the selection rectangle
+		RefreshSelectionRectangle();
+		TypedViewportClient->Invalidate();
+	}
+}
+
+FText STileSetSelectorViewport::GetTitleText() const
+{
+	if (UPaperTileSet* TileSet = TileSetPtr.Get())
+	{
+		return FText::FromString(TileSet->GetName());
+	}
+
+	return LOCTEXT("TileSetSelectorTitle", "Tile Set Selector");
+}
+
+
+void STileSetSelectorViewport::OnSelectionChanged(FMarqueeOperation Marquee, bool bIsPreview)
+{
+	UPaperTileSet* TileSetBeingEdited = TileSetPtr.Get();
+	if (TileSetBeingEdited == nullptr)
+	{
+		return;
+	}
+
+	const FVector2D TopLeftUnrounded = Marquee.Rect.GetUpperLeft();
+	const FVector2D BottomRightUnrounded = Marquee.Rect.GetLowerRight();
+	if ((TopLeftUnrounded != FVector2D::ZeroVector) || Marquee.IsValid())
+	{
+		const int32 TileCountX = TileSetBeingEdited->GetTileCountX();
+		const int32 TileCountY = TileSetBeingEdited->GetTileCountY();
+
+		// Round down the top left corner
+		const FIntPoint TileTopLeft = TileSetBeingEdited->GetTileXYFromTextureUV(TopLeftUnrounded, false);
+		SelectionTopLeft = FIntPoint(FMath::Clamp<int32>(TileTopLeft.X, 0, TileCountX), FMath::Clamp<int32>(TileTopLeft.Y, 0, TileCountY));
+
+		// Round up the bottom right corner
+		const FIntPoint TileBottomRight = TileSetBeingEdited->GetTileXYFromTextureUV(BottomRightUnrounded, true);
+		const FIntPoint SelectionBottomRight(FMath::Clamp<int32>(TileBottomRight.X, 0, TileCountX), FMath::Clamp<int32>(TileBottomRight.Y, 0, TileCountY));
+
+		// Compute the new selection dimensions
+		SelectionDimensions = SelectionBottomRight - SelectionTopLeft;
+	}
+	else
+	{
+		SelectionTopLeft.X = 0;
+		SelectionTopLeft.Y = 0;
+		SelectionDimensions.X = 0;
+		SelectionDimensions.Y = 0;
+	}
+
+	OnTileSelectionChanged.Broadcast(SelectionTopLeft, SelectionDimensions);
+
+	const bool bHasSelection = (SelectionDimensions.X * SelectionDimensions.Y) > 0;
+	if (bIsPreview && bHasSelection)
+	{
+		if (TileMapEditor != nullptr)
+		{
+			TileMapEditor->SetActivePaint(TileSetBeingEdited, SelectionTopLeft, SelectionDimensions);
+
+			// Switch to paint brush mode if we were in the eraser mode since the user is trying to select some ink to paint with
+			if (TileMapEditor->GetActiveTool() == ETileMapEditorTool::Eraser)
+			{
+				TileMapEditor->SetActiveTool(ETileMapEditorTool::Paintbrush);
+			}
+		}
+
+		RefreshSelectionRectangle();
+	}
+}
+
+void STileSetSelectorViewport::RefreshSelectionRectangle()
+{
+	if (FTileSetEditorViewportClient* Client = TypedViewportClient.Get())
+	{
+		UPaperTileSet* TileSetBeingEdited = TileSetPtr.Get();
+
+		const bool bHasSelection = (SelectionDimensions.X * SelectionDimensions.Y) > 0;
+		Client->bHasValidPaintRectangle = bHasSelection && (TileSetBeingEdited != nullptr);
+
+		const int32 TileIndex = (bHasSelection && (TileSetBeingEdited != nullptr)) ? (SelectionTopLeft.X + SelectionTopLeft.Y * TileSetBeingEdited->GetTileCountX()) : INDEX_NONE;
+		Client->TileIndex = TileIndex;
+
+		if (bHasSelection && (TileSetBeingEdited != nullptr))
+		{
+			Client->ValidPaintRectangle.Color = FLinearColor::White;
+
+			const FIntPoint TopLeft = TileSetBeingEdited->GetTileUVFromTileXY(SelectionTopLeft);
+			const FIntPoint BottomRight = TileSetBeingEdited->GetTileUVFromTileXY(SelectionTopLeft + SelectionDimensions);
+			Client->ValidPaintRectangle.Dimensions = BottomRight - TopLeft;
+			Client->ValidPaintRectangle.TopLeft = TopLeft;
+		}
+	}
+}
+
+//////////////////////////////////////////////////////////////////////////
+
+#undef LOCTEXT_NAMESPACE
\ No newline at end of file
diff --git a/Engine/Plugins/2D/Paper2D/Source/Paper2DEditor/Private/TileSetEditor/TileSetSelectorViewport.h b/Engine/Plugins/2D/Paper2D/Source/Paper2DEditor/Private/TileSetEditor/TileSetSelectorViewport.h
new file mode 100644
--- /dev/null
+++ b/Engine/Plugins/2D/Paper2D/Source/Paper2DEditor/Private/TileSetEditor/TileSetSelectorViewport.h
@@ -0,0 +1,48 @@
+// Copyright 1998-2015 Epic Games, Inc. All Rights Reserved.
+
+#pragma once
+
+#include "Toolkits/AssetEditorToolkit.h"
+#include "Toolkits/AssetEditorManager.h"
+#include "SPaperEditorViewport.h"
+
+//////////////////////////////////////////////////////////////////////////
+// STileSetSelectorViewport
+
+class STileSetSelectorViewport : public SPaperEditorViewport
+{
+public:
+	DECLARE_MULTICAST_DELEGATE_TwoParams(FOnTileViewportSelectionChanged, const FIntPoint& /*TopLeft*/, const FIntPoint& /*Dimensions*/);
+
+	SLATE_BEGIN_ARGS(STileSetSelectorViewport) {}
+	SLATE_END_ARGS()
+
+	~STileSetSelectorViewport();
+
+	void Construct(const FArguments& InArgs, UPaperTileSet* InTileSet, class FEdModeTileMap* InTileMapEditor);
+
+	void ChangeTileSet(UPaperTileSet* InTileSet);
+
+	FOnTileViewportSelectionChanged& GetTileSelectionChanged()
+	{
+		return OnTileSelectionChanged;
+	}
+
+protected:
+	// SPaperEditorViewport interface
+	virtual FText GetTitleText() const override;
+	// End of SPaperEditorViewport interface
+
+private:
+	void OnSelectionChanged(FMarqueeOperation Marquee, bool bIsPreview);
+	void RefreshSelectionRectangle();
+
+private:
+	TWeakObjectPtr<class UPaperTileSet> TileSetPtr;
+	TSharedPtr<class FTileSetEditorViewportClient> TypedViewportClient;
+	class FEdModeTileMap* TileMapEditor;
+	FIntPoint SelectionTopLeft;
+	FIntPoint SelectionDimensions;
+
+	FOnTileViewportSelectionChanged OnTileSelectionChanged;
+};
\ No newline at end of file
