Paper2D: Tile Maps are officially out of Experimental and are considered an Early Access Preview [UE-5209]

diff --git a/Engine/Plugins/2D/Paper2D/Source/Paper2D/Classes/PaperTileMapComponent.h b/Engine/Plugins/2D/Paper2D/Source/Paper2D/Classes/PaperTileMapComponent.h
--- a/Engine/Plugins/2D/Paper2D/Source/Paper2D/Classes/PaperTileMapComponent.h
+++ b/Engine/Plugins/2D/Paper2D/Source/Paper2D/Classes/PaperTileMapComponent.h
@@ -12,12 +12,18 @@
  * This component is created when you drag a tile map asset from the content browser into a Blueprint, or
  * contained inside of the actor created when you drag one into the level.
  *
- * NOTE: This is an experimental class and certain functionality is not properly supported yet (such as collision).  Use at your own risk!
+ * NOTE: This is an early access preview class.  While not considered production-ready, it is a step beyond
+ * 'experimental' and is being provided as a preview of things to come:
+ *  - We will try to provide forward-compatibility for content you create.
+ *  - The classes may change significantly in the future.
+ *  - The code is in an early state and may not meet the desired polish / quality bar.
+ *  - There is probably no documentation or example content yet.
+ *  - They will be promoted out of 'Early Access' when they are production ready.
  *
  * @see UPrimitiveComponent, UPaperTileMap
  */
 
-UCLASS(hideCategories=Object, ClassGroup=Paper2D, Experimental, meta=(BlueprintSpawnableComponent))
+UCLASS(hideCategories=Object, ClassGroup=Paper2D, EarlyAccessPreview, meta=(BlueprintSpawnableComponent))
 class PAPER2D_API UPaperTileMapComponent : public UMeshComponent
 {
 	GENERATED_UCLASS_BODY()
diff --git a/Engine/Plugins/2D/Paper2D/Source/Paper2DEditor/Private/TileMapEditing/TileMapEditor.cpp b/Engine/Plugins/2D/Paper2D/Source/Paper2DEditor/Private/TileMapEditing/TileMapEditor.cpp
--- a/Engine/Plugins/2D/Paper2D/Source/Paper2DEditor/Private/TileMapEditing/TileMapEditor.cpp
+++ b/Engine/Plugins/2D/Paper2D/Source/Paper2DEditor/Private/TileMapEditing/TileMapEditor.cpp
@@ -221,7 +221,7 @@ TSharedRef<SDockTab> FTileMapEditor::SpawnTab_Viewport(const FSpawnTabArgs& Args
 					SNew(STextBlock)
 					.Visibility(EVisibility::HitTestInvisible)
 					.TextStyle(FEditorStyle::Get(), "Graph.CornerText")
-					.Text(LOCTEXT("TileMapEditorViewportExperimentalWarning", "Experimental!"))
+					.Text(LOCTEXT("TileMapEditorViewportEarlyAccessPreviewWarning", "Early access preview"))
 				]
 		];
 }
