Paper2D: Renamed PaperJsonImporter module and associated classes to PaperSpriteSheetImporter

diff --git a/Engine/Plugins/2D/Paper2D/Paper2D.uplugin b/Engine/Plugins/2D/Paper2D/Paper2D.uplugin
--- a/Engine/Plugins/2D/Paper2D/Paper2D.uplugin
+++ b/Engine/Plugins/2D/Paper2D/Paper2D.uplugin
@@ -23,7 +23,7 @@
 			"Type" : "Editor"
 		},
 		{
-			"Name" : "PaperJsonImporter",
+			"Name" : "PaperSpriteSheetImporter",
 			"Type" : "Editor"
 		},
 		{
diff --git a/Engine/Plugins/2D/Paper2D/Source/PaperJsonImporter/Classes/PaperJsonImporterFactory.h b/Engine/Plugins/2D/Paper2D/Source/PaperJsonImporter/Classes/PaperJsonImporterFactory.h
deleted file mode 100644
diff --git a/Engine/Plugins/2D/Paper2D/Source/PaperJsonImporter/PaperJsonImporter.Build.cs b/Engine/Plugins/2D/Paper2D/Source/PaperJsonImporter/PaperJsonImporter.Build.cs
deleted file mode 100644
diff --git a/Engine/Plugins/2D/Paper2D/Source/PaperJsonImporter/Private/PaperJsonImporterFactory.cpp b/Engine/Plugins/2D/Paper2D/Source/PaperJsonImporter/Private/PaperJsonImporterFactory.cpp
deleted file mode 100644
diff --git a/Engine/Plugins/2D/Paper2D/Source/PaperJsonImporter/Private/PaperJsonImporterModule.cpp b/Engine/Plugins/2D/Paper2D/Source/PaperJsonImporter/Private/PaperJsonImporterModule.cpp
deleted file mode 100644
diff --git a/Engine/Plugins/2D/Paper2D/Source/PaperJsonImporter/Private/PaperJsonImporterPrivatePCH.h b/Engine/Plugins/2D/Paper2D/Source/PaperJsonImporter/Private/PaperJsonImporterPrivatePCH.h
deleted file mode 100644
diff --git a/Engine/Plugins/2D/Paper2D/Source/PaperJsonImporter/Private/PaperJsonSpriteSheetImporter.cpp b/Engine/Plugins/2D/Paper2D/Source/PaperJsonImporter/Private/PaperJsonSpriteSheetImporter.cpp
deleted file mode 100644
diff --git a/Engine/Plugins/2D/Paper2D/Source/PaperJsonImporter/Private/PaperJsonSpriteSheetImporter.h b/Engine/Plugins/2D/Paper2D/Source/PaperJsonImporter/Private/PaperJsonSpriteSheetImporter.h
deleted file mode 100644
diff --git a/Engine/Plugins/2D/Paper2D/Source/PaperJsonImporter/Private/PaperSpriteSheet.cpp b/Engine/Plugins/2D/Paper2D/Source/PaperJsonImporter/Private/PaperSpriteSheet.cpp
deleted file mode 100644
diff --git a/Engine/Plugins/2D/Paper2D/Source/PaperJsonImporter/Private/PaperSpriteSheet.h b/Engine/Plugins/2D/Paper2D/Source/PaperJsonImporter/Private/PaperSpriteSheet.h
deleted file mode 100644
diff --git a/Engine/Plugins/2D/Paper2D/Source/PaperJsonImporter/Private/PaperSpriteSheetAssetTypeActions.cpp b/Engine/Plugins/2D/Paper2D/Source/PaperJsonImporter/Private/PaperSpriteSheetAssetTypeActions.cpp
deleted file mode 100644
diff --git a/Engine/Plugins/2D/Paper2D/Source/PaperJsonImporter/Private/PaperSpriteSheetAssetTypeActions.h b/Engine/Plugins/2D/Paper2D/Source/PaperJsonImporter/Private/PaperSpriteSheetAssetTypeActions.h
deleted file mode 100644
diff --git a/Engine/Plugins/2D/Paper2D/Source/PaperJsonImporter/Private/ReimportPaperJsonImporterFactory.cpp b/Engine/Plugins/2D/Paper2D/Source/PaperJsonImporter/Private/ReimportPaperJsonImporterFactory.cpp
deleted file mode 100644
diff --git a/Engine/Plugins/2D/Paper2D/Source/PaperJsonImporter/Private/ReimportPaperJsonImporterFactory.h b/Engine/Plugins/2D/Paper2D/Source/PaperJsonImporter/Private/ReimportPaperJsonImporterFactory.h
deleted file mode 100644
diff --git a/Engine/Plugins/2D/Paper2D/Source/PaperSpriteSheetImporter/Classes/PaperSpriteSheetImportFactory.h b/Engine/Plugins/2D/Paper2D/Source/PaperSpriteSheetImporter/Classes/PaperSpriteSheetImportFactory.h
new file mode 100644
--- /dev/null
+++ b/Engine/Plugins/2D/Paper2D/Source/PaperSpriteSheetImporter/Classes/PaperSpriteSheetImportFactory.h
@@ -0,0 +1,18 @@
+// Copyright 1998-2015 Epic Games, Inc. All Rights Reserved.
+
+#pragma once
+
+#include "PaperSpriteSheetImportFactory.generated.h"
+
+// Imports a sprite sheet (and associated paper sprites and textures) from a JSON file exported from Adobe Flash CS6, Texture Packer, or other compatible tool
+UCLASS()
+class UPaperSpriteSheetImportFactory : public UFactory
+{
+	GENERATED_UCLASS_BODY()
+
+	// UFactory interface
+	virtual FText GetToolTip() const override;
+	virtual UObject* FactoryCreateText(UClass* InClass, UObject* InParent, FName InName, EObjectFlags Flags, UObject* Context, const TCHAR* Type, const TCHAR*& Buffer, const TCHAR* BufferEnd, FFeedbackContext* Warn) override;
+	// End of UFactory interface
+};
+
diff --git a/Engine/Plugins/2D/Paper2D/Source/PaperSpriteSheetImporter/PaperSpriteSheetImporter.Build.cs b/Engine/Plugins/2D/Paper2D/Source/PaperSpriteSheetImporter/PaperSpriteSheetImporter.Build.cs
new file mode 100644
--- /dev/null
+++ b/Engine/Plugins/2D/Paper2D/Source/PaperSpriteSheetImporter/PaperSpriteSheetImporter.Build.cs
@@ -0,0 +1,37 @@
+// Copyright 1998-2015 Epic Games, Inc. All Rights Reserved.
+
+using UnrealBuildTool;
+
+public class PaperSpriteSheetImporter : ModuleRules
+{
+	public PaperSpriteSheetImporter(TargetInfo Target)
+	{
+		PrivateDependencyModuleNames.AddRange(
+			new string[] {
+				"Core",
+				"CoreUObject",
+				"Json",
+				"Slate",
+				"SlateCore",
+				"Engine",
+				"Paper2D",
+				"UnrealEd",
+				"Paper2DEditor",
+				"AssetTools",
+				"ContentBrowser",
+                "EditorStyle"
+			});
+
+		PrivateIncludePathModuleNames.AddRange(
+			new string[] {
+				"AssetTools",
+				"AssetRegistry"
+			});
+
+		DynamicallyLoadedModuleNames.AddRange(
+			new string[] {
+				"AssetTools",
+				"AssetRegistry"
+			});
+	}
+}
diff --git a/Engine/Plugins/2D/Paper2D/Source/PaperSpriteSheetImporter/Private/PaperJsonSpriteSheetImporter.cpp b/Engine/Plugins/2D/Paper2D/Source/PaperSpriteSheetImporter/Private/PaperJsonSpriteSheetImporter.cpp
new file mode 100644
--- /dev/null
+++ b/Engine/Plugins/2D/Paper2D/Source/PaperSpriteSheetImporter/Private/PaperJsonSpriteSheetImporter.cpp
@@ -0,0 +1,563 @@
+// Copyright 1998-2015 Epic Games, Inc. All Rights Reserved.
+
+#include "PaperSpriteSheetImporterPrivatePCH.h"
+#include "Paper2DClasses.h"
+#include "Json.h"
+#include "PaperJSONHelpers.h"
+#include "PaperJsonSpriteSheetImporter.h"
+#include "AssetToolsModule.h"
+#include "AssetRegistryModule.h"
+#include "PackageTools.h"
+#include "ObjectTools.h"
+#include "PaperSpriteSheet.h"
+
+TSharedPtr<FJsonObject> ParseJSON(const FString& FileContents, const FString& NameForErrors)
+{
+	// Load the file up (JSON format)
+	if (!FileContents.IsEmpty())
+	{
+		const TSharedRef< TJsonReader<> >& Reader = TJsonReaderFactory<>::Create(FileContents);
+
+		TSharedPtr<FJsonObject> SpriteDescriptorObject;
+		if (FJsonSerializer::Deserialize(Reader, SpriteDescriptorObject) && SpriteDescriptorObject.IsValid())
+		{
+			// File was loaded and deserialized OK!
+			return SpriteDescriptorObject;
+		}
+		else
+		{
+			UE_LOG(LogPaperSpriteSheetImporter, Warning, TEXT("Failed to parse sprite descriptor file '%s'.  Error: '%s'"), *NameForErrors, *Reader->GetErrorMessage());
+			return nullptr;
+		}
+	}
+	else
+	{
+		UE_LOG(LogPaperSpriteSheetImporter, Warning, TEXT("Sprite descriptor file '%s' was empty.  This sprite cannot be imported."), *NameForErrors);
+		return nullptr;
+	}
+}
+
+TSharedPtr<FJsonObject> ParseJSON(FArchive* const Stream, const FString& NameForErrors)
+{
+	const TSharedRef< TJsonReader<> >& Reader = TJsonReaderFactory<>::Create(Stream);
+
+	TSharedPtr<FJsonObject> SpriteDescriptorObject;
+	if (FJsonSerializer::Deserialize(Reader, SpriteDescriptorObject) && SpriteDescriptorObject.IsValid())
+	{
+		// File was loaded and deserialized OK!
+		return SpriteDescriptorObject;
+	}
+	else
+	{
+		UE_LOG(LogPaperSpriteSheetImporter, Warning, TEXT("Failed to parse sprite descriptor file '%s'.  Error: '%s'"), *NameForErrors, *Reader->GetErrorMessage());
+		return nullptr;
+	}
+}
+
+bool ParseMetaBlock(const FString& NameForErrors, TSharedPtr<FJsonObject>& SpriteDescriptorObject, FString& OutImage)
+{
+	bool bLoadedSuccessfully = true;
+	TSharedPtr<FJsonObject> MetaBlock = FPaperJSONHelpers::ReadObject(SpriteDescriptorObject, TEXT("meta"));
+	if (MetaBlock.IsValid())
+	{
+		// Example contents:
+		//   "app": "Adobe Flash CS6",
+		//   "version": "12.0.0.481",        (ignored)
+		//   "image": "MySprite.png",
+		//   "format": "RGBA8888",           (ignored)
+		//   "size": {"w":2048,"h":2048},    (ignored)
+		//   "scale": "1"                    (ignored)
+
+		const FString AppName = FPaperJSONHelpers::ReadString(MetaBlock, TEXT("app"), TEXT(""));
+		OutImage = FPaperJSONHelpers::ReadString(MetaBlock, TEXT("image"), TEXT(""));
+
+		const FString FlashPrefix(TEXT("Adobe Flash"));
+		const FString TexturePackerPrefix(TEXT("http://www.codeandweb.com/texturepacker"));
+
+		if (AppName.StartsWith(FlashPrefix) || AppName.StartsWith(TexturePackerPrefix))
+		{
+			// Cool, we (mostly) know how to handle these sorts of files!
+			UE_LOG(LogPaperSpriteSheetImporter, Log, TEXT("Parsing sprite sheet exported from '%s'"), *AppName);
+		}
+		else if (!AppName.IsEmpty())
+		{
+			// It's got an app tag inside a meta block, so we'll take a crack at it
+			UE_LOG(LogPaperSpriteSheetImporter, Warning, TEXT("Unexpected 'app' named '%s' while parsing sprite descriptor file '%s'.  Parsing will continue but the format may not be fully supported"), *AppName, *NameForErrors);
+		}
+		else
+		{
+			// Probably not a sprite sheet
+			UE_LOG(LogPaperSpriteSheetImporter, Warning, TEXT("Failed to parse sprite descriptor file '%s'.  Expected 'app' key indicating the exporter (might not be a sprite sheet)"), *NameForErrors);
+			bLoadedSuccessfully = false;
+		}
+
+		if (OutImage.IsEmpty())
+		{
+			UE_LOG(LogPaperSpriteSheetImporter, Warning, TEXT("Failed to parse sprite descriptor file '%s'.  Expected valid 'image' tag"), *NameForErrors);
+			bLoadedSuccessfully = false;
+		}
+	}
+	else
+	{
+		UE_LOG(LogPaperSpriteSheetImporter, Warning, TEXT("Failed to parse sprite descriptor file '%s'.  Missing meta block"), *NameForErrors);
+		bLoadedSuccessfully = false;
+	}
+	return bLoadedSuccessfully;
+}
+
+static bool ParseFrame(TSharedPtr<FJsonObject> &FrameData, FSpriteFrame &OutFrame)
+{
+	bool bReadFrameSuccessfully = true;
+	// An example frame:
+	//   "frame": {"x":210,"y":10,"w":190,"h":223},
+	//   "rotated": false,
+	//   "trimmed": true,
+	//   "spriteSourceSize": {"x":0,"y":11,"w":216,"h":240},
+	//   "sourceSize": {"w":216,"h":240},	
+	//   "pivot": {"x":0.5,"y":0.5}			[optional]
+
+	OutFrame.bRotated = FPaperJSONHelpers::ReadBoolean(FrameData, TEXT("rotated"), false);
+	OutFrame.bTrimmed = FPaperJSONHelpers::ReadBoolean(FrameData, TEXT("trimmed"), false);
+
+	bReadFrameSuccessfully = bReadFrameSuccessfully && FPaperJSONHelpers::ReadRectangle(FrameData, "frame", /*out*/ OutFrame.SpritePosInSheet, /*out*/ OutFrame.SpriteSizeInSheet);
+
+	if (OutFrame.bTrimmed)
+	{
+		bReadFrameSuccessfully = bReadFrameSuccessfully && FPaperJSONHelpers::ReadSize(FrameData, "sourceSize", /*out*/ OutFrame.ImageSourceSize);
+		bReadFrameSuccessfully = bReadFrameSuccessfully && FPaperJSONHelpers::ReadRectangle(FrameData, "spriteSourceSize", /*out*/ OutFrame.SpriteSourcePos, /*out*/ OutFrame.SpriteSourceSize);
+	}
+	else
+	{
+		OutFrame.SpriteSourcePos = FVector2D::ZeroVector;
+		OutFrame.SpriteSourceSize = OutFrame.SpriteSizeInSheet;
+		OutFrame.ImageSourceSize = OutFrame.SpriteSizeInSheet;
+	}
+
+	if (!FPaperJSONHelpers::ReadXY(FrameData, "pivot", /*out*/ OutFrame.Pivot))
+	{
+		OutFrame.Pivot.X = 0.5f;
+		OutFrame.Pivot.Y = 0.5f;
+	}
+
+	// A few more prerequisites to sort out before rotation can be enabled
+	if (OutFrame.bRotated)
+	{
+		// Sprite Source Pos is from top left, our pivot when rotated is bottom left
+		OutFrame.SpriteSourcePos.Y = OutFrame.ImageSourceSize.Y - OutFrame.SpriteSourcePos.Y - OutFrame.SpriteSizeInSheet.Y;
+
+		// We rotate final sprite geometry 90 deg CCW to fixup
+		// These need to be swapped to be valid in texture space.
+		Swap(OutFrame.SpriteSizeInSheet.X, OutFrame.SpriteSizeInSheet.Y);
+		Swap(OutFrame.ImageSourceSize.X, OutFrame.ImageSourceSize.Y);
+
+		Swap(OutFrame.SpriteSourcePos.X, OutFrame.SpriteSourcePos.Y);
+		Swap(OutFrame.SpriteSourceSize.X, OutFrame.SpriteSourceSize.Y);
+	}
+
+	return bReadFrameSuccessfully;
+}
+
+static bool ParseFramesFromSpriteHash(TSharedPtr<FJsonObject> ObjectBlock, TArray<FSpriteFrame>& OutSpriteFrames, TSet<FName>& FrameNames)
+{
+	GWarn->BeginSlowTask(NSLOCTEXT("Paper2D", "PaperJsonImporterFactory_ParsingSprites", "Parsing Sprite Frame"), true, true);
+	bool bLoadedSuccessfully = true;
+
+	// Parse all of the frames
+	int32 FrameCount = 0;
+	for (auto FrameIt = ObjectBlock->Values.CreateIterator(); FrameIt; ++FrameIt)
+	{
+		GWarn->StatusUpdate(FrameCount, ObjectBlock->Values.Num(), NSLOCTEXT("Paper2D", "PaperJsonImporterFactory_ParsingSprites", "Parsing Sprite Frames"));
+
+		bool bReadFrameSuccessfully = true;
+
+		FSpriteFrame Frame;
+		Frame.FrameName = *FrameIt.Key();
+
+		if (FrameNames.Contains(Frame.FrameName))
+		{
+			bReadFrameSuccessfully = false;
+		}
+		else
+		{
+			FrameNames.Add(Frame.FrameName);
+		}
+
+		TSharedPtr<FJsonValue> FrameDataAsValue = FrameIt.Value();
+		TSharedPtr<FJsonObject> FrameData;
+		if (FrameDataAsValue->Type == EJson::Object)
+		{
+			FrameData = FrameDataAsValue->AsObject();
+			bReadFrameSuccessfully = bReadFrameSuccessfully && ParseFrame(FrameData, /*out*/Frame);
+		}
+		else
+		{
+			bReadFrameSuccessfully = false;
+		}
+
+		if (bReadFrameSuccessfully)
+		{
+			OutSpriteFrames.Add(Frame);
+		}
+		else
+		{
+			UE_LOG(LogPaperSpriteSheetImporter, Warning, TEXT("Frame %s is in an unexpected format"), *Frame.FrameName.ToString());
+			bLoadedSuccessfully = false;
+		}
+
+		FrameCount++;
+	}
+
+	GWarn->EndSlowTask();
+	return bLoadedSuccessfully;
+}
+
+static bool ParseFramesFromSpriteArray(const TArray<TSharedPtr<FJsonValue>>& ArrayBlock, TArray<FSpriteFrame>& OutSpriteFrames, TSet<FName>& FrameNames)
+{
+	GWarn->BeginSlowTask(NSLOCTEXT("Paper2D", "PaperJsonImporterFactory_ParsingSprites", "Parsing Sprite Frame"), true, true);
+	bool bLoadedSuccessfully = true;
+
+	// Parse all of the frames
+	for (int32 FrameCount = 0; FrameCount < ArrayBlock.Num(); ++FrameCount)
+	{
+		GWarn->StatusUpdate(FrameCount, ArrayBlock.Num(), NSLOCTEXT("Paper2D", "PaperJsonImporterFactory_ParsingSprites", "Parsing Sprite Frames"));
+		bool bReadFrameSuccessfully = true;
+		FSpriteFrame Frame;
+
+		const TSharedPtr<FJsonValue>& FrameDataAsValue = ArrayBlock[FrameCount];
+		if (FrameDataAsValue->Type == EJson::Object)
+		{
+			TSharedPtr<FJsonObject> FrameData;
+			FrameData = FrameDataAsValue->AsObject();
+
+			FString FrameFilename = FPaperJSONHelpers::ReadString(FrameData, TEXT("filename"), TEXT(""));
+			if (!FrameFilename.IsEmpty())
+			{
+				Frame.FrameName = FName(*FrameFilename); // case insensitive
+				if (FrameNames.Contains(Frame.FrameName))
+				{
+					bReadFrameSuccessfully = false;
+				}
+				else
+				{
+					FrameNames.Add(Frame.FrameName);
+				}
+
+				bReadFrameSuccessfully = bReadFrameSuccessfully && ParseFrame(FrameData, /*out*/Frame);
+			}
+			else
+			{
+				bReadFrameSuccessfully = false;
+			}
+		}
+		else
+		{
+			bReadFrameSuccessfully = false;
+		}
+
+		if (bReadFrameSuccessfully)
+		{
+			OutSpriteFrames.Add(Frame);
+		}
+		else
+		{
+			UE_LOG(LogPaperSpriteSheetImporter, Warning, TEXT("Frame %s is in an unexpected format"), *Frame.FrameName.ToString());
+			bLoadedSuccessfully = false;
+		}
+	}
+
+	GWarn->EndSlowTask();
+	return bLoadedSuccessfully;
+}
+
+static ESpritePivotMode::Type GetBestPivotType(FVector2D JsonPivot)
+{
+	// Not assuming layout of ESpritePivotMode
+	if (JsonPivot.X == 0 && JsonPivot.Y == 0)
+	{
+		return ESpritePivotMode::Top_Left;
+	}
+	else if (JsonPivot.X == 0.5f && JsonPivot.Y == 0)
+	{
+		return ESpritePivotMode::Top_Center;
+	}
+	else if (JsonPivot.X == 1.0f && JsonPivot.Y == 0)
+	{
+		return ESpritePivotMode::Top_Right;
+	}
+	else if (JsonPivot.X == 0 && JsonPivot.Y == 0.5f)
+	{
+		return ESpritePivotMode::Center_Left;
+	}
+	else if (JsonPivot.X == 0.5f && JsonPivot.Y == 0.5f)
+	{
+		return ESpritePivotMode::Center_Center;
+	}
+	else if (JsonPivot.X == 1.0f && JsonPivot.Y == 0.5f)
+	{
+		return ESpritePivotMode::Center_Right;
+	}
+	else if (JsonPivot.X == 0 && JsonPivot.Y == 1.0f)
+	{
+		return ESpritePivotMode::Bottom_Left;
+	}
+	else if (JsonPivot.X == 0.5f && JsonPivot.Y == 1.0f)
+	{
+		return ESpritePivotMode::Bottom_Center;
+	}
+	else if (JsonPivot.X == 1.0f && JsonPivot.Y == 1.0f)
+	{
+		return ESpritePivotMode::Bottom_Right;
+	}
+	else
+	{
+		return ESpritePivotMode::Custom;
+	}
+}
+
+// FPaperJsonSpriteSheetImporter
+////////////////////////////////////////////
+
+FPaperJsonSpriteSheetImporter::FPaperJsonSpriteSheetImporter() : bIsReimporting(false)
+{
+}
+
+void FPaperJsonSpriteSheetImporter::SetReimportData(const FString& InExistingTextureName, UTexture2D* InExistingTexture, const TArray<FString>& ExistingSpriteNames, const TArray< TAssetPtr<class UPaperSprite> >& ExistingSpriteAssetPtrs)
+{
+	ExistingTextureName = InExistingTextureName;
+	ExistingTexture = InExistingTexture;
+
+	check(ExistingSpriteNames.Num() == ExistingSpriteAssetPtrs.Num());
+	if (ExistingSpriteNames.Num() == ExistingSpriteAssetPtrs.Num())
+	{
+		for (int i = 0; i < ExistingSpriteAssetPtrs.Num(); ++i)
+		{
+			const TAssetPtr<class UPaperSprite> SpriteAssetPtr = ExistingSpriteAssetPtrs[i];
+			FStringAssetReference SpriteStringRef = SpriteAssetPtr.ToStringReference();
+			if (!SpriteStringRef.ToString().IsEmpty())
+			{
+				UPaperSprite* LoadedSprite = Cast<UPaperSprite>(StaticLoadObject(UPaperSprite::StaticClass(), nullptr, *SpriteStringRef.ToString(), nullptr, LOAD_None, nullptr));
+				if (LoadedSprite != nullptr)
+				{
+					ExistingSprites.Add(ExistingSpriteNames[i], LoadedSprite);
+				}
+			}
+		}
+	}
+	bIsReimporting = true;
+}
+
+bool FPaperJsonSpriteSheetImporter::Import(TSharedPtr<FJsonObject> SpriteDescriptorObject, const FString& NameForErrors)
+{
+	bool bLoadedSuccessfully = ParseMetaBlock(NameForErrors, SpriteDescriptorObject, Image);
+	if (bLoadedSuccessfully)
+	{
+		TSharedPtr<FJsonObject> ObjectFrameBlock = FPaperJSONHelpers::ReadObject(SpriteDescriptorObject, TEXT("frames"));
+		TSet<FName> FrameNames;
+
+		if (ObjectFrameBlock.IsValid())
+		{
+			bLoadedSuccessfully = ParseFramesFromSpriteHash(ObjectFrameBlock, /*out*/Frames, /*inout*/FrameNames);
+		}
+		else
+		{
+			// Try loading as an array
+			TArray<TSharedPtr<FJsonValue>> ArrayBlock = FPaperJSONHelpers::ReadArray(SpriteDescriptorObject, TEXT("frames"));
+			if (ArrayBlock.Num() > 0)
+			{
+				bLoadedSuccessfully = ParseFramesFromSpriteArray(ArrayBlock, /*out*/Frames, /*inout*/FrameNames);
+			}
+			else
+			{
+				UE_LOG(LogPaperSpriteSheetImporter, Warning, TEXT("Failed to parse sprite descriptor file '%s'.  Missing frames block"), *NameForErrors);
+				bLoadedSuccessfully = false;
+			}
+		}
+
+		if (bLoadedSuccessfully && Frames.Num() == 0)
+		{
+			UE_LOG(LogPaperSpriteSheetImporter, Warning, TEXT("Failed to parse sprite descriptor file '%s'.  No frames loaded"), *NameForErrors);
+			bLoadedSuccessfully = false;
+		}
+	}
+	return bLoadedSuccessfully;
+}
+
+bool FPaperJsonSpriteSheetImporter::ImportFromString(const FString& FileContents, const FString& NameForErrors)
+{
+	TSharedPtr<FJsonObject> SpriteDescriptorObject = ParseJSON(FileContents, NameForErrors);
+	return SpriteDescriptorObject.IsValid() &&
+		Import(SpriteDescriptorObject, NameForErrors);
+}
+
+bool FPaperJsonSpriteSheetImporter::ImportFromArchive(FArchive* Archive, const FString& NameForErrors)
+{
+	TSharedPtr<FJsonObject> SpriteDescriptorObject = ParseJSON(Archive, NameForErrors);
+	return SpriteDescriptorObject.IsValid() &&
+		Import(SpriteDescriptorObject, NameForErrors);
+}
+
+
+bool FPaperJsonSpriteSheetImporter::ImportTextures(const FString& LongPackagePath, const FString& SourcePath)
+{
+	bool bLoadedSuccessfully = true;
+	bool bFoundExistingTexture = false;
+	const FString TargetSubPath = LongPackagePath + TEXT("/Textures");
+
+	const FString SourceSheetImageFilename = FPaths::Combine(*SourcePath, *Image);
+	TArray<FString> SheetFileNames;
+	SheetFileNames.Add(SourceSheetImageFilename);
+
+	if (bIsReimporting && (ExistingTextureName == Image) && (ExistingTexture != nullptr))
+	{
+		ImageTexture = ExistingTexture;
+		bFoundExistingTexture = true;
+	}
+
+	// Import the texture if we need to
+	bool bNeedsImport = !bIsReimporting || (bIsReimporting && !bFoundExistingTexture);
+
+	if (bIsReimporting && bFoundExistingTexture)
+	{
+		if (!FReimportManager::Instance()->Reimport(ImageTexture, /*bAskForNewFileIfMissing=*/ true))
+		{
+			bNeedsImport = true;
+		}
+	}
+
+	if (bNeedsImport)
+	{
+		FAssetToolsModule& AssetToolsModule = FModuleManager::GetModuleChecked<FAssetToolsModule>("AssetTools");
+
+		UReimportTextureFactory* ReimportFactory = nullptr;
+
+		TArray<UObject*> ImportedSheets = AssetToolsModule.Get().ImportAssets(SheetFileNames, TargetSubPath, ReimportFactory);
+
+		UTexture2D* ImportedTexture = (ImportedSheets.Num() > 0) ? Cast<UTexture2D>(ImportedSheets[0]) : nullptr;
+
+		if (ImportedTexture == nullptr)
+		{
+			UE_LOG(LogPaperSpriteSheetImporter, Warning, TEXT("Failed to import sprite sheet image '%s'."), *SourceSheetImageFilename);
+			bLoadedSuccessfully = false;
+		}
+		else
+		{
+			// Change the compression settings
+			//@TODO: Should we always be doing this (particularly the TF_Nearest seems potentially undesirable; maybe make it a property of the sprite sheet asset?)
+			ImportedTexture->Modify();
+			ImportedTexture->LODGroup = TEXTUREGROUP_UI;
+			ImportedTexture->CompressionSettings = TC_EditorIcon;
+			ImportedTexture->Filter = TF_Nearest;
+			ImportedTexture->PostEditChange();
+		}
+
+		ImageTexture = ImportedTexture;
+	}
+
+	return bLoadedSuccessfully;
+}
+
+UPaperSprite* FPaperJsonSpriteSheetImporter::FindExistingSprite(const FString& Name)
+{
+	UPaperSprite**ExistingSprite = ExistingSprites.Find(Name);
+	if (ExistingSprite != nullptr)
+	{
+		return *ExistingSprite;
+	}
+	else
+	{
+		return nullptr;
+	}
+}
+
+bool FPaperJsonSpriteSheetImporter::PerformImport(const FString& LongPackagePath, EObjectFlags Flags, UPaperSpriteSheet* SpriteSheet)
+{
+	FAssetToolsModule& AssetToolsModule = FModuleManager::GetModuleChecked<FAssetToolsModule>("AssetTools");
+
+	GWarn->BeginSlowTask(NSLOCTEXT("Paper2D", "PaperJsonImporterFactory_ImportingSprites", "Importing Sprite Frame"), true, true);
+	for (int32 FrameIndex = 0; FrameIndex < Frames.Num(); ++FrameIndex)
+	{
+		GWarn->StatusUpdate(FrameIndex, Frames.Num(), NSLOCTEXT("Paper2D", "PaperJsonImporterFactory_ImportingSprites", "Importing Sprite Frames"));
+
+		// Check for user canceling the import
+		if (GWarn->ReceivedUserCancel())
+		{
+			break;
+		}
+
+		const FSpriteFrame& Frame = Frames[FrameIndex];
+
+		// Create a package for the frame
+		const FString TargetSubPath = LongPackagePath + TEXT("/Frames");
+
+		UObject* OuterForFrame = nullptr; // @TODO: Use this if we don't want them to be individual assets - Flipbook;
+
+		// Create a frame in the package
+		UPaperSprite* TargetSprite = nullptr;
+		
+		if (bIsReimporting)
+		{
+			TargetSprite = FindExistingSprite(Frame.FrameName.ToString());
+		}
+		
+		if (TargetSprite == nullptr)
+		{
+			const FString SanitizedFrameName = ObjectTools::SanitizeObjectName(Frame.FrameName.ToString());
+			const FString TentativePackagePath = PackageTools::SanitizePackageName(TargetSubPath + TEXT("/") + SanitizedFrameName);
+			FString DefaultSuffix;
+			FString AssetName;
+			FString PackageName;
+			AssetToolsModule.Get().CreateUniqueAssetName(TentativePackagePath, DefaultSuffix, /*out*/ PackageName, /*out*/ AssetName);
+
+			// Create a unique package name and asset name for the frame
+			if (OuterForFrame == nullptr)
+			{
+				// Create a package for the frame
+				OuterForFrame = CreatePackage(nullptr, *PackageName);
+			}
+
+			// Create the asset
+			TargetSprite = NewObject<UPaperSprite>(OuterForFrame, *AssetName, Flags);
+			FAssetRegistryModule::AssetCreated(TargetSprite);
+		}
+
+		TargetSprite->Modify();
+		FSpriteAssetInitParameters SpriteInitParams;
+		SpriteInitParams.Texture = ImageTexture;
+		SpriteInitParams.Offset = Frame.SpritePosInSheet;
+		SpriteInitParams.Dimension = Frame.SpriteSizeInSheet;
+		SpriteInitParams.bNewlyCreated = true;
+		TargetSprite->InitializeSprite(SpriteInitParams);
+
+		if (Frame.bRotated)
+		{
+			TargetSprite->SetRotated(true);
+		}
+
+		if (Frame.bTrimmed)
+		{
+			TargetSprite->SetTrim(Frame.bTrimmed, Frame.SpriteSourcePos, Frame.ImageSourceSize);
+		}
+
+		// Set up pivot on object based on Texture Packer json
+		ESpritePivotMode::Type PivotType = GetBestPivotType(Frame.Pivot);
+		FVector2D TextureSpacePivotPoint = FVector2D::ZeroVector;
+		if (PivotType == ESpritePivotMode::Custom)
+		{
+			TextureSpacePivotPoint.X = Frame.SpritePosInSheet.X - Frame.SpriteSourcePos.X + Frame.ImageSourceSize.X * Frame.Pivot.X;
+			TextureSpacePivotPoint.Y = Frame.SpritePosInSheet.Y - Frame.SpriteSourcePos.Y + Frame.ImageSourceSize.Y * Frame.Pivot.Y;
+		}
+		TargetSprite->SetPivotMode(PivotType, TextureSpacePivotPoint);
+
+		// Create the entry in the animation
+		SpriteSheet->SpriteNames.Add(Frame.FrameName.ToString());
+		SpriteSheet->Sprites.Add(TargetSprite);
+		
+		TargetSprite->PostEditChange();
+	}
+
+	SpriteSheet->TextureName = Image;
+	SpriteSheet->Texture = ImageTexture;
+
+	GWarn->EndSlowTask();
+	return true;
+}
diff --git a/Engine/Plugins/2D/Paper2D/Source/PaperSpriteSheetImporter/Private/PaperJsonSpriteSheetImporter.h b/Engine/Plugins/2D/Paper2D/Source/PaperSpriteSheetImporter/Private/PaperJsonSpriteSheetImporter.h
new file mode 100644
--- /dev/null
+++ b/Engine/Plugins/2D/Paper2D/Source/PaperSpriteSheetImporter/Private/PaperJsonSpriteSheetImporter.h
@@ -0,0 +1,61 @@
+// Copyright 1998-2015 Epic Games, Inc. All Rights Reserved.
+
+#pragma once
+
+#include "PaperSpriteSheetImporterPrivatePCH.h"
+#include "Paper2DClasses.h"
+
+class UPaperSpriteSheet;
+
+//////////////////////////////////////////////////////////////////////////
+// FSpriteFrame
+
+// Represents one parsed frame in a sprite sheet
+struct FSpriteFrame
+{
+	FName FrameName;
+
+	FVector2D SpritePosInSheet;
+	FVector2D SpriteSizeInSheet;
+
+	FVector2D SpriteSourcePos;
+	FVector2D SpriteSourceSize;
+
+	FVector2D ImageSourceSize;
+
+	FVector2D Pivot;
+
+	bool bTrimmed;
+	bool bRotated;
+};
+
+//////////////////////////////////////////////////////////////////////////
+// FJsonPaperSpriteSheetImporter
+
+// Parses a json from FileContents and imports / reimports a spritesheet
+class FPaperJsonSpriteSheetImporter
+{
+public:
+	FPaperJsonSpriteSheetImporter();
+	bool ImportFromString(const FString& FileContents, const FString& NameForErrors);
+	bool ImportFromArchive(FArchive* Archive, const FString& NameForErrors);
+
+	bool ImportTextures(const FString& LongPackagePath, const FString& SourcePath);
+	bool PerformImport(const FString& LongPackagePath, EObjectFlags Flags, UPaperSpriteSheet* SpriteSheet);
+
+	void SetReimportData(const FString& ExistingImageName, UTexture2D* ExistingTexture, const TArray<FString>& ExistingSpriteNames, const TArray< TAssetPtr<class UPaperSprite> >& ExistingSpriteAssetPtrs);
+
+protected:
+	bool Import(TSharedPtr<FJsonObject> SpriteDescriptorObject, const FString& NameForErrors);
+	UPaperSprite* FindExistingSprite(const FString& Name);
+
+	TArray<FSpriteFrame> Frames;
+	FString Image;
+	UTexture2D* ImageTexture;
+
+	bool bIsReimporting;
+ 	FString ExistingTextureName;					// Name of the last imported texture
+	UTexture2D* ExistingTexture;					// And the texture itself
+	TMap<FString, UPaperSprite*> ExistingSprites;	// Map of a sprite name (as seen in the importer) -> UPaperSprite
+};
+
diff --git a/Engine/Plugins/2D/Paper2D/Source/PaperSpriteSheetImporter/Private/PaperSpriteSheet.cpp b/Engine/Plugins/2D/Paper2D/Source/PaperSpriteSheetImporter/Private/PaperSpriteSheet.cpp
new file mode 100644
--- /dev/null
+++ b/Engine/Plugins/2D/Paper2D/Source/PaperSpriteSheetImporter/Private/PaperSpriteSheet.cpp
@@ -0,0 +1,21 @@
+// Copyright 1998-2015 Epic Games, Inc. All Rights Reserved.
+
+#include "PaperSpriteSheetImporterPrivatePCH.h"
+#include "PaperSpriteSheet.h"
+
+UPaperSpriteSheet::UPaperSpriteSheet(const FObjectInitializer& ObjectInitializer)
+	: Super(ObjectInitializer)
+{
+}
+
+#if WITH_EDITORONLY_DATA
+void UPaperSpriteSheet::GetAssetRegistryTags(TArray<FAssetRegistryTag>& OutTags) const
+{
+	if (AssetImportData)
+	{
+		OutTags.Add( FAssetRegistryTag(SourceFileTagName(), AssetImportData->SourceFilePath, FAssetRegistryTag::TT_Hidden) );
+	}
+
+	Super::GetAssetRegistryTags(OutTags);
+}
+#endif
\ No newline at end of file
diff --git a/Engine/Plugins/2D/Paper2D/Source/PaperSpriteSheetImporter/Private/PaperSpriteSheet.h b/Engine/Plugins/2D/Paper2D/Source/PaperSpriteSheetImporter/Private/PaperSpriteSheet.h
new file mode 100644
--- /dev/null
+++ b/Engine/Plugins/2D/Paper2D/Source/PaperSpriteSheetImporter/Private/PaperSpriteSheet.h
@@ -0,0 +1,36 @@
+// Copyright 1998-2015 Epic Games, Inc. All Rights Reserved.
+
+#pragma once
+
+#include "PaperSpriteSheet.generated.h"
+
+UCLASS(BlueprintType, meta = (DisplayThumbnail = "true"))
+class UPaperSpriteSheet : public UObject
+{
+	GENERATED_UCLASS_BODY()
+
+public:
+	// The names of sprites during import
+	UPROPERTY(VisibleAnywhere, Category=Data)
+	TArray<FString> SpriteNames;
+
+	UPROPERTY(VisibleAnywhere, Category = Data)
+	TArray< TAssetPtr<class UPaperSprite> > Sprites;
+
+	// The name of the texture during import
+	UPROPERTY(VisibleAnywhere, Category = Data)
+	FString TextureName;
+
+	UPROPERTY(VisibleAnywhere, Category = Data)
+	UTexture2D* Texture;
+
+#if WITH_EDITORONLY_DATA
+	// Import data for this 
+	UPROPERTY(VisibleAnywhere, Instanced, Category=ImportSettings)
+	class UAssetImportData* AssetImportData;
+
+	// UObject interface
+	virtual void GetAssetRegistryTags(TArray<FAssetRegistryTag>& OutTags) const override;
+	// End of UObject interface
+#endif
+};
diff --git a/Engine/Plugins/2D/Paper2D/Source/PaperSpriteSheetImporter/Private/PaperSpriteSheetAssetTypeActions.cpp b/Engine/Plugins/2D/Paper2D/Source/PaperSpriteSheetImporter/Private/PaperSpriteSheetAssetTypeActions.cpp
new file mode 100644
--- /dev/null
+++ b/Engine/Plugins/2D/Paper2D/Source/PaperSpriteSheetImporter/Private/PaperSpriteSheetAssetTypeActions.cpp
@@ -0,0 +1,169 @@
+// Copyright 1998-2015 Epic Games, Inc. All Rights Reserved.
+
+#include "PaperSpriteSheetImporterPrivatePCH.h"
+#include "Paper2DClasses.h"
+#include "PaperSpriteSheet.h"
+#include "PaperSpriteSheetAssetTypeActions.h"
+#include "AssetToolsModule.h"
+#include "ContentBrowserModule.h"
+#include "AssetRegistryModule.h"
+#include "PaperFlipbookFactory.h"
+#include "PackageTools.h"
+#include "PaperFlipbookHelpers.h"
+
+#define LOCTEXT_NAMESPACE "AssetTypeActions"
+
+//////////////////////////////////////////////////////////////////////////
+// FPaperSpriteSheetImportAssetTypeActions
+
+FText FPaperSpriteSheetAssetTypeActions::GetName() const
+{
+	return LOCTEXT("FSpriteSheetAssetTypeActionsName", "Sprite Sheet");
+}
+
+FColor FPaperSpriteSheetAssetTypeActions::GetTypeColor() const
+{
+	return FColor::Cyan;
+}
+
+UClass* FPaperSpriteSheetAssetTypeActions::GetSupportedClass() const
+{
+	return UPaperSpriteSheet::StaticClass();
+}
+
+uint32 FPaperSpriteSheetAssetTypeActions::GetCategories()
+{
+	return EAssetTypeCategories::Misc;
+}
+
+bool FPaperSpriteSheetAssetTypeActions::IsImportedAsset() const
+{
+	return true;
+}
+
+void FPaperSpriteSheetAssetTypeActions::GetResolvedSourceFilePaths(const TArray<UObject*>& TypeAssets, TArray<FString>& OutSourceFilePaths) const
+{
+	for (auto& Asset : TypeAssets)
+	{
+		const auto SpriteSheet = CastChecked<UPaperSpriteSheet>(Asset);
+		if (SpriteSheet->AssetImportData)
+		{
+			OutSourceFilePaths.Add(FReimportManager::ResolveImportFilename(SpriteSheet->AssetImportData->SourceFilePath, SpriteSheet));
+		}
+	}
+}
+
+void FPaperSpriteSheetAssetTypeActions::GetActions(const TArray<UObject*>& InObjects, FMenuBuilder& MenuBuilder)
+{
+	auto SpriteSheetImports = GetTypedWeakObjectPtrs<UPaperSpriteSheet>(InObjects);
+
+	MenuBuilder.AddMenuEntry(
+		LOCTEXT("SpriteSheet_CreateFlipbooks", "Create Flipbooks"),
+		LOCTEXT("SpriteSheet_CreateFlipbooksTooltip", "Creates flipbooks from sprites in this sprite sheet."),
+		FSlateIcon(FEditorStyle::GetStyleSetName(), "ClassIcon.PaperFlipbook"),
+		FUIAction(
+		FExecuteAction::CreateSP(this, &FPaperSpriteSheetAssetTypeActions::ExecuteCreateFlipbooks, SpriteSheetImports),
+		FCanExecuteAction()
+		)
+	);
+}
+
+//////////////////////////////////////////////////////////////////////////
+
+void FPaperSpriteSheetAssetTypeActions::ExecuteCreateFlipbooks(TArray<TWeakObjectPtr<UPaperSpriteSheet>> Objects)
+{	
+	for (int SpriteSheetIndex = 0; SpriteSheetIndex < Objects.Num(); ++SpriteSheetIndex)
+	{
+		UPaperSpriteSheet* SpriteSheet = Objects[SpriteSheetIndex].Get();
+		if (SpriteSheet != nullptr)
+		{
+			const FString PackagePath = FPackageName::GetLongPackagePath(SpriteSheet->GetOutermost()->GetPathName());
+
+			FAssetToolsModule& AssetToolsModule = FModuleManager::GetModuleChecked<FAssetToolsModule>("AssetTools");
+			FContentBrowserModule& ContentBrowserModule = FModuleManager::LoadModuleChecked<FContentBrowserModule>("ContentBrowser");
+			
+			check(SpriteSheet->SpriteNames.Num() == SpriteSheet->Sprites.Num());
+			bool useSpriteNames = (SpriteSheet->SpriteNames.Num() == SpriteSheet->Sprites.Num());
+
+			// Create a list of sprites and sprite names to feed into paper flipbook helpers
+			TMap<FString, TArray<UPaperSprite*> > SpriteFlipbookMap;
+
+			{
+				TArray<UPaperSprite*> Sprites;
+				TArray<FString> SpriteNames;
+
+				for (int SpriteIndex = 0; SpriteIndex < SpriteSheet->Sprites.Num(); ++SpriteIndex)
+				{
+					auto SpriteAssetPtr = SpriteSheet->Sprites[SpriteIndex];
+					FStringAssetReference SpriteStringRef = SpriteAssetPtr.ToStringReference();
+					UPaperSprite* Sprite = nullptr;
+					if (!SpriteStringRef.ToString().IsEmpty())
+					{
+						Sprite = Cast<UPaperSprite>(StaticLoadObject(UPaperSprite::StaticClass(), nullptr, *SpriteStringRef.ToString(), nullptr, LOAD_None, nullptr));
+					}
+					if (Sprite != nullptr)
+					{
+						const FString SpriteName = useSpriteNames ? SpriteSheet->SpriteNames[SpriteIndex] : Sprite->GetName();
+						Sprites.Add(Sprite);
+						SpriteNames.Add(SpriteName);
+					}
+				}
+
+				FPaperFlipbookHelpers::ExtractFlipbooksFromSprites(/*out*/ SpriteFlipbookMap, Sprites, SpriteNames);
+			}
+
+			// Create one flipbook for every grouped flipbook name
+			if (SpriteFlipbookMap.Num() > 0)
+			{
+				UPaperFlipbookFactory* FlipbookFactory = NewObject<UPaperFlipbookFactory>();
+
+				GWarn->BeginSlowTask(NSLOCTEXT("Paper2D", "Paper2D_CreateFlipbooks", "Creating flipbooks from selection"), true, true);
+
+				int Progress = 0;
+				int TotalProgress = SpriteFlipbookMap.Num();
+				TArray<UObject*> ObjectsToSync;
+				for (auto It : SpriteFlipbookMap)
+				{
+					GWarn->UpdateProgress(Progress++, TotalProgress);
+
+					const FString& FlipbookName = It.Key;
+					TArray<UPaperSprite*>& Sprites = It.Value;
+
+					const FString TentativePackagePath = PackageTools::SanitizePackageName(PackagePath + TEXT("/") + FlipbookName);
+					FString DefaultSuffix;
+					FString AssetName;
+					FString PackageName;
+					AssetToolsModule.Get().CreateUniqueAssetName(TentativePackagePath, DefaultSuffix, /*out*/ PackageName, /*out*/ AssetName);
+
+					FlipbookFactory->KeyFrames.Empty();
+					for (int32 SpriteIndex = 0; SpriteIndex < Sprites.Num(); ++SpriteIndex)
+					{
+						FPaperFlipbookKeyFrame* KeyFrame = new (FlipbookFactory->KeyFrames) FPaperFlipbookKeyFrame();
+						KeyFrame->Sprite = Sprites[SpriteIndex];
+						KeyFrame->FrameRun = 1;
+					}
+
+					if (UObject* NewAsset = AssetToolsModule.Get().CreateAsset(AssetName, PackagePath, UPaperFlipbook::StaticClass(), FlipbookFactory))
+					{
+						ObjectsToSync.Add(NewAsset);
+					}
+
+					if (GWarn->ReceivedUserCancel())
+					{
+						break;
+					}
+				}
+
+				GWarn->EndSlowTask();
+
+				if (ObjectsToSync.Num() > 0)
+				{
+					ContentBrowserModule.Get().SyncBrowserToAssets(ObjectsToSync);
+				}
+			}
+		}
+	}
+}
+
+#undef LOCTEXT_NAMESPACE
+
diff --git a/Engine/Plugins/2D/Paper2D/Source/PaperSpriteSheetImporter/Private/PaperSpriteSheetAssetTypeActions.h b/Engine/Plugins/2D/Paper2D/Source/PaperSpriteSheetImporter/Private/PaperSpriteSheetAssetTypeActions.h
new file mode 100644
--- /dev/null
+++ b/Engine/Plugins/2D/Paper2D/Source/PaperSpriteSheetImporter/Private/PaperSpriteSheetAssetTypeActions.h
@@ -0,0 +1,25 @@
+// Copyright 1998-2015 Epic Games, Inc. All Rights Reserved.
+
+#pragma once
+
+#include "AssetTypeActions_Base.h"
+
+class UPaperSpriteSheet;
+
+class FPaperSpriteSheetAssetTypeActions : public FAssetTypeActions_Base
+{
+public:
+	// IAssetTypeActions interface
+	virtual FText GetName() const override;
+	virtual FColor GetTypeColor() const override;
+	virtual UClass* GetSupportedClass() const override;
+	virtual bool HasActions(const TArray<UObject*>& InObjects) const override { return true; }
+	virtual void GetActions(const TArray<UObject*>& InObjects, FMenuBuilder& MenuBuilder) override;
+	virtual uint32 GetCategories() override;
+	virtual bool IsImportedAsset() const override;
+	virtual void GetResolvedSourceFilePaths(const TArray<UObject*>& TypeAssets, TArray<FString>& OutSourceFilePaths) const override;
+	// End of IAssetTypeActions interface
+
+private:
+	void ExecuteCreateFlipbooks(TArray<TWeakObjectPtr<UPaperSpriteSheet>> Objects);
+};
\ No newline at end of file
diff --git a/Engine/Plugins/2D/Paper2D/Source/PaperSpriteSheetImporter/Private/PaperSpriteSheetImportFactory.cpp b/Engine/Plugins/2D/Paper2D/Source/PaperSpriteSheetImporter/Private/PaperSpriteSheetImportFactory.cpp
new file mode 100644
--- /dev/null
+++ b/Engine/Plugins/2D/Paper2D/Source/PaperSpriteSheetImporter/Private/PaperSpriteSheetImportFactory.cpp
@@ -0,0 +1,89 @@
+// Copyright 1998-2015 Epic Games, Inc. All Rights Reserved.
+
+#include "PaperSpriteSheetImporterPrivatePCH.h"
+#include "Paper2DClasses.h"
+#include "Json.h"
+#include "PaperJSONHelpers.h"
+#include "AssetToolsModule.h"
+#include "AssetRegistryModule.h"
+#include "PackageTools.h"
+#include "PaperSpriteSheet.h"
+#include "PaperJsonSpriteSheetImporter.h"
+#include "PaperSpriteSheetImportFactory.h"
+#include "PaperSpriteSheetReimportFactory.h"
+
+//////////////////////////////////////////////////////////////////////////
+// UPaperSpriteSheetImportFactory
+
+UPaperSpriteSheetImportFactory::UPaperSpriteSheetImportFactory(const FObjectInitializer& ObjectInitializer)
+	: Super(ObjectInitializer)
+{
+	bCreateNew = false;
+	//bEditAfterNew = true;
+	SupportedClass = UPaperSpriteSheet::StaticClass();
+
+	bEditorImport = true;
+	bText = true;
+
+	Formats.Add(TEXT("json;Spritesheet JSON file"));
+	Formats.Add(TEXT("paper2dsprites;Spritesheet JSON file"));
+}
+
+FText UPaperSpriteSheetImportFactory::GetToolTip() const
+{
+	return NSLOCTEXT("Paper2D", "PaperJsonImporterFactoryDescription", "Sprite sheets exported from Adobe Flash or Texture Packer");
+}
+
+UObject* UPaperSpriteSheetImportFactory::FactoryCreateText(UClass* InClass, UObject* InParent, FName InName, EObjectFlags Flags, UObject* Context, const TCHAR* Type, const TCHAR*& Buffer, const TCHAR* BufferEnd, FFeedbackContext* Warn)
+{
+	Flags |= RF_Transactional;
+
+	FEditorDelegates::OnAssetPreImport.Broadcast(this, InClass, InParent, InName, Type);
+
+	FAssetToolsModule& AssetToolsModule = FModuleManager::GetModuleChecked<FAssetToolsModule>("AssetTools");
+
+	const FString CurrentFilename = UFactory::GetCurrentFilename();
+	FString CurrentSourcePath;
+	FString FilenameNoExtension;
+	FString UnusedExtension;
+	FPaths::Split(CurrentFilename, CurrentSourcePath, FilenameNoExtension, UnusedExtension);
+
+	
+
+	const FString LongPackagePath = FPackageName::GetLongPackagePath(InParent->GetOutermost()->GetPathName());
+	UPaperSpriteSheet* Result = nullptr;
+	
+	const FString NameForErrors(InName.ToString());
+	const FString FileContent(BufferEnd - Buffer, Buffer);
+	
+	FPaperJsonSpriteSheetImporter Importer;
+
+	// Are we reimporting?
+	if (UPaperSpriteSheetReimportFactory* CurrentReimportFactory = Cast<UPaperSpriteSheetReimportFactory>(this))
+	{
+		Importer.SetReimportData(CurrentReimportFactory->ExistingTextureName, CurrentReimportFactory->ExistingTexture, CurrentReimportFactory->ExistingSpriteNames, CurrentReimportFactory->ExistingSprites);
+	}
+
+	if (Importer.ImportFromString(FileContent, NameForErrors) &&
+		Importer.ImportTextures(LongPackagePath, CurrentSourcePath))
+	{
+		UPaperSpriteSheet* SpriteSheet = NewObject<UPaperSpriteSheet>(InParent, InName, Flags);
+		if (Importer.PerformImport(LongPackagePath, Flags, SpriteSheet))
+		{
+			if (SpriteSheet->AssetImportData == nullptr)
+			{
+				SpriteSheet->AssetImportData = NewObject<UAssetImportData>(SpriteSheet);
+			}
+
+			SpriteSheet->AssetImportData->SourceFilePath = FReimportManager::SanitizeImportFilename(CurrentFilename, SpriteSheet);
+			SpriteSheet->AssetImportData->SourceFileTimestamp = IFileManager::Get().GetTimeStamp(*CurrentFilename).ToString();
+			SpriteSheet->AssetImportData->bDirty = false;
+
+			Result = SpriteSheet;
+		}
+	}
+	
+	FEditorDelegates::OnAssetPostImport.Broadcast(this, Result);
+
+	return Result;
+}
diff --git a/Engine/Plugins/2D/Paper2D/Source/PaperSpriteSheetImporter/Private/PaperSpriteSheetImporterModule.cpp b/Engine/Plugins/2D/Paper2D/Source/PaperSpriteSheetImporter/Private/PaperSpriteSheetImporterModule.cpp
new file mode 100644
--- /dev/null
+++ b/Engine/Plugins/2D/Paper2D/Source/PaperSpriteSheetImporter/Private/PaperSpriteSheetImporterModule.cpp
@@ -0,0 +1,42 @@
+// Copyright 1998-2015 Epic Games, Inc. All Rights Reserved.
+
+#include "PaperSpriteSheetImporterPrivatePCH.h"
+
+#include "AssetToolsModule.h"
+#include "PaperSpriteSheetAssetTypeActions.h"
+
+//////////////////////////////////////////////////////////////////////////
+// FPaperSpriteSheetImporterModule
+
+class FPaperSpriteSheetImporterModule : public FDefaultModuleImpl
+{
+public:
+	virtual void StartupModule() override
+	{
+		// Register asset types
+		IAssetTools& AssetTools = FModuleManager::LoadModuleChecked<FAssetToolsModule>("AssetTools").Get();
+
+		SpriteSheetImportAssetTypeActions = MakeShareable(new FPaperSpriteSheetAssetTypeActions);
+		AssetTools.RegisterAssetTypeActions(SpriteSheetImportAssetTypeActions.ToSharedRef());
+	}
+
+	virtual void ShutdownModule() override
+	{
+		if (FModuleManager::Get().IsModuleLoaded("AssetTools"))
+		{
+			IAssetTools& AssetTools = FModuleManager::GetModuleChecked<FAssetToolsModule>("AssetTools").Get();
+			if (SpriteSheetImportAssetTypeActions.IsValid())
+			{
+				AssetTools.UnregisterAssetTypeActions(SpriteSheetImportAssetTypeActions.ToSharedRef());
+			}
+		}
+	}
+
+private:
+	TSharedPtr<IAssetTypeActions> SpriteSheetImportAssetTypeActions;
+};
+
+//////////////////////////////////////////////////////////////////////////
+
+IMPLEMENT_MODULE(FPaperSpriteSheetImporterModule, PaperSpriteSheetImporter);
+DEFINE_LOG_CATEGORY(LogPaperSpriteSheetImporter);
diff --git a/Engine/Plugins/2D/Paper2D/Source/PaperSpriteSheetImporter/Private/PaperSpriteSheetImporterPrivatePCH.h b/Engine/Plugins/2D/Paper2D/Source/PaperSpriteSheetImporter/Private/PaperSpriteSheetImporterPrivatePCH.h
new file mode 100644
--- /dev/null
+++ b/Engine/Plugins/2D/Paper2D/Source/PaperSpriteSheetImporter/Private/PaperSpriteSheetImporterPrivatePCH.h
@@ -0,0 +1,9 @@
+// Copyright 1998-2015 Epic Games, Inc. All Rights Reserved.
+
+#pragma once
+
+#include "UnrealEd.h"
+#include "Paper2DClasses.h"
+
+DECLARE_LOG_CATEGORY_EXTERN(LogPaperSpriteSheetImporter, Verbose, All);
+
diff --git a/Engine/Plugins/2D/Paper2D/Source/PaperSpriteSheetImporter/Private/PaperSpriteSheetReimportFactory.cpp b/Engine/Plugins/2D/Paper2D/Source/PaperSpriteSheetImporter/Private/PaperSpriteSheetReimportFactory.cpp
new file mode 100644
--- /dev/null
+++ b/Engine/Plugins/2D/Paper2D/Source/PaperSpriteSheetImporter/Private/PaperSpriteSheetReimportFactory.cpp
@@ -0,0 +1,90 @@
+// Copyright 1998-2015 Epic Games, Inc. All Rights Reserved.
+
+#include "PaperSpriteSheetImporterPrivatePCH.h"
+#include "PaperSpriteSheetReimportFactory.h"
+#include "PaperSpriteSheet.h"
+
+#define LOCTEXT_NAMESPACE "PaperJsonImporter"
+
+//////////////////////////////////////////////////////////////////////////
+// UPaperSpriteSheetReimportFactory
+
+UPaperSpriteSheetReimportFactory::UPaperSpriteSheetReimportFactory(const FObjectInitializer& ObjectInitializer)
+	: Super(ObjectInitializer)
+{
+	SupportedClass = UPaperSpriteSheet::StaticClass();
+
+	bCreateNew = false;
+	bText = true;
+}
+
+bool UPaperSpriteSheetReimportFactory::CanReimport(UObject* Obj, TArray<FString>& OutFilenames)
+{
+	UPaperSpriteSheet* SpriteSheet = Cast<UPaperSpriteSheet>(Obj);
+	if (SpriteSheet && SpriteSheet->AssetImportData)
+	{
+		OutFilenames.Add(FReimportManager::ResolveImportFilename(SpriteSheet->AssetImportData->SourceFilePath, SpriteSheet));
+		return true;
+	}
+	return false;
+}
+
+void UPaperSpriteSheetReimportFactory::SetReimportPaths(UObject* Obj, const TArray<FString>& NewReimportPaths)
+{
+	UPaperSpriteSheet* SpriteSheet = Cast<UPaperSpriteSheet>(Obj);
+	if (SpriteSheet && ensure(NewReimportPaths.Num() == 1))
+	{
+		SpriteSheet->AssetImportData->SourceFilePath = FReimportManager::ResolveImportFilename(NewReimportPaths[0], SpriteSheet);
+	}
+}
+
+EReimportResult::Type UPaperSpriteSheetReimportFactory::Reimport(UObject* Obj)
+{
+	UPaperSpriteSheet* SpriteSheet = Cast<UPaperSpriteSheet>(Obj);
+	if (!SpriteSheet)
+	{
+		return EReimportResult::Failed;
+	}
+
+	// Make sure file is valid and exists
+	const FString Filename = FReimportManager::ResolveImportFilename(SpriteSheet->AssetImportData->SourceFilePath, SpriteSheet);
+	if (!Filename.Len() || IFileManager::Get().FileSize(*Filename) == INDEX_NONE)
+	{
+		return EReimportResult::Failed;
+	}
+
+	ExistingSpriteNames = SpriteSheet->SpriteNames;
+	ExistingSprites = SpriteSheet->Sprites;
+	ExistingTextureName = SpriteSheet->TextureName;
+	ExistingTexture = SpriteSheet->Texture;
+
+	EReimportResult::Type Result = EReimportResult::Failed;
+	if (UFactory::StaticImportObject(SpriteSheet->GetClass(), SpriteSheet->GetOuter(), *SpriteSheet->GetName(), RF_Public | RF_Standalone, *Filename, nullptr, this))
+	{
+		UE_LOG(LogPaperSpriteSheetImporter, Log, TEXT("Imported successfully"));
+		// Try to find the outer package so we can dirty it up
+		if (SpriteSheet->GetOuter())
+		{
+			SpriteSheet->GetOuter()->MarkPackageDirty();
+		}
+		else
+		{
+			SpriteSheet->MarkPackageDirty();
+		}
+		Result = EReimportResult::Succeeded;
+	}
+	else
+	{
+		UE_LOG(LogPaperSpriteSheetImporter, Warning, TEXT("-- import failed"));
+		Result = EReimportResult::Failed;
+	}
+
+	return Result;
+}
+
+int32 UPaperSpriteSheetReimportFactory::GetPriority() const
+{
+	return ImportPriority;
+}
+
+#undef LOCTEXT_NAMESPACE
diff --git a/Engine/Plugins/2D/Paper2D/Source/PaperSpriteSheetImporter/Private/PaperSpriteSheetReimportFactory.h b/Engine/Plugins/2D/Paper2D/Source/PaperSpriteSheetImporter/Private/PaperSpriteSheetReimportFactory.h
new file mode 100644
--- /dev/null
+++ b/Engine/Plugins/2D/Paper2D/Source/PaperSpriteSheetImporter/Private/PaperSpriteSheetReimportFactory.h
@@ -0,0 +1,27 @@
+// Copyright 1998-2015 Epic Games, Inc. All Rights Reserved.
+
+#pragma once
+
+#include "PaperSpriteSheetImportFactory.h"
+#include "PaperSpriteSheetReimportFactory.generated.h"
+
+// Imports a paper animated sprite (and associated paper sprites and textures) from an Adobe Flash CS6 exported JSON file
+UCLASS()
+class UPaperSpriteSheetReimportFactory : public UPaperSpriteSheetImportFactory, public FReimportHandler
+{
+	GENERATED_UCLASS_BODY()
+
+	// Begin FReimportHandler interface
+	virtual bool CanReimport(UObject* Obj, TArray<FString>& OutFilenames) override;
+	virtual void SetReimportPaths(UObject* Obj, const TArray<FString>& NewReimportPaths) override;
+	virtual EReimportResult::Type Reimport(UObject* Obj) override;
+	virtual int32 GetPriority() const override;
+	// End FReimportHandler interface
+
+	TArray<FString> ExistingSpriteNames;
+	TArray< TAssetPtr<class UPaperSprite> > ExistingSprites;
+	
+	FString ExistingTextureName;
+	UTexture2D* ExistingTexture;
+};
+
