Paper2D: Added includes necessary for compiling the Paper2D editor without including Paper2DEditorClasses.h in the PCH
- Deleted empty unused tool classes
- Renamed PaperAnimatedSpriteActorFactory.cpp to PaperFlipbookActorFactory.cpp to match the current naming

diff --git a/Engine/Plugins/2D/Paper2D/Source/Paper2DEditor/Classes/PaperFlipbookThumbnailRenderer.h b/Engine/Plugins/2D/Paper2D/Source/Paper2DEditor/Classes/PaperFlipbookThumbnailRenderer.h
--- a/Engine/Plugins/2D/Paper2D/Source/Paper2DEditor/Classes/PaperFlipbookThumbnailRenderer.h
+++ b/Engine/Plugins/2D/Paper2D/Source/Paper2DEditor/Classes/PaperFlipbookThumbnailRenderer.h
@@ -2,6 +2,7 @@
 
 #pragma once
 
+#include "PaperSpriteThumbnailRenderer.h"
 #include "PaperFlipbookThumbnailRenderer.generated.h"
 
 UCLASS()
diff --git a/Engine/Plugins/2D/Paper2D/Source/Paper2DEditor/Classes/SpriteEditorTools/PSE_DeleteTool.h b/Engine/Plugins/2D/Paper2D/Source/Paper2DEditor/Classes/SpriteEditorTools/PSE_DeleteTool.h
deleted file mode 100644
diff --git a/Engine/Plugins/2D/Paper2D/Source/Paper2DEditor/Classes/SpriteEditorTools/PSE_LineTool.h b/Engine/Plugins/2D/Paper2D/Source/Paper2DEditor/Classes/SpriteEditorTools/PSE_LineTool.h
deleted file mode 100644
diff --git a/Engine/Plugins/2D/Paper2D/Source/Paper2DEditor/Classes/SpriteEditorTools/PSE_SplitEdgeTool.h b/Engine/Plugins/2D/Paper2D/Source/Paper2DEditor/Classes/SpriteEditorTools/PSE_SplitEdgeTool.h
deleted file mode 100644
diff --git a/Engine/Plugins/2D/Paper2D/Source/Paper2DEditor/Classes/SpriteEditorTools/SpriteEditorTool.h b/Engine/Plugins/2D/Paper2D/Source/Paper2DEditor/Classes/SpriteEditorTools/SpriteEditorTool.h
deleted file mode 100644
diff --git a/Engine/Plugins/2D/Paper2D/Source/Paper2DEditor/Private/ContentBrowserExtensions/ContentBrowserExtensions.cpp b/Engine/Plugins/2D/Paper2D/Source/Paper2DEditor/Private/ContentBrowserExtensions/ContentBrowserExtensions.cpp
--- a/Engine/Plugins/2D/Paper2D/Source/Paper2DEditor/Private/ContentBrowserExtensions/ContentBrowserExtensions.cpp
+++ b/Engine/Plugins/2D/Paper2D/Source/Paper2DEditor/Private/ContentBrowserExtensions/ContentBrowserExtensions.cpp
@@ -3,6 +3,8 @@
 #include "Paper2DEditorPrivatePCH.h"
 #include "ContentBrowserModule.h"
 #include "ContentBrowserExtensions.h"
+#include "PaperSpriteFactory.h"
+#include "PaperTileSetFactory.h"
 
 #include "ExtractSprites/SPaperExtractSpritesDialog.h"
 
diff --git a/Engine/Plugins/2D/Paper2D/Source/Paper2DEditor/Private/ExtractSprites/SPaperExtractSpritesDialog.cpp b/Engine/Plugins/2D/Paper2D/Source/Paper2DEditor/Private/ExtractSprites/SPaperExtractSpritesDialog.cpp
--- a/Engine/Plugins/2D/Paper2D/Source/Paper2DEditor/Private/ExtractSprites/SPaperExtractSpritesDialog.cpp
+++ b/Engine/Plugins/2D/Paper2D/Source/Paper2DEditor/Private/ExtractSprites/SPaperExtractSpritesDialog.cpp
@@ -10,6 +10,7 @@
 #include "ContentBrowserModule.h"
 #include "AssetToolsModule.h"
 #include "CanvasTypes.h"
+#include "PaperSpriteFactory.h"
 
 #define LOCTEXT_NAMESPACE "PaperEditor"
 
diff --git a/Engine/Plugins/2D/Paper2D/Source/Paper2DEditor/Private/Paper2DEditorModule.cpp b/Engine/Plugins/2D/Paper2D/Source/Paper2DEditor/Private/Paper2DEditorModule.cpp
--- a/Engine/Plugins/2D/Paper2D/Source/Paper2DEditor/Private/Paper2DEditorModule.cpp
+++ b/Engine/Plugins/2D/Paper2D/Source/Paper2DEditor/Private/Paper2DEditorModule.cpp
@@ -15,6 +15,7 @@
 // Sprite support
 #include "SpriteAssetTypeActions.h"
 #include "PaperSpriteAssetBroker.h"
+#include "PaperSpriteThumbnailRenderer.h"
 #include "SpriteEditor/SpriteDetailsCustomization.h"
 #include "SpriteEditor/SpriteComponentDetailsCustomization.h"
 #include "SpriteEditor/SpritePolygonCollectionCustomization.h"
@@ -24,10 +25,12 @@
 // Flipbook support
 #include "FlipbookAssetTypeActions.h"
 #include "PaperFlipbookAssetBroker.h"
+#include "PaperFlipbookThumbnailRenderer.h"
 #include "FlipbookEditor/FlipbookComponentDetailsCustomization.h"
 
 // Tile set support
 #include "TileSetAssetTypeActions.h"
+#include "PaperTileSetThumbnailRenderer.h"
 
 // Tile map support
 #include "TileMapEditing/TileMapAssetTypeActions.h"
diff --git a/Engine/Plugins/2D/Paper2D/Source/Paper2DEditor/Private/PaperAnimatedSpriteActorFactory.cpp b/Engine/Plugins/2D/Paper2D/Source/Paper2DEditor/Private/PaperAnimatedSpriteActorFactory.cpp
deleted file mode 100644
diff --git a/Engine/Plugins/2D/Paper2D/Source/Paper2DEditor/Private/PaperFlipbookActorFactory.cpp b/Engine/Plugins/2D/Paper2D/Source/Paper2DEditor/Private/PaperFlipbookActorFactory.cpp
new file mode 100644
--- /dev/null
+++ b/Engine/Plugins/2D/Paper2D/Source/Paper2DEditor/Private/PaperFlipbookActorFactory.cpp
@@ -0,0 +1,58 @@
+// Copyright 1998-2015 Epic Games, Inc. All Rights Reserved.
+
+#include "Paper2DEditorPrivatePCH.h"
+#include "AssetData.h"
+#include "PaperFlipbookActorFactory.h"
+
+//////////////////////////////////////////////////////////////////////////
+// UPaperFlipbookActorFactory
+
+UPaperFlipbookActorFactory::UPaperFlipbookActorFactory(const FObjectInitializer& ObjectInitializer)
+	: Super(ObjectInitializer)
+{
+	DisplayName = NSLOCTEXT("Paper2D", "PaperFlipbookFactoryDisplayName", "Add Animated Sprite");
+	NewActorClass = APaperFlipbookActor::StaticClass();
+}
+
+void UPaperFlipbookActorFactory::PostSpawnActor(UObject* Asset, AActor* NewActor)
+{
+	Super::PostSpawnActor(Asset, NewActor);
+
+	if (UPaperFlipbook* Flipbook = Cast<UPaperFlipbook>(Asset))
+	{
+		APaperFlipbookActor* TypedActor = CastChecked<APaperFlipbookActor>(NewActor);
+		UPaperFlipbookComponent* RenderComponent = TypedActor->GetRenderComponent();
+		check(RenderComponent);
+
+		RenderComponent->UnregisterComponent();
+		RenderComponent->SetFlipbook(Flipbook);
+		RenderComponent->RegisterComponent();
+	}
+}
+
+void UPaperFlipbookActorFactory::PostCreateBlueprint(UObject* Asset, AActor* CDO)
+{
+	if (UPaperFlipbook* Flipbook = Cast<UPaperFlipbook>(Asset))
+	{
+		if (APaperFlipbookActor* TypedActor = Cast<APaperFlipbookActor>(CDO))
+		{
+			UPaperFlipbookComponent* RenderComponent = TypedActor->GetRenderComponent();
+			check(RenderComponent);
+
+			RenderComponent->SetFlipbook(Flipbook);
+		}
+	}
+}
+
+bool UPaperFlipbookActorFactory::CanCreateActorFrom(const FAssetData& AssetData, FText& OutErrorMsg)
+{
+	if (AssetData.IsValid() && AssetData.GetClass()->IsChildOf(UPaperFlipbook::StaticClass()))
+	{
+		return true;
+	}
+	else
+	{
+		OutErrorMsg = NSLOCTEXT("Paper2D", "CanCreateActorFrom_NoFlipbook", "No flipbook was specified.");
+		return false;
+	}
+}
diff --git a/Engine/Plugins/2D/Paper2D/Source/Paper2DEditor/Private/PaperFlipbookFactory.cpp b/Engine/Plugins/2D/Paper2D/Source/Paper2DEditor/Private/PaperFlipbookFactory.cpp
--- a/Engine/Plugins/2D/Paper2D/Source/Paper2DEditor/Private/PaperFlipbookFactory.cpp
+++ b/Engine/Plugins/2D/Paper2D/Source/Paper2DEditor/Private/PaperFlipbookFactory.cpp
@@ -1,6 +1,7 @@
 // Copyright 1998-2015 Epic Games, Inc. All Rights Reserved.
 
 #include "Paper2DEditorPrivatePCH.h"
+#include "PaperFlipbookFactory.h"
 
 #define LOCTEXT_NAMESPACE "Paper2D"
 
diff --git a/Engine/Plugins/2D/Paper2D/Source/Paper2DEditor/Private/PaperFlipbookThumbnailRenderer.cpp b/Engine/Plugins/2D/Paper2D/Source/Paper2DEditor/Private/PaperFlipbookThumbnailRenderer.cpp
--- a/Engine/Plugins/2D/Paper2D/Source/Paper2DEditor/Private/PaperFlipbookThumbnailRenderer.cpp
+++ b/Engine/Plugins/2D/Paper2D/Source/Paper2DEditor/Private/PaperFlipbookThumbnailRenderer.cpp
@@ -2,6 +2,7 @@
 
 #include "Paper2DEditorPrivatePCH.h"
 #include "CanvasItem.h"
+#include "PaperFlipbookThumbnailRenderer.h"
 
 //////////////////////////////////////////////////////////////////////////
 // UPaperFlipbookThumbnailRenderer
diff --git a/Engine/Plugins/2D/Paper2D/Source/Paper2DEditor/Private/PaperSpriteActorFactory.cpp b/Engine/Plugins/2D/Paper2D/Source/Paper2DEditor/Private/PaperSpriteActorFactory.cpp
--- a/Engine/Plugins/2D/Paper2D/Source/Paper2DEditor/Private/PaperSpriteActorFactory.cpp
+++ b/Engine/Plugins/2D/Paper2D/Source/Paper2DEditor/Private/PaperSpriteActorFactory.cpp
@@ -3,6 +3,8 @@
 #include "Paper2DEditorPrivatePCH.h"
 #include "AssetData.h"
 #include "PhysicsEngine/BodySetup.h"
+#include "PaperSpriteActorFactory.h"
+
 //////////////////////////////////////////////////////////////////////////
 // UPaperSpriteActorFactory
 
diff --git a/Engine/Plugins/2D/Paper2D/Source/Paper2DEditor/Private/PaperSpriteFactory.cpp b/Engine/Plugins/2D/Paper2D/Source/Paper2DEditor/Private/PaperSpriteFactory.cpp
--- a/Engine/Plugins/2D/Paper2D/Source/Paper2DEditor/Private/PaperSpriteFactory.cpp
+++ b/Engine/Plugins/2D/Paper2D/Source/Paper2DEditor/Private/PaperSpriteFactory.cpp
@@ -1,6 +1,7 @@
 // Copyright 1998-2015 Epic Games, Inc. All Rights Reserved.
 
 #include "Paper2DEditorPrivatePCH.h"
+#include "PaperSpriteFactory.h"
 
 #define LOCTEXT_NAMESPACE "Paper2D"
 
diff --git a/Engine/Plugins/2D/Paper2D/Source/Paper2DEditor/Private/PaperSpriteThumbnailRenderer.cpp b/Engine/Plugins/2D/Paper2D/Source/Paper2DEditor/Private/PaperSpriteThumbnailRenderer.cpp
--- a/Engine/Plugins/2D/Paper2D/Source/Paper2DEditor/Private/PaperSpriteThumbnailRenderer.cpp
+++ b/Engine/Plugins/2D/Paper2D/Source/Paper2DEditor/Private/PaperSpriteThumbnailRenderer.cpp
@@ -2,6 +2,7 @@
 
 #include "Paper2DEditorPrivatePCH.h"
 #include "CanvasTypes.h"
+#include "PaperSpriteThumbnailRenderer.h"
 
 //////////////////////////////////////////////////////////////////////////
 // UPaperSpriteThumbnailRenderer
diff --git a/Engine/Plugins/2D/Paper2D/Source/Paper2DEditor/Private/PaperTileMapFactory.cpp b/Engine/Plugins/2D/Paper2D/Source/Paper2DEditor/Private/PaperTileMapFactory.cpp
--- a/Engine/Plugins/2D/Paper2D/Source/Paper2DEditor/Private/PaperTileMapFactory.cpp
+++ b/Engine/Plugins/2D/Paper2D/Source/Paper2DEditor/Private/PaperTileMapFactory.cpp
@@ -1,6 +1,7 @@
 // Copyright 1998-2015 Epic Games, Inc. All Rights Reserved.
 
 #include "Paper2DEditorPrivatePCH.h"
+#include "PaperTileMapFactory.h"
 
 #define LOCTEXT_NAMESPACE "Paper2D"
 
diff --git a/Engine/Plugins/2D/Paper2D/Source/Paper2DEditor/Private/PaperTileMapPromotionFactory.cpp b/Engine/Plugins/2D/Paper2D/Source/Paper2DEditor/Private/PaperTileMapPromotionFactory.cpp
--- a/Engine/Plugins/2D/Paper2D/Source/Paper2DEditor/Private/PaperTileMapPromotionFactory.cpp
+++ b/Engine/Plugins/2D/Paper2D/Source/Paper2DEditor/Private/PaperTileMapPromotionFactory.cpp
@@ -1,6 +1,7 @@
 // Copyright 1998-2015 Epic Games, Inc. All Rights Reserved.
 
 #include "Paper2DEditorPrivatePCH.h"
+#include "PaperTileMapPromotionFactory.h"
 
 #define LOCTEXT_NAMESPACE "Paper2D"
 
diff --git a/Engine/Plugins/2D/Paper2D/Source/Paper2DEditor/Private/PaperTileSetFactory.cpp b/Engine/Plugins/2D/Paper2D/Source/Paper2DEditor/Private/PaperTileSetFactory.cpp
--- a/Engine/Plugins/2D/Paper2D/Source/Paper2DEditor/Private/PaperTileSetFactory.cpp
+++ b/Engine/Plugins/2D/Paper2D/Source/Paper2DEditor/Private/PaperTileSetFactory.cpp
@@ -1,6 +1,7 @@
 // Copyright 1998-2015 Epic Games, Inc. All Rights Reserved.
 
 #include "Paper2DEditorPrivatePCH.h"
+#include "PaperTileSetFactory.h"
 
 #define LOCTEXT_NAMESPACE "Paper2D"
 
diff --git a/Engine/Plugins/2D/Paper2D/Source/Paper2DEditor/Private/PaperTileSetThumbnailRenderer.cpp b/Engine/Plugins/2D/Paper2D/Source/Paper2DEditor/Private/PaperTileSetThumbnailRenderer.cpp
--- a/Engine/Plugins/2D/Paper2D/Source/Paper2DEditor/Private/PaperTileSetThumbnailRenderer.cpp
+++ b/Engine/Plugins/2D/Paper2D/Source/Paper2DEditor/Private/PaperTileSetThumbnailRenderer.cpp
@@ -2,6 +2,7 @@
 
 #include "Paper2DEditorPrivatePCH.h"
 #include "CanvasTypes.h"
+#include "PaperTileSetThumbnailRenderer.h"
 
 //////////////////////////////////////////////////////////////////////////
 // UPaperTileSetThumbnailRenderer
diff --git a/Engine/Plugins/2D/Paper2D/Source/Paper2DEditor/Private/SpriteAssetTypeActions.cpp b/Engine/Plugins/2D/Paper2D/Source/Paper2DEditor/Private/SpriteAssetTypeActions.cpp
--- a/Engine/Plugins/2D/Paper2D/Source/Paper2DEditor/Private/SpriteAssetTypeActions.cpp
+++ b/Engine/Plugins/2D/Paper2D/Source/Paper2DEditor/Private/SpriteAssetTypeActions.cpp
@@ -7,6 +7,7 @@
 #include "AssetToolsModule.h"
 #include "ContentBrowserModule.h"
 #include "PaperFlipbookHelpers.h"
+#include "PaperFlipbookFactory.h"
 
 #define LOCTEXT_NAMESPACE "AssetTypeActions"
 
diff --git a/Engine/Plugins/2D/Paper2D/Source/Paper2DEditor/Private/SpriteEditor/PSE_DeleteTool.cpp b/Engine/Plugins/2D/Paper2D/Source/Paper2DEditor/Private/SpriteEditor/PSE_DeleteTool.cpp
deleted file mode 100644
diff --git a/Engine/Plugins/2D/Paper2D/Source/Paper2DEditor/Private/SpriteEditor/PSE_LineTool.cpp b/Engine/Plugins/2D/Paper2D/Source/Paper2DEditor/Private/SpriteEditor/PSE_LineTool.cpp
deleted file mode 100644
diff --git a/Engine/Plugins/2D/Paper2D/Source/Paper2DEditor/Private/SpriteEditor/PSE_SplitEdgeTool.cpp b/Engine/Plugins/2D/Paper2D/Source/Paper2DEditor/Private/SpriteEditor/PSE_SplitEdgeTool.cpp
deleted file mode 100644
diff --git a/Engine/Plugins/2D/Paper2D/Source/Paper2DEditor/Private/SpriteEditor/SpriteEditorTool.cpp b/Engine/Plugins/2D/Paper2D/Source/Paper2DEditor/Private/SpriteEditor/SpriteEditorTool.cpp
deleted file mode 100644
diff --git a/Engine/Plugins/2D/Paper2D/Source/Paper2DEditor/Private/SpriteEditor/SpriteEditorViewportClient.cpp b/Engine/Plugins/2D/Paper2D/Source/Paper2DEditor/Private/SpriteEditor/SpriteEditorViewportClient.cpp
--- a/Engine/Plugins/2D/Paper2D/Source/Paper2DEditor/Private/SpriteEditor/SpriteEditorViewportClient.cpp
+++ b/Engine/Plugins/2D/Paper2D/Source/Paper2DEditor/Private/SpriteEditor/SpriteEditorViewportClient.cpp
@@ -23,6 +23,8 @@
 #include "SNotificationList.h"
 #include "NotificationManager.h"
 
+#include "PaperSpriteFactory.h"
+
 #define LOCTEXT_NAMESPACE "SpriteEditor"
 
 //////////////////////////////////////////////////////////////////////////
diff --git a/Engine/Plugins/2D/Paper2D/Source/Paper2DEditor/Private/TerrainSplineEditing/TerrainSplineActorFactory.cpp b/Engine/Plugins/2D/Paper2D/Source/Paper2DEditor/Private/TerrainSplineEditing/TerrainSplineActorFactory.cpp
--- a/Engine/Plugins/2D/Paper2D/Source/Paper2DEditor/Private/TerrainSplineEditing/TerrainSplineActorFactory.cpp
+++ b/Engine/Plugins/2D/Paper2D/Source/Paper2DEditor/Private/TerrainSplineEditing/TerrainSplineActorFactory.cpp
@@ -2,6 +2,7 @@
 
 #include "Paper2DEditorPrivatePCH.h"
 #include "AssetData.h"
+#include "TerrainSplineActorFactory.h"
 
 //////////////////////////////////////////////////////////////////////////
 // UTerrainSplineActorFactory
diff --git a/Engine/Plugins/2D/Paper2D/Source/Paper2DEditor/Private/TileMapEditing/PaperTileMapDetailsCustomization.cpp b/Engine/Plugins/2D/Paper2D/Source/Paper2DEditor/Private/TileMapEditing/PaperTileMapDetailsCustomization.cpp
--- a/Engine/Plugins/2D/Paper2D/Source/Paper2DEditor/Private/TileMapEditing/PaperTileMapDetailsCustomization.cpp
+++ b/Engine/Plugins/2D/Paper2D/Source/Paper2DEditor/Private/TileMapEditing/PaperTileMapDetailsCustomization.cpp
@@ -7,6 +7,7 @@
 #include "EdModeTileMap.h"
 #include "PropertyEditing.h"
 #include "AssetToolsModule.h"
+#include "PaperTileMapPromotionFactory.h"
 
 #include "STileLayerList.h"
 #include "ScopedTransaction.h"
diff --git a/Engine/Plugins/2D/Paper2D/Source/Paper2DEditor/Private/TileMapEditing/TileMapActorFactory.cpp b/Engine/Plugins/2D/Paper2D/Source/Paper2DEditor/Private/TileMapEditing/TileMapActorFactory.cpp
--- a/Engine/Plugins/2D/Paper2D/Source/Paper2DEditor/Private/TileMapEditing/TileMapActorFactory.cpp
+++ b/Engine/Plugins/2D/Paper2D/Source/Paper2DEditor/Private/TileMapEditing/TileMapActorFactory.cpp
@@ -2,6 +2,7 @@
 
 #include "Paper2DEditorPrivatePCH.h"
 #include "AssetData.h"
+#include "TileMapActorFactory.h"
 
 //////////////////////////////////////////////////////////////////////////
 // UTileMapActorFactory
diff --git a/Engine/Plugins/2D/Paper2D/Source/Paper2DEditor/Private/TileMapEditing/TileMapAssetImportData.cpp b/Engine/Plugins/2D/Paper2D/Source/Paper2DEditor/Private/TileMapEditing/TileMapAssetImportData.cpp
--- a/Engine/Plugins/2D/Paper2D/Source/Paper2DEditor/Private/TileMapEditing/TileMapAssetImportData.cpp
+++ b/Engine/Plugins/2D/Paper2D/Source/Paper2DEditor/Private/TileMapEditing/TileMapAssetImportData.cpp
@@ -2,6 +2,7 @@
 
 #include "Paper2DEditorPrivatePCH.h"
 #include "AssetData.h"
+#include "TileMapAssetImportData.h"
 
 //////////////////////////////////////////////////////////////////////////
 // UTileMapAssetImportData
