Paper2D: Added GetLayerIndex() to UPaperTileLayer and added comments to a few other methods

diff --git a/Engine/Plugins/2D/Paper2D/Source/Paper2D/Classes/PaperTileLayer.h b/Engine/Plugins/2D/Paper2D/Source/Paper2D/Classes/PaperTileLayer.h
--- a/Engine/Plugins/2D/Paper2D/Source/Paper2D/Classes/PaperTileLayer.h
+++ b/Engine/Plugins/2D/Paper2D/Source/Paper2D/Classes/PaperTileLayer.h
@@ -158,15 +158,22 @@ public:
 #endif
 	// End of UObject interface
 
+	// Returns the parent tile map
 	class UPaperTileMap* GetTileMap() const;
 
+	// Returns the index of this layer in the parent tile map
+	int32 GetLayerIndex() const;
+
+	// Returns the tile information about the specified cell
 	FPaperTileInfo GetCell(int32 X, int32 Y) const;
+
+	// Sets the tile information about the specified cell
 	void SetCell(int32 X, int32 Y, const FPaperTileInfo& NewValue);
 
-	// This is a destructive operation!
+	// Reallocates the map.  This is a destructive operation that does not copy data across!
 	void DestructiveAllocateMap(int32 NewWidth, int32 NewHeight);
 	
-	// This tries to preserve contents
+	// Reallocates the map.  This tries to preserve contents.
 	void ResizeMap(int32 NewWidth, int32 NewHeight);
 
 	// Adds collision to the specified body setup
diff --git a/Engine/Plugins/2D/Paper2D/Source/Paper2D/Private/PaperTileLayer.cpp b/Engine/Plugins/2D/Paper2D/Source/Paper2D/Private/PaperTileLayer.cpp
--- a/Engine/Plugins/2D/Paper2D/Source/Paper2D/Private/PaperTileLayer.cpp
+++ b/Engine/Plugins/2D/Paper2D/Source/Paper2D/Private/PaperTileLayer.cpp
@@ -184,6 +184,11 @@ UPaperTileMap* UPaperTileLayer::GetTileMap() const
 	return CastChecked<UPaperTileMap>(GetOuter());
 }
 
+int32 UPaperTileLayer::GetLayerIndex() const
+{
+	return GetTileMap()->TileLayers.Find(const_cast<UPaperTileLayer*>(this));
+}
+
 FPaperTileInfo UPaperTileLayer::GetCell(int32 X, int32 Y) const
 {
 	if ((X < 0) || (X >= LayerWidth) || (Y < 0) || (Y >= LayerHeight))
