Paper2D: Rebuild the nav mesh when painting tile maps with collision [UE-14043]
Note: Currently does a full rebuild for the component, as just flagging a region as dirty fails to build paths for new or removed tiles

diff --git a/Engine/Plugins/2D/Paper2D/Source/Paper2D/Private/PaperTileMapComponent.cpp b/Engine/Plugins/2D/Paper2D/Source/Paper2D/Private/PaperTileMapComponent.cpp
--- a/Engine/Plugins/2D/Paper2D/Source/Paper2D/Private/PaperTileMapComponent.cpp
+++ b/Engine/Plugins/2D/Paper2D/Source/Paper2D/Private/PaperTileMapComponent.cpp
@@ -28,6 +28,7 @@ UPaperTileMapComponent::UPaperTileMapComponent(const FObjectInitializer& ObjectI
 
 	CastShadow = false;
 	bUseAsOccluder = false;
+	bCanEverAffectNavigation = true;
 }
 
 FPrimitiveSceneProxy* UPaperTileMapComponent::CreateSceneProxy()
diff --git a/Engine/Plugins/2D/Paper2D/Source/Paper2DEditor/Private/TileMapEditing/EdModeTileMap.cpp b/Engine/Plugins/2D/Paper2D/Source/Paper2DEditor/Private/TileMapEditing/EdModeTileMap.cpp
--- a/Engine/Plugins/2D/Paper2D/Source/Paper2DEditor/Private/TileMapEditing/EdModeTileMap.cpp
+++ b/Engine/Plugins/2D/Paper2D/Source/Paper2DEditor/Private/TileMapEditing/EdModeTileMap.cpp
@@ -13,6 +13,52 @@
 
 #define LOCTEXT_NAMESPACE "Paper2D"
 
+//////////////////////////////////////////////////////////////////////////
+// Editor constants
+
+namespace TileMapEditorConstants
+{
+	const float IntervalBetweenNavMeshRebuilds = 2.0f;
+	const bool bFlushEntireComponentWhenNavMeshIsDirty = true;
+};
+
+//////////////////////////////////////////////////////////////////////////
+// FTileMapDirtyRegion
+
+FTileMapDirtyRegion::FTileMapDirtyRegion(UPaperTileMapComponent* InComponent, const FBox& DirtyRegionInTileSpace)
+	: ComponentPtr(InComponent)
+	, DirtyRegionInWorldSpace(ForceInitToZero)
+{
+	if (DirtyRegionInTileSpace.IsValid)
+	{
+		if (UPaperTileMap* TileMap = InComponent->TileMap)
+		{
+			const FTransform ComponentToWorld = InComponent->GetComponentTransform();
+			const FVector MinCoordLS = TileMap->GetTilePositionInLocalSpace(DirtyRegionInTileSpace.Min.X, DirtyRegionInTileSpace.Min.Y, (int32)DirtyRegionInTileSpace.Min.Z);
+			const FVector MaxCoordLS = TileMap->GetTilePositionInLocalSpace(DirtyRegionInTileSpace.Max.X + 1.0f, DirtyRegionInTileSpace.Max.Y + 1.0f, (int32)DirtyRegionInTileSpace.Max.Z);
+			
+			DirtyRegionInWorldSpace += ComponentToWorld.TransformPosition(MinCoordLS - TileMap->GetCollisionThickness() * PaperAxisZ);
+			DirtyRegionInWorldSpace += ComponentToWorld.TransformPosition(MaxCoordLS + TileMap->GetCollisionThickness() * PaperAxisZ);
+
+			//DrawDebugBox(InComponent->GetWorld(), DirtyRegionInWorldSpace.GetCenter(), DirtyRegionInWorldSpace.GetExtent(), FQuat::Identity, FColor::White, true, 1.0f);
+		}
+	}
+}
+
+void FTileMapDirtyRegion::PushToNavSystem() const
+{
+	if (UPaperTileMapComponent* Component = ComponentPtr.Get())
+	{
+		if (Component->IsNavigationRelevant())
+		{
+			if (UNavigationSystem* NavSys = UNavigationSystem::GetCurrent(Component))
+			{
+				NavSys->AddDirtyArea(DirtyRegionInWorldSpace, ENavigationDirtyFlag::All);
+			}
+		}
+	}
+}
+
 //////////////////////////////////////////////////////////////////////////
 // FHorizontalSpan - used for flood filling
 
@@ -86,6 +132,8 @@ FEdModeTileMap::FEdModeTileMap()
 	, bIsLastCursorValid(false)
 	, DrawPreviewDimensionsLS(0.0f, 0.0f, 0.0f)
 	, EraseBrushSize(1)
+	, TimeUntilNavMeshRebuild(TileMapEditorConstants::IntervalBetweenNavMeshRebuilds)
+	, ActiveTool(ETileMapEditorTool::Paintbrush)
 {
 	bDrawPivot = false;
 	bDrawGrid = false;
@@ -130,6 +178,8 @@ void FEdModeTileMap::Enter()
 
 void FEdModeTileMap::Exit()
 {
+	FlushPendingDirtyRegions();
+
 	if (Toolkit.IsValid())
 	{
 		FToolkitManager::Get().CloseToolkit(Toolkit.ToSharedRef());
@@ -144,6 +194,54 @@ void FEdModeTileMap::Exit()
 	FEdMode::Exit();
 }
 
+void FEdModeTileMap::FlushPendingDirtyRegions()
+{
+	TSet<UPaperTileMapComponent*> ComponentsToInvalidate;
+
+	for (const FTileMapDirtyRegion& DirtyRegion : PendingDirtyRegions)
+	{
+		if (UPaperTileMapComponent* Component = DirtyRegion.GetComponent())
+		{
+			if (Component->IsNavigationRelevant())
+			{
+				if (TileMapEditorConstants::bFlushEntireComponentWhenNavMeshIsDirty)
+				{
+					ComponentsToInvalidate.Add(Component);
+				}
+				else
+				{
+					DirtyRegion.PushToNavSystem();
+				}
+			}
+		}
+	}
+
+	for (UPaperTileMapComponent* Component : ComponentsToInvalidate)
+	{
+		if (UNavigationSystem* NavSys = UNavigationSystem::GetCurrent(Component))
+		{
+			NavSys->UpdateNavOctree(Component);
+		}
+	}
+
+	PendingDirtyRegions.Empty();
+}
+
+void FEdModeTileMap::Tick(FEditorViewportClient* ViewportClient, float DeltaTime)
+{
+	FEdMode::Tick(ViewportClient, DeltaTime);
+
+	if (PendingDirtyRegions.Num() > 0)
+	{
+		TimeUntilNavMeshRebuild -= DeltaTime;
+		if (TimeUntilNavMeshRebuild <= 0.0f)
+		{
+			TimeUntilNavMeshRebuild = TileMapEditorConstants::IntervalBetweenNavMeshRebuilds;
+			FlushPendingDirtyRegions();
+		}
+	}
+}
+
 void FEdModeTileMap::ActorSelectionChangeNotify()
 {
 	if (FindSelectedComponent() == nullptr)
@@ -604,10 +702,12 @@ bool FEdModeTileMap::UseActiveToolAtLocation(const FViewportCursorLocation& Ray)
 	};
 }
 
-bool BlitLayer(UPaperTileLayer* SourceLayer, UPaperTileLayer* TargetLayer, int32 OffsetX = 0, int32 OffsetY = 0, bool bBlitEmptyTiles = false)
+bool FEdModeTileMap::BlitLayer(UPaperTileLayer* SourceLayer, UPaperTileLayer* TargetLayer, FBox& OutDirtyRect, int32 OffsetX, int32 OffsetY, bool bBlitEmptyTiles)
 {
 	FScopedTransaction Transaction(LOCTEXT("TileMapPaintAction", "Tile Painting"));
 
+	const int32 LayerCoord = TargetLayer->GetLayerIndex();
+
 	bool bPaintedOnSomething = false;
 	bool bChangedSomething = false;
 
@@ -639,6 +739,8 @@ bool BlitLayer(UPaperTileLayer* SourceLayer, UPaperTileLayer* TargetLayer, int32
 					TargetLayer->Modify();
 					bChangedSomething = true;
 				}
+
+				OutDirtyRect += FVector(TargetX, TargetY, LayerCoord);
 				TargetLayer->SetCell(TargetX, TargetY, Ink);
 			}
 
@@ -706,7 +808,13 @@ bool FEdModeTileMap::PaintTiles(const FViewportCursorLocation& Ray)
 
 	if (UPaperTileLayer* TargetLayer = GetSelectedLayerUnderCursor(Ray, /*out*/ DestTileX, /*out*/ DestTileY))
 	{
-		bPaintedOnSomething = BlitLayer(GetSourceInkLayer(), TargetLayer, DestTileX, DestTileY);
+		FBox DirtyRect(ForceInitToZero);
+		bPaintedOnSomething = BlitLayer(GetSourceInkLayer(), TargetLayer, /*out*/ DirtyRect, DestTileX, DestTileY);
+
+		if (DirtyRect.IsValid)
+		{
+			new (PendingDirtyRegions) FTileMapDirtyRegion(FindSelectedComponent(), DirtyRect);
+		}
 	}
 
 	return bPaintedOnSomething;
@@ -716,6 +824,8 @@ bool FEdModeTileMap::EraseTiles(const FViewportCursorLocation& Ray)
 {
 	bool bPaintedOnSomething = false;
 	bool bChangedSomething = false;
+	FBox DirtyRect(ForceInitToZero);
+
 	const int32 BrushWidth = GetBrushWidth();
 	const int32 BrushHeight = GetBrushHeight();
 
@@ -727,6 +837,7 @@ bool FEdModeTileMap::EraseTiles(const FViewportCursorLocation& Ray)
 	if (UPaperTileLayer* Layer = GetSelectedLayerUnderCursor(Ray, /*out*/ DestTileX, /*out*/ DestTileY))
 	{
 		UPaperTileMap* TileMap = Layer->GetTileMap();
+		const int32 LayerCoord = Layer->GetLayerIndex();
 
 		FScopedTransaction Transaction( LOCTEXT("TileMapEraseAction", "Tile Erasing") );
 
@@ -748,7 +859,7 @@ bool FEdModeTileMap::EraseTiles(const FViewportCursorLocation& Ray)
 					continue;
 				}
 
-				if (Layer->GetCell(DX, DY) != EmptyCellValue)
+				if (Layer->GetCell(DX, DY).IsValid())
 				{
 					if (!bChangedSomething)
 					{
@@ -757,6 +868,7 @@ bool FEdModeTileMap::EraseTiles(const FViewportCursorLocation& Ray)
 						bChangedSomething = true;
 					}
 					Layer->SetCell(DX, DY, EmptyCellValue);
+					DirtyRect += FVector(DX, DY, LayerCoord);
 				}
 
 				bPaintedOnSomething = true;
@@ -765,6 +877,12 @@ bool FEdModeTileMap::EraseTiles(const FViewportCursorLocation& Ray)
 
 		if (bChangedSomething)
 		{
+			if (DirtyRect.IsValid)
+			{
+				UPaperTileMapComponent* TileMapComponent = FindSelectedComponent();
+				new (PendingDirtyRegions) FTileMapDirtyRegion(TileMapComponent, DirtyRect);
+			}
+
 			TileMap->PostEditChange();
 		}
 
@@ -799,8 +917,11 @@ bool FEdModeTileMap::FloodFillTiles(const FViewportCursorLocation& Ray)
 			return false;
 		}
 
+		FBox DirtyRect(ForceInitToZero);
+
 		// The kind of ink we'll replace, starting at the seed point
 		const FPaperTileInfo RequiredInk = TargetLayer->GetCell(DestTileX, DestTileY);
+		const int32 LayerIndex = TargetLayer->GetLayerIndex();
 
 		UPaperTileMap* TileMap = TargetLayer->GetTileMap();
 
@@ -875,6 +996,9 @@ bool FEdModeTileMap::FloodFillTiles(const FViewportCursorLocation& Ray)
 							TargetLayer->Modify();
 							bChangedSomething = true;
 						}
+
+						DirtyRect += FVector(DX, DY, LayerIndex);
+
 						TargetLayer->SetCell(DX, DY, NewInk);
 					}
 
@@ -885,6 +1009,11 @@ bool FEdModeTileMap::FloodFillTiles(const FViewportCursorLocation& Ray)
 
 		if (bChangedSomething)
 		{
+			if (DirtyRect.IsValid)
+			{
+				new (PendingDirtyRegions) FTileMapDirtyRegion(FindSelectedComponent(), DirtyRect);
+			}
+
 			TileMap->PostEditChange();
 		}
 
diff --git a/Engine/Plugins/2D/Paper2D/Source/Paper2DEditor/Private/TileMapEditing/EdModeTileMap.h b/Engine/Plugins/2D/Paper2D/Source/Paper2DEditor/Private/TileMapEditing/EdModeTileMap.h
--- a/Engine/Plugins/2D/Paper2D/Source/Paper2DEditor/Private/TileMapEditing/EdModeTileMap.h
+++ b/Engine/Plugins/2D/Paper2D/Source/Paper2DEditor/Private/TileMapEditing/EdModeTileMap.h
@@ -14,6 +14,30 @@ namespace ETileMapEditorTool
 	};
 }
 
+//////////////////////////////////////////////////////////////////////////
+// FTileMapDirtyRegion
+
+struct FTileMapDirtyRegion
+{
+public:
+	TWeakObjectPtr<UPaperTileMapComponent> ComponentPtr;
+	FBox DirtyRegionInWorldSpace;
+
+public:
+	FTileMapDirtyRegion()
+	{
+	}
+
+	FTileMapDirtyRegion(UPaperTileMapComponent* InComponent, const FBox& DirtyRegionInTileSpace);
+
+	void PushToNavSystem() const;
+	
+	UPaperTileMapComponent* GetComponent() const
+	{
+		return ComponentPtr.Get();
+	}
+};
+
 //////////////////////////////////////////////////////////////////////////
 // FEdModeTileMap
 
@@ -34,6 +58,7 @@ public:
 	virtual bool UsesToolkits() const override;
 	virtual void Enter() override;
 	virtual void Exit() override;
+	virtual void Tick(FEditorViewportClient* ViewportClient, float DeltaTime) override;
 	virtual bool MouseEnter(FEditorViewportClient* ViewportClient, FViewport* Viewport, int32 x, int32 y) override;
 	virtual bool MouseLeave(FEditorViewportClient* ViewportClient, FViewport* Viewport) override;
 	virtual bool MouseMove(FEditorViewportClient* ViewportClient, FViewport* Viewport, int32 x, int32 y) override;
@@ -94,6 +119,10 @@ protected:
 	void RotateTilesInSelection(bool bIsClockwise);
 
 	bool IsToolReadyToBeUsed() const;
+
+	bool BlitLayer(UPaperTileLayer* SourceLayer, UPaperTileLayer* TargetLayer, FBox& OutDirtyRect, int32 OffsetX = 0, int32 OffsetY = 0, bool bBlitEmptyTiles = false);
+
+	void FlushPendingDirtyRegions();
 public:
 	UPaperTileMapComponent* FindSelectedComponent() const;
 
@@ -141,6 +170,11 @@ protected:
 
 	UPaperTileMapComponent* CursorPreviewComponent;
 
+	// Nav mesh rebuilding
+	float TimeUntilNavMeshRebuild;
+	TArray<FTileMapDirtyRegion> PendingDirtyRegions;
+
+	//
 	ETileMapEditorTool::Type ActiveTool;
 	mutable FTransform ComponentToWorld;
 };
