Flipbook Editor: Improved the styling and behavior of the flipbook editor timeline
- Added a scroll bar so timelines longer than the width of the window can still be edited
- Added an alternating background per-frame
- Changed style of keyframe strips so that they stretch across the full duration of frames, with the grab-handle overlaid on top
[UE-13178] [UE-13179]

diff --git a/Engine/Plugins/2D/Paper2D/Source/Paper2DEditor/Private/FlipbookEditor/SFlipbookTimeline.cpp b/Engine/Plugins/2D/Paper2D/Source/Paper2DEditor/Private/FlipbookEditor/SFlipbookTimeline.cpp
--- a/Engine/Plugins/2D/Paper2D/Source/Paper2DEditor/Private/FlipbookEditor/SFlipbookTimeline.cpp
+++ b/Engine/Plugins/2D/Paper2D/Source/Paper2DEditor/Private/FlipbookEditor/SFlipbookTimeline.cpp
@@ -28,52 +28,78 @@ void SFlipbookTimeline::Construct(const FArguments& InArgs, TSharedPtr<FUIComman
 
 	SlateUnitsPerFrame = 32;
 
+	BackgroundPerFrameSlices = SNew(SHorizontalBox);
+
+	TimelineHeader = SNew(STimelineHeader)
+		.SlateUnitsPerFrame(SlateUnitsPerFrame)
+		.FlipbookBeingEdited(FlipbookBeingEdited)
+		.PlayTime(PlayTime);
+
+	TimelineTrack = SNew(SFlipbookTimelineTrack, CommandList)
+		.SlateUnitsPerFrame(SlateUnitsPerFrame)
+		.FlipbookBeingEdited(FlipbookBeingEdited)
+		.OnSelectionChanged(OnSelectionChanged);
+
+
 	ChildSlot
 	[
 		SNew(SBorder)
 		.BorderImage( FEditorStyle::GetBrush("ToolPanel.GroupBorder") )
 		[
-			SNew(SHorizontalBox)
-
-			// Empty flipbook instructions
-			+SHorizontalBox::Slot()
-			.VAlign(VAlign_Center)
-			.HAlign(HAlign_Center)
+			SNew(SScrollBox)
+			.Orientation(Orient_Horizontal)
+			.ScrollBarAlwaysVisible(true)
+			+SScrollBox::Slot()
 			[
-				SNew(STextBlock)
-				.Visibility(this, &SFlipbookTimeline::NoFramesWarningVisibility)
-				.Text(LOCTEXT("EmptyTimelineInstruction", "Right-click here or drop in sprites to add key frames"))
-			]
+				SNew(SOverlay)
 
-			// Flipbook header and track
-			+SHorizontalBox::Slot()
-			[
-				SNew(SVerticalBox)
-		
-				+SVerticalBox::Slot()
-				.AutoHeight()
-				.Padding(0,0,0,2)
+				// Per-frame background
+				+SOverlay::Slot()
+				.VAlign(VAlign_Fill)
 				[
-					SNew(STimelineHeader)
-					.SlateUnitsPerFrame(SlateUnitsPerFrame)
-					.FlipbookBeingEdited(FlipbookBeingEdited)
-					.PlayTime(PlayTime)
+					BackgroundPerFrameSlices.ToSharedRef()
 				]
 
-				+SVerticalBox::Slot()
-				.AutoHeight()
+				// Flipbook header and track
+				+SOverlay::Slot()
 				[
-					SNew(SBox).HeightOverride(FFlipbookUIConstants::FrameHeight)
+					SNew(SVerticalBox)
+		
+					+SVerticalBox::Slot()
+					.AutoHeight()
+					.Padding(0,0,0,2)
 					[
-						SNew(SFlipbookTimelineTrack, InCommandList)
-						.SlateUnitsPerFrame(SlateUnitsPerFrame)
-						.FlipbookBeingEdited(FlipbookBeingEdited)
-						.OnSelectionChanged(OnSelectionChanged)
+						TimelineHeader.ToSharedRef()
 					]
+
+					+SVerticalBox::Slot()
+					.AutoHeight()
+					[
+						SNew(SBox)
+						.HeightOverride(FFlipbookUIConstants::FrameHeight)
+						[
+							TimelineTrack.ToSharedRef()
+						]
+					]
+				]
+
+				// Empty flipbook instructions
+				+ SOverlay::Slot()
+				.VAlign(VAlign_Center)
+				.HAlign(HAlign_Center)
+				[
+					SNew(STextBlock)
+					.Visibility(this, &SFlipbookTimeline::NoFramesWarningVisibility)
+					.Text(LOCTEXT("EmptyTimelineInstruction", "Right-click here or drop in sprites to add key frames"))
 				]
 			]
 		]
 	];
+
+	UPaperFlipbook* Flipbook = FlipbookBeingEdited.Get();
+	NumKeyFramesFromLastRebuild = (Flipbook != nullptr) ? Flipbook->GetNumKeyFrames() : 0;
+	NumFramesFromLastRebuild = (Flipbook != nullptr) ? Flipbook->GetNumFrames() : 0;
+	RebuildPerFrameBG();
 }
 
 void SFlipbookTimeline::OnDragEnter(const FGeometry& MyGeometry, const FDragDropEvent& DragDropEvent)
@@ -232,6 +258,26 @@ FReply SFlipbookTimeline::OnMouseButtonUp(const FGeometry& MyGeometry, const FPo
 	}
 }
 
+void SFlipbookTimeline::Tick(const FGeometry& AllottedGeometry, const double InCurrentTime, const float InDeltaTime)
+{
+	UPaperFlipbook* Flipbook = FlipbookBeingEdited.Get();
+
+	const int32 NewNumKeyframes = (Flipbook != nullptr) ? Flipbook->GetNumKeyFrames() : 0;
+	if (NewNumKeyframes != NumKeyFramesFromLastRebuild)
+	{
+		NumKeyFramesFromLastRebuild = NewNumKeyframes;
+		TimelineTrack->Rebuild();
+	}
+
+	const int32 NewNumFrames = (Flipbook != nullptr) ? Flipbook->GetNumFrames() : 0;
+	if (NewNumFrames != NumFramesFromLastRebuild)
+	{
+		NumFramesFromLastRebuild = NewNumFrames;
+		TimelineHeader->Rebuild();
+		RebuildPerFrameBG();
+	}
+}
+
 EVisibility SFlipbookTimeline::NoFramesWarningVisibility() const
 {
 	UPaperFlipbook* Flipbook = FlipbookBeingEdited.Get();
@@ -239,6 +285,28 @@ EVisibility SFlipbookTimeline::NoFramesWarningVisibility() const
 	return (TotalNumFrames == 0) ? EVisibility::Visible : EVisibility::Collapsed;
 }
 
+void SFlipbookTimeline::RebuildPerFrameBG()
+{
+	const FLinearColor BackgroundColors[2] = { FLinearColor(1.0f, 1.0f, 1.0f, 0.05f), FLinearColor(0.0f, 0.0f, 0.0f, 0.05f) };
+
+	BackgroundPerFrameSlices->ClearChildren();
+	for (int32 FrameIndex = 0; FrameIndex < NumFramesFromLastRebuild; ++FrameIndex)
+	{
+		const FLinearColor& BackgroundColorForFrameIndex = BackgroundColors[FrameIndex & 1];
+
+		BackgroundPerFrameSlices->AddSlot()
+		.AutoWidth()
+		[
+			SNew(SBox)
+			.WidthOverride(SlateUnitsPerFrame)
+			[
+				SNew(SColorBlock)
+				.Color(BackgroundColorForFrameIndex)
+			]
+		];
+	}
+}
+
 //////////////////////////////////////////////////////////////////////////
 
 #undef LOCTEXT_NAMESPACE
\ No newline at end of file
diff --git a/Engine/Plugins/2D/Paper2D/Source/Paper2DEditor/Private/FlipbookEditor/SFlipbookTimeline.h b/Engine/Plugins/2D/Paper2D/Source/Paper2DEditor/Private/FlipbookEditor/SFlipbookTimeline.h
--- a/Engine/Plugins/2D/Paper2D/Source/Paper2DEditor/Private/FlipbookEditor/SFlipbookTimeline.h
+++ b/Engine/Plugins/2D/Paper2D/Source/Paper2DEditor/Private/FlipbookEditor/SFlipbookTimeline.h
@@ -28,17 +28,25 @@ public:
 	virtual FReply OnDrop(const FGeometry& MyGeometry, const FDragDropEvent& DragDropEvent) override;
 	virtual int32 OnPaint(const FPaintArgs& Args, const FGeometry& AllottedGeometry, const FSlateRect& MyClippingRect, FSlateWindowElementList& OutDrawElements, int32 LayerId, const FWidgetStyle& InWidgetStyle, bool bParentEnabled) const override;
 	virtual FReply OnMouseButtonUp(const FGeometry& MyGeometry, const FPointerEvent& MouseEvent) override;
+	virtual void Tick(const FGeometry& AllottedGeometry, const double InCurrentTime, const float InDeltaTime) override;
 	// End of SWidget interface
 
 private:
 	void OnAssetsDropped(const class FAssetDragDropOp& DragDropOp);
+	void RebuildPerFrameBG();
 
 	TSharedRef<SWidget> GenerateContextMenu();
 	EVisibility NoFramesWarningVisibility() const;
 private:
+	TSharedPtr<SHorizontalBox> BackgroundPerFrameSlices;
+	TSharedPtr<class STimelineHeader> TimelineHeader;
+	TSharedPtr<class SFlipbookTimelineTrack> TimelineTrack;
+
 	TAttribute<class UPaperFlipbook*> FlipbookBeingEdited;
 	TAttribute<float> PlayTime;
 	TSharedPtr<FUICommandList> CommandList;
 	FOnFlipbookKeyframeSelectionChanged OnSelectionChanged;
 	int32 SlateUnitsPerFrame;
+	int32 NumFramesFromLastRebuild;
+	int32 NumKeyFramesFromLastRebuild;
 };
diff --git a/Engine/Plugins/2D/Paper2D/Source/Paper2DEditor/Private/FlipbookEditor/STimelineHeader.h b/Engine/Plugins/2D/Paper2D/Source/Paper2DEditor/Private/FlipbookEditor/STimelineHeader.h
--- a/Engine/Plugins/2D/Paper2D/Source/Paper2DEditor/Private/FlipbookEditor/STimelineHeader.h
+++ b/Engine/Plugins/2D/Paper2D/Source/Paper2DEditor/Private/FlipbookEditor/STimelineHeader.h
@@ -36,17 +36,6 @@ public:
 		Rebuild();
 	}
 
-	virtual void Tick( const FGeometry& AllottedGeometry, const double InCurrentTime, const float InDeltaTime ) override
-	{
-		UPaperFlipbook* Flipbook = FlipbookBeingEdited.Get();
-		int32 NewNumFrames = (Flipbook != nullptr) ? Flipbook->GetNumFrames() : 0;
-		if (NewNumFrames != NumFramesFromLastRebuild)
-		{
-			Rebuild();
-		}
-	}
-
-private:
 	void Rebuild()
 	{
 		MainBoxPtr->ClearChildren();
@@ -56,7 +45,7 @@ private:
 		if ((Flipbook != nullptr) && (LocalSlateUnitsPerFrame > 0))
 		{
 			const int32 NumFrames = Flipbook->GetNumFrames();
-			for (int32 FrameIdx = 0; FrameIdx < NumFrames; ++FrameIdx)
+			for (int32 FrameIndex = 0; FrameIndex < NumFrames; ++FrameIndex)
 			{
 				MainBoxPtr->AddSlot()
 					.AutoWidth()
@@ -66,7 +55,7 @@ private:
 						.HAlign(HAlign_Center)
 						[
 							SNew(STextBlock)
-							.Text(FText::AsNumber(FrameIdx))
+							.Text(FText::AsNumber(FrameIndex))
 						]
 					];
 			}
diff --git a/Engine/Plugins/2D/Paper2D/Source/Paper2DEditor/Private/FlipbookEditor/STimelineTrack.cpp b/Engine/Plugins/2D/Paper2D/Source/Paper2DEditor/Private/FlipbookEditor/STimelineTrack.cpp
--- a/Engine/Plugins/2D/Paper2D/Source/Paper2DEditor/Private/FlipbookEditor/STimelineTrack.cpp
+++ b/Engine/Plugins/2D/Paper2D/Source/Paper2DEditor/Private/FlipbookEditor/STimelineTrack.cpp
@@ -79,38 +79,37 @@ void SFlipbookKeyframeWidget::Construct(const FArguments& InArgs, int32 InFrameI
 	};
 
 	ChildSlot
+	[
+		SNew(SOverlay)
+		+SOverlay::Slot()
 		[
-			SNew(SHorizontalBox)
-			+ SHorizontalBox::Slot()
-			.AutoWidth()
+			SNew(SBox)
+			.Padding(FFlipbookUIConstants::FramePadding)
+			.WidthOverride(this, &SFlipbookKeyframeWidget::GetFrameWidth)
 			[
-				SNew(SBox)
-				.Padding(FFlipbookUIConstants::FramePadding)
-				.WidthOverride(this, &SFlipbookKeyframeWidget::GetFrameWidth)
+				SNew(SBorder)
+				.BorderImage(FPaperStyle::Get()->GetBrush("FlipbookEditor.RegionBody"))
+				.BorderBackgroundColor_Static(BorderColorDelegate, FlipbookBeingEdited, FrameIndex)
+				.OnMouseButtonUp(this, &SFlipbookKeyframeWidget::KeyframeOnMouseButtonUp)
+				.ToolTipText(this, &SFlipbookKeyframeWidget::GetKeyframeTooltip)
 				[
-					SNew(SBorder)
-					.BorderImage(FPaperStyle::Get()->GetBrush("FlipbookEditor.RegionBody"))
-					.BorderBackgroundColor_Static(BorderColorDelegate, FlipbookBeingEdited, FrameIndex)
-					.OnMouseButtonUp(this, &SFlipbookKeyframeWidget::KeyframeOnMouseButtonUp)
-					.ToolTipText(this, &SFlipbookKeyframeWidget::GetKeyframeTooltip)
-					[
-						SNullWidget::NullWidget
-					]
+					SNullWidget::NullWidget
 				]
 			]
-			+ SHorizontalBox::Slot()
-				.AutoWidth()
-				[
-					SNew(SBox)
-					.WidthOverride(FFlipbookUIConstants::HandleWidth)
-					[
-						SNew(SFlipbookTrackHandle)
-						.SlateUnitsPerFrame(SlateUnitsPerFrame)
-						.FlipbookBeingEdited(FlipbookBeingEdited)
-						.KeyFrameIdx(FrameIndex)
-					]
-				]
-		];
+		]
+		+SOverlay::Slot()
+		.HAlign(HAlign_Right)
+		[
+			SNew(SBox)
+			.WidthOverride(FFlipbookUIConstants::HandleWidth)
+			[
+				SNew(SFlipbookTrackHandle)
+				.SlateUnitsPerFrame(SlateUnitsPerFrame)
+				.FlipbookBeingEdited(FlipbookBeingEdited)
+				.KeyFrameIdx(FrameIndex)
+			]
+		]
+	];
 }
 
 FReply SFlipbookKeyframeWidget::OnMouseButtonDown(const FGeometry& MyGeometry, const FPointerEvent& MouseEvent)
@@ -215,10 +214,13 @@ FText SFlipbookKeyframeWidget::GetKeyframeTooltip() const
 	{
 		const FText SpriteLine = (KeyFrame->Sprite != nullptr) ? FText::FromString(KeyFrame->Sprite->GetName()) : LOCTEXT("NoSprite", "(none)");
 
-		return FText::Format(LOCTEXT("KeyFrameTooltip", "Sprite: {0}\nIndex: {1}\nDuration: {2} frame(s)"),
+		const FText FramesText = (KeyFrame->FrameRun == 1) ? LOCTEXT("SingularFrames", "frame") : LOCTEXT("PluralFrames", "frames");
+		
+		return FText::Format(LOCTEXT("KeyFrameTooltip", "Sprite: {0}\nIndex: {1}\nDuration: {2} {3}"),
 			SpriteLine,
 			FText::AsNumber(FrameIndex),
-			FText::AsNumber(KeyFrame->FrameRun));
+			FText::AsNumber(KeyFrame->FrameRun),
+			FramesText);
 	}
 	else
 	{
@@ -230,7 +232,7 @@ FOptionalSize SFlipbookKeyframeWidget::GetFrameWidth() const
 {
 	if (const FPaperFlipbookKeyFrame* KeyFrame = GetKeyFrameData())
 	{
-		return FMath::Max<float>(0, KeyFrame->FrameRun * SlateUnitsPerFrame.Get() - FFlipbookUIConstants::HandleWidth);
+		return FMath::Max<float>(0, KeyFrame->FrameRun * SlateUnitsPerFrame.Get()/* - FFlipbookUIConstants::HandleWidth*/);
 	}
 	else
 	{
@@ -399,12 +401,10 @@ void SFlipbookTimelineTrack::Construct(const FArguments& InArgs, TSharedPtr<FUIC
 	FlipbookBeingEdited = InArgs._FlipbookBeingEdited;
 	OnSelectionChanged = InArgs._OnSelectionChanged;
 
-	NumKeyframesFromLastRebuild = 0;
-
 	ChildSlot
-		[
-			SAssignNew(MainBoxPtr, SHorizontalBox)
-		];
+	[
+		SAssignNew(MainBoxPtr, SHorizontalBox)
+	];
 
 	Rebuild();
 }
@@ -418,22 +418,15 @@ void SFlipbookTimelineTrack::Rebuild()
 	{
 		for (int32 KeyFrameIdx = 0; KeyFrameIdx < Flipbook->GetNumKeyFrames(); ++KeyFrameIdx)
 		{
-			//@TODO: Draggy bits go here
 			MainBoxPtr->AddSlot()
-				.AutoWidth()
-				[
-					SNew(SFlipbookKeyframeWidget, KeyFrameIdx, CommandList)
-					.SlateUnitsPerFrame(this->SlateUnitsPerFrame)
-					.FlipbookBeingEdited(this->FlipbookBeingEdited)
-					.OnSelectionChanged(this->OnSelectionChanged)
-				];
+			.AutoWidth()
+			[
+				SNew(SFlipbookKeyframeWidget, KeyFrameIdx, CommandList)
+				.SlateUnitsPerFrame(this->SlateUnitsPerFrame)
+				.FlipbookBeingEdited(this->FlipbookBeingEdited)
+				.OnSelectionChanged(this->OnSelectionChanged)
+			];
 		}
-
-		NumKeyframesFromLastRebuild = Flipbook->GetNumKeyFrames();
-	}
-	else
-	{
-		NumKeyframesFromLastRebuild = 0;
 	}
 }
 
diff --git a/Engine/Plugins/2D/Paper2D/Source/Paper2DEditor/Private/FlipbookEditor/STimelineTrack.h b/Engine/Plugins/2D/Paper2D/Source/Paper2DEditor/Private/FlipbookEditor/STimelineTrack.h
--- a/Engine/Plugins/2D/Paper2D/Source/Paper2DEditor/Private/FlipbookEditor/STimelineTrack.h
+++ b/Engine/Plugins/2D/Paper2D/Source/Paper2DEditor/Private/FlipbookEditor/STimelineTrack.h
@@ -9,6 +9,7 @@ namespace FFlipbookUIConstants
 {
 	const float HandleWidth = 12.0f;
 	const float FrameHeight = 48;
+	const float HeightBeforeFrames = 16;
 	const FMargin FramePadding(0.0f, 7.0f, 0.0f, 7.0f);
 };
 
@@ -120,17 +121,6 @@ public:
 
 	void Construct(const FArguments& InArgs, TSharedPtr<FUICommandList> InCommandList);
 
-	virtual void Tick( const FGeometry& AllottedGeometry, const double InCurrentTime, const float InDeltaTime ) override
-	{
-		UPaperFlipbook* Flipbook = FlipbookBeingEdited.Get();
-		int32 NewNumKeyframes = (Flipbook != nullptr) ? Flipbook->GetNumKeyFrames() : 0;
-		if (NewNumKeyframes != NumKeyframesFromLastRebuild)
-		{
-			Rebuild();
-		}
-	}
-
-private:
 	void Rebuild();
 
 private:
@@ -139,7 +129,6 @@ private:
 
 	TSharedPtr<SHorizontalBox> MainBoxPtr;
 
-	int32 NumKeyframesFromLastRebuild;
 	float HandleWidth;
 
 	FOnFlipbookKeyframeSelectionChanged OnSelectionChanged;
