Flipbook Editor: Added display of the sprite name in keyframe strips and when drag-dropping

diff --git a/Engine/Plugins/2D/Paper2D/Source/Paper2DEditor/Private/FlipbookEditor/STimelineTrack.cpp b/Engine/Plugins/2D/Paper2D/Source/Paper2DEditor/Private/FlipbookEditor/STimelineTrack.cpp
--- a/Engine/Plugins/2D/Paper2D/Source/Paper2DEditor/Private/FlipbookEditor/STimelineTrack.cpp
+++ b/Engine/Plugins/2D/Paper2D/Source/Paper2DEditor/Private/FlipbookEditor/STimelineTrack.cpp
@@ -92,8 +92,12 @@ void SFlipbookKeyframeWidget::Construct(const FArguments& InArgs, int32 InFrameI
 				.BorderBackgroundColor_Static(BorderColorDelegate, FlipbookBeingEdited, FrameIndex)
 				.OnMouseButtonUp(this, &SFlipbookKeyframeWidget::KeyframeOnMouseButtonUp)
 				.ToolTipText(this, &SFlipbookKeyframeWidget::GetKeyframeTooltip)
+				.HAlign(HAlign_Center)
+				.VAlign(VAlign_Center)
 				[
-					SNullWidget::NullWidget
+					SNew(STextBlock)
+					.ColorAndOpacity(FLinearColor::Black)
+					.Text(this, &SFlipbookKeyframeWidget::GetKeyframeText)
 				]
 			]
 		]
@@ -208,6 +212,19 @@ FText SFlipbookKeyframeWidget::GetKeyframeAssetName() const
 	}
 }
 
+FText SFlipbookKeyframeWidget::GetKeyframeText() const
+{
+	if (const FPaperFlipbookKeyFrame* KeyFrame = GetKeyFrameData())
+	{
+		if (KeyFrame->Sprite != nullptr)
+		{
+			return FText::AsCultureInvariant(KeyFrame->Sprite->GetName());
+		}
+	}
+
+	return FText::GetEmpty();
+}
+
 FText SFlipbookKeyframeWidget::GetKeyframeTooltip() const
 {
 	if (const FPaperFlipbookKeyFrame* KeyFrame = GetKeyFrameData())
@@ -232,7 +249,7 @@ FOptionalSize SFlipbookKeyframeWidget::GetFrameWidth() const
 {
 	if (const FPaperFlipbookKeyFrame* KeyFrame = GetKeyFrameData())
 	{
-		return FMath::Max<float>(0, KeyFrame->FrameRun * SlateUnitsPerFrame.Get()/* - FFlipbookUIConstants::HandleWidth*/);
+		return FMath::Max<float>(0, KeyFrame->FrameRun * SlateUnitsPerFrame.Get());
 	}
 	else
 	{
@@ -322,8 +339,12 @@ TSharedPtr<SWidget> FFlipbookKeyFrameDragDropOp::GetDefaultDecorator() const
 			SNew(SBorder)
 			.BorderImage(FPaperStyle::Get()->GetBrush("FlipbookEditor.RegionBody"))
 			.BorderBackgroundColor(BorderColor)
+			.HAlign(HAlign_Center)
+			.VAlign(VAlign_Center)
 			[
-				SNullWidget::NullWidget
+				SNew(STextBlock)
+				.ColorAndOpacity(FLinearColor::Black)
+				.Text(BodyText)
 			]
 		];
 }
@@ -339,6 +360,15 @@ void FFlipbookKeyFrameDragDropOp::OnDragged(const class FDragDropEvent& DragDrop
 void FFlipbookKeyFrameDragDropOp::Construct()
 {
 	MouseCursor = EMouseCursor::GrabHandClosed;
+
+	if (UPaperFlipbook* Flipbook = SourceFlipbook.Get())
+	{
+		if (UPaperSprite* Sprite = Flipbook->GetSpriteAtFrame(SourceFrameIndex))
+		{
+			BodyText = FText::AsCultureInvariant(Sprite->GetName());
+		}
+	}
+
 	FDragDropOperation::Construct();
 }
 
@@ -386,7 +416,8 @@ TSharedRef<FFlipbookKeyFrameDragDropOp> FFlipbookKeyFrameDragDropOp::New(int32 I
 	return Operation;
 }
 
-FFlipbookKeyFrameDragDropOp::FFlipbookKeyFrameDragDropOp() : Transaction(LOCTEXT("MovedFramesInTimeline", "Reorder key frames"))
+FFlipbookKeyFrameDragDropOp::FFlipbookKeyFrameDragDropOp()
+	: Transaction(LOCTEXT("MovedFramesInTimeline", "Reorder key frames"))
 {
 
 }
diff --git a/Engine/Plugins/2D/Paper2D/Source/Paper2DEditor/Private/FlipbookEditor/STimelineTrack.h b/Engine/Plugins/2D/Paper2D/Source/Paper2DEditor/Private/FlipbookEditor/STimelineTrack.h
--- a/Engine/Plugins/2D/Paper2D/Source/Paper2DEditor/Private/FlipbookEditor/STimelineTrack.h
+++ b/Engine/Plugins/2D/Paper2D/Source/Paper2DEditor/Private/FlipbookEditor/STimelineTrack.h
@@ -24,16 +24,16 @@ public:
 	float WidgetWidth;
 	FPaperFlipbookKeyFrame KeyFrameData;
 	int32 SourceFrameIndex;
+	FText BodyText;
 	TWeakObjectPtr<UPaperFlipbook> SourceFlipbook;
 	FScopedTransaction Transaction;
 
+	// FDragDropOperation interface
 	virtual TSharedPtr<SWidget> GetDefaultDecorator() const override;
-
 	virtual void OnDragged(const class FDragDropEvent& DragDropEvent) override;
-
 	virtual void Construct() override;
-
 	virtual void OnDrop(bool bDropWasHandled, const FPointerEvent& MouseEvent) override;
+	// End of FDragDropOperation interface
 
 	void AppendToFlipbook(UPaperFlipbook* DestinationFlipbook);
 
@@ -79,6 +79,7 @@ protected:
 	FReply KeyframeOnMouseButtonUp(const FGeometry& MyGeometry, const FPointerEvent& MouseEvent);
 
 	FText GetKeyframeAssetName() const;
+	FText GetKeyframeText() const;
 	FText GetKeyframeTooltip() const;
 
 	TSharedRef<SWidget> GenerateContextMenu();
