Flipbook Editor: Improved the help text

diff --git a/Engine/Plugins/2D/Paper2D/Source/Paper2DEditor/Private/FlipbookEditor/FlipbookEditorViewportClient.cpp b/Engine/Plugins/2D/Paper2D/Source/Paper2DEditor/Private/FlipbookEditor/FlipbookEditorViewportClient.cpp
--- a/Engine/Plugins/2D/Paper2D/Source/Paper2DEditor/Private/FlipbookEditor/FlipbookEditorViewportClient.cpp
+++ b/Engine/Plugins/2D/Paper2D/Source/Paper2DEditor/Private/FlipbookEditor/FlipbookEditorViewportClient.cpp
@@ -49,7 +49,7 @@ void FFlipbookEditorViewportClient::DrawCanvas(FViewport& Viewport, FSceneView&
 
 	int32 YPos = 42;
 
-	static const FText FlipbookHelpStr = LOCTEXT("FlipbookEditHelp", "Flipbook editor");
+	static const FText FlipbookHelpStr = LOCTEXT("FlipbookEditHelp", "Flipbook editor\n\nAdd keys using the toolbar or by drag-dropping sprite assets\nChange the timeline scale using Ctrl+MouseWheel\nEdit keys using the handles and right-click menu\nReorder keys by dragging and dropping");
 
 	// Display tool help
 	{
