Flipbook Editor: Added Ctrl+MouseWheel to change the time scale

diff --git a/Engine/Plugins/2D/Paper2D/Source/Paper2DEditor/Private/FlipbookEditor/SFlipbookTimeline.cpp b/Engine/Plugins/2D/Paper2D/Source/Paper2DEditor/Private/FlipbookEditor/SFlipbookTimeline.cpp
--- a/Engine/Plugins/2D/Paper2D/Source/Paper2DEditor/Private/FlipbookEditor/SFlipbookTimeline.cpp
+++ b/Engine/Plugins/2D/Paper2D/Source/Paper2DEditor/Private/FlipbookEditor/SFlipbookTimeline.cpp
@@ -26,21 +26,20 @@ void SFlipbookTimeline::Construct(const FArguments& InArgs, TSharedPtr<FUIComman
 	OnSelectionChanged = InArgs._OnSelectionChanged;
 	CommandList = InCommandList;
 
-	SlateUnitsPerFrame = 32;
+	SlateUnitsPerFrame = 64.0f;
 
 	BackgroundPerFrameSlices = SNew(SHorizontalBox);
 
 	TimelineHeader = SNew(STimelineHeader)
-		.SlateUnitsPerFrame(SlateUnitsPerFrame)
+		.SlateUnitsPerFrame(this, &SFlipbookTimeline::GetSlateUnitsPerFrame)
 		.FlipbookBeingEdited(FlipbookBeingEdited)
 		.PlayTime(PlayTime);
 
 	TimelineTrack = SNew(SFlipbookTimelineTrack, CommandList)
-		.SlateUnitsPerFrame(SlateUnitsPerFrame)
+		.SlateUnitsPerFrame(this, &SFlipbookTimeline::GetSlateUnitsPerFrame)
 		.FlipbookBeingEdited(FlipbookBeingEdited)
 		.OnSelectionChanged(OnSelectionChanged);
 
-
 	ChildSlot
 	[
 		SNew(SBorder)
@@ -228,6 +227,26 @@ int32 SFlipbookTimeline::OnPaint(const FPaintArgs& Args, const FGeometry& Allott
 	return LayerId;
 }
 
+FReply SFlipbookTimeline::OnMouseWheel(const FGeometry& MyGeometry, const FPointerEvent& MouseEvent)
+{
+	if (MouseEvent.IsControlDown())
+	{
+		const float DirectionScale = 0.08f;
+		const float MinFrameSize = 16.0f;
+		const float Direction = MouseEvent.GetWheelDelta();
+		const float NewUnitsPerFrame = FMath::Max(MinFrameSize, SlateUnitsPerFrame * (1.0f + Direction * DirectionScale));
+		SlateUnitsPerFrame = NewUnitsPerFrame;
+		
+		CheckForRebuild(/*bRebuildAll=*/ true);
+
+		return FReply::Handled();
+	}
+	else
+	{
+		return FReply::Unhandled();
+	}
+}
+
 TSharedRef<SWidget> SFlipbookTimeline::GenerateContextMenu()
 {
 	FMenuBuilder MenuBuilder(true, CommandList);
@@ -259,18 +278,23 @@ FReply SFlipbookTimeline::OnMouseButtonUp(const FGeometry& MyGeometry, const FPo
 }
 
 void SFlipbookTimeline::Tick(const FGeometry& AllottedGeometry, const double InCurrentTime, const float InDeltaTime)
+{
+	CheckForRebuild();
+}
+
+void SFlipbookTimeline::CheckForRebuild(bool bRebuildAll)
 {
 	UPaperFlipbook* Flipbook = FlipbookBeingEdited.Get();
 
 	const int32 NewNumKeyframes = (Flipbook != nullptr) ? Flipbook->GetNumKeyFrames() : 0;
-	if (NewNumKeyframes != NumKeyFramesFromLastRebuild)
+	if ((NewNumKeyframes != NumKeyFramesFromLastRebuild) || bRebuildAll)
 	{
 		NumKeyFramesFromLastRebuild = NewNumKeyframes;
 		TimelineTrack->Rebuild();
 	}
 
 	const int32 NewNumFrames = (Flipbook != nullptr) ? Flipbook->GetNumFrames() : 0;
-	if (NewNumFrames != NumFramesFromLastRebuild)
+	if ((NewNumFrames != NumFramesFromLastRebuild) || bRebuildAll)
 	{
 		NumFramesFromLastRebuild = NewNumFrames;
 		TimelineHeader->Rebuild();
diff --git a/Engine/Plugins/2D/Paper2D/Source/Paper2DEditor/Private/FlipbookEditor/SFlipbookTimeline.h b/Engine/Plugins/2D/Paper2D/Source/Paper2DEditor/Private/FlipbookEditor/SFlipbookTimeline.h
--- a/Engine/Plugins/2D/Paper2D/Source/Paper2DEditor/Private/FlipbookEditor/SFlipbookTimeline.h
+++ b/Engine/Plugins/2D/Paper2D/Source/Paper2DEditor/Private/FlipbookEditor/SFlipbookTimeline.h
@@ -27,6 +27,7 @@ public:
 	virtual void OnDragLeave(const FDragDropEvent& DragDropEvent) override;
 	virtual FReply OnDrop(const FGeometry& MyGeometry, const FDragDropEvent& DragDropEvent) override;
 	virtual int32 OnPaint(const FPaintArgs& Args, const FGeometry& AllottedGeometry, const FSlateRect& MyClippingRect, FSlateWindowElementList& OutDrawElements, int32 LayerId, const FWidgetStyle& InWidgetStyle, bool bParentEnabled) const override;
+	virtual FReply OnMouseWheel(const FGeometry& MyGeometry, const FPointerEvent& MouseEvent) override;
 	virtual FReply OnMouseButtonUp(const FGeometry& MyGeometry, const FPointerEvent& MouseEvent) override;
 	virtual void Tick(const FGeometry& AllottedGeometry, const double InCurrentTime, const float InDeltaTime) override;
 	// End of SWidget interface
@@ -37,6 +38,9 @@ private:
 
 	TSharedRef<SWidget> GenerateContextMenu();
 	EVisibility NoFramesWarningVisibility() const;
+	void CheckForRebuild(bool bRebuildAll = false);
+
+	float GetSlateUnitsPerFrame() const { return SlateUnitsPerFrame; }
 private:
 	TSharedPtr<SHorizontalBox> BackgroundPerFrameSlices;
 	TSharedPtr<class STimelineHeader> TimelineHeader;
@@ -46,7 +50,9 @@ private:
 	TAttribute<float> PlayTime;
 	TSharedPtr<FUICommandList> CommandList;
 	FOnFlipbookKeyframeSelectionChanged OnSelectionChanged;
-	int32 SlateUnitsPerFrame;
+	float SlateUnitsPerFrame;
+
+	// Observer values (used to detect if the widget needs to be rebuilt)
 	int32 NumFramesFromLastRebuild;
 	int32 NumKeyFramesFromLastRebuild;
 };
diff --git a/Engine/Plugins/2D/Paper2D/Source/Paper2DEditor/Private/FlipbookEditor/STimelineHeader.h b/Engine/Plugins/2D/Paper2D/Source/Paper2DEditor/Private/FlipbookEditor/STimelineHeader.h
--- a/Engine/Plugins/2D/Paper2D/Source/Paper2DEditor/Private/FlipbookEditor/STimelineHeader.h
+++ b/Engine/Plugins/2D/Paper2D/Source/Paper2DEditor/Private/FlipbookEditor/STimelineHeader.h
@@ -15,7 +15,7 @@ public:
 		, _PlayTime(0)
 	{}
 
-	SLATE_ATTRIBUTE(float, SlateUnitsPerFrame)
+		SLATE_ATTRIBUTE(float, SlateUnitsPerFrame)
 		SLATE_ATTRIBUTE(class UPaperFlipbook*, FlipbookBeingEdited)
 		SLATE_ATTRIBUTE(float, PlayTime)
 	SLATE_END_ARGS()
@@ -41,7 +41,7 @@ public:
 		MainBoxPtr->ClearChildren();
 
 		UPaperFlipbook* Flipbook = FlipbookBeingEdited.Get();
-		float LocalSlateUnitsPerFrame = SlateUnitsPerFrame.Get();
+		const float LocalSlateUnitsPerFrame = SlateUnitsPerFrame.Get();
 		if ((Flipbook != nullptr) && (LocalSlateUnitsPerFrame > 0))
 		{
 			const int32 NumFrames = Flipbook->GetNumFrames();
