Tile Map: Modified the "New" button to read "New Empty Map" to clarify it and reduce confusion versus the "Add Layer" button

diff --git a/Engine/Plugins/2D/Paper2D/Source/Paper2DEditor/Private/TileMapEditing/PaperTileMapDetailsCustomization.cpp b/Engine/Plugins/2D/Paper2D/Source/Paper2DEditor/Private/TileMapEditing/PaperTileMapDetailsCustomization.cpp
--- a/Engine/Plugins/2D/Paper2D/Source/Paper2DEditor/Private/TileMapEditing/PaperTileMapDetailsCustomization.cpp
+++ b/Engine/Plugins/2D/Paper2D/Source/Paper2DEditor/Private/TileMapEditing/PaperTileMapDetailsCustomization.cpp
@@ -74,7 +74,7 @@ void FPaperTileMapDetailsCustomization::CustomizeDetails(IDetailLayoutBuilder& D
 	if (TileComponent != nullptr)
 	{
 		TileMapCategory
-		.AddCustomRow(LOCTEXT( "TileMapInstancingControlsSearchText", "Edit New Promote Asset"))
+		.AddCustomRow(LOCTEXT( "TileMapInstancingControlsSearchText", "Edit Map New Empty Map Promote Asset"))
 		[
 			SNew(SVerticalBox)
 			+ SVerticalBox::Slot()
@@ -95,7 +95,7 @@ void FPaperTileMapDetailsCustomization::CustomizeDetails(IDetailLayoutBuilder& D
 					.VAlign(VAlign_Center)
 					.OnClicked(this, &FPaperTileMapDetailsCustomization::EnterTileMapEditingMode)
 					.Visibility(this, &FPaperTileMapDetailsCustomization::GetNonEditModeVisibility)
-					.Text( LOCTEXT("EditAsset", "Edit") )
+					.Text( LOCTEXT("EditAsset", "Edit Map") )
 					.ToolTipText( LOCTEXT("EditAssetToolTip", "Edit this tile map") )
 				]
 
@@ -110,8 +110,8 @@ void FPaperTileMapDetailsCustomization::CustomizeDetails(IDetailLayoutBuilder& D
 					.VAlign(VAlign_Center)
 					.OnClicked(this, &FPaperTileMapDetailsCustomization::OnNewButtonClicked)
 					.Visibility(this, &FPaperTileMapDetailsCustomization::GetNewButtonVisiblity)
-					.Text(LOCTEXT("CreateNewInstancedMap", "New"))
-					.ToolTipText( LOCTEXT("CreateNewInstancedMapToolTip", "Create a new tile map") )
+					.Text(LOCTEXT("CreateNewInstancedMap", "New Empty Map"))
+					.ToolTipText( LOCTEXT("CreateNewInstancedMapToolTip", "Create a new (instanced) tile map") )
 				]
 
 				// Promote to asset button
