Tile Maps: Added per-tile map instance and (non-instanced) per-layer editable colors [UE-13633]

diff --git a/Engine/Plugins/2D/Paper2D/Source/Paper2D/Classes/PaperTileLayer.h b/Engine/Plugins/2D/Paper2D/Source/Paper2D/Classes/PaperTileLayer.h
--- a/Engine/Plugins/2D/Paper2D/Source/Paper2D/Classes/PaperTileLayer.h
+++ b/Engine/Plugins/2D/Paper2D/Source/Paper2D/Classes/PaperTileLayer.h
@@ -130,6 +130,10 @@ class PAPER2D_API UPaperTileLayer : public UObject
 	bool bHiddenInGame;
 
 protected:
+	// The color of this layer (multiplied with the tile map color and passed to the material as a vertex color)
+	UPROPERTY(EditAnywhere, Category=Sprite)
+	FLinearColor LayerColor;
+
 	// The allocated width of the tile data (used to handle resizing without data loss)
 	UPROPERTY()
 	int32 AllocatedWidth;
@@ -178,6 +182,13 @@ public:
 
 	// Adds collision to the specified body setup
 	void AugmentBodySetup(UBodySetup* ShapeBodySetup);
+
+	// Gets the layer-specific color multiplier
+	FLinearColor GetLayerColor() const;
+
+	// Sets the layer-specific color multiplier (Note: does not invalidate any components using this layer!)
+	void SetLayerColor(FLinearColor NewColor);
+
 protected:
 	void ReallocateAndCopyMap();
 };
diff --git a/Engine/Plugins/2D/Paper2D/Source/Paper2D/Classes/PaperTileMapComponent.h b/Engine/Plugins/2D/Paper2D/Source/Paper2D/Classes/PaperTileMapComponent.h
--- a/Engine/Plugins/2D/Paper2D/Source/Paper2D/Classes/PaperTileMapComponent.h
+++ b/Engine/Plugins/2D/Paper2D/Source/Paper2D/Classes/PaperTileMapComponent.h
@@ -50,6 +50,10 @@ private:
 	UPROPERTY()
 	TArray<class UPaperTileLayer*> TileLayers_DEPRECATED;
 
+	// The color of the tile map (multiplied with the per-layer color and passed to the material as a vertex color)
+	UPROPERTY(EditAnywhere, Category=Materials)
+	FLinearColor TileMapColor;
+
 public:
 	// The tile map used by this component
 	UPROPERTY(Category=Setup, EditAnywhere, BlueprintReadOnly)
@@ -105,9 +109,27 @@ public:
 	UFUNCTION(BlueprintCallable, Category="Sprite")
 	void ResizeMap(int32 NewWidthInTiles, int32 NewHeightInTiles);
 
-	// Creates and adds a new layer to the tile map (Note: This will only work on components that own their own tile map(OwnsTileMap returns true), you cannot modify standalone tile map assets)
+	// Creates and adds a new layer to the tile map
+	// Note: This will only work on components that own their own tile map (OwnsTileMap returns true), you cannot modify standalone tile map assets
 	UFUNCTION(BlueprintCallable, Category="Sprite")
 	class UPaperTileLayer* AddNewLayer();
+
+	// Gets the tile map global color multiplier (multiplied with the per-layer color and passed to the material as a vertex color)
+	UFUNCTION(BlueprintPure, Category="Sprite")
+	FLinearColor GetTileMapColor() const;
+
+	// Sets the tile map global color multiplier (multiplied with the per-layer color and passed to the material as a vertex color)
+	UFUNCTION(BlueprintCallable, Category="Sprite")
+	void SetTileMapColor(FLinearColor NewColor);
+
+	// Gets the per-layer color multiplier for a specific layer (multiplied with the tile map color and passed to the material as a vertex color)
+	UFUNCTION(BlueprintPure, Category = "Sprite")
+	FLinearColor GetLayerColor(int32 Layer = 0) const;
+
+	// Sets the per-layer color multiplier for a specific layer (multiplied with the tile map color and passed to the material as a vertex color)
+	// Note: This will only work on components that own their own tile map (OwnsTileMap returns true), you cannot modify standalone tile map assets
+	UFUNCTION(BlueprintCallable, Category = "Sprite")
+	void SetLayerColor(FLinearColor NewColor, int32 Layer = 0);
 };
 
 // Allow the old name to continue to work for one release
diff --git a/Engine/Plugins/2D/Paper2D/Source/Paper2D/Private/PaperTileLayer.cpp b/Engine/Plugins/2D/Paper2D/Source/Paper2D/Private/PaperTileLayer.cpp
--- a/Engine/Plugins/2D/Paper2D/Source/Paper2D/Private/PaperTileLayer.cpp
+++ b/Engine/Plugins/2D/Paper2D/Source/Paper2D/Private/PaperTileLayer.cpp
@@ -93,6 +93,7 @@ protected:
 
 UPaperTileLayer::UPaperTileLayer(const FObjectInitializer& ObjectInitializer)
 	: Super(ObjectInitializer)
+	, LayerColor(FLinearColor::White)
 {
 	LayerWidth = 4;
 	LayerHeight = 4;
@@ -245,6 +246,16 @@ void UPaperTileLayer::AugmentBodySetup(UBodySetup* ShapeBodySetup)
 	}
 }
 
+FLinearColor UPaperTileLayer::GetLayerColor() const
+{
+	return LayerColor;
+}
+
+void UPaperTileLayer::SetLayerColor(FLinearColor NewColor)
+{
+	LayerColor = NewColor;
+}
+
 void UPaperTileLayer::ConvertToTileSetPerCell()
 {
 	AllocatedCells.Empty(AllocatedGrid_DEPRECATED.Num());
diff --git a/Engine/Plugins/2D/Paper2D/Source/Paper2D/Private/PaperTileMapComponent.cpp b/Engine/Plugins/2D/Paper2D/Source/Paper2D/Private/PaperTileMapComponent.cpp
--- a/Engine/Plugins/2D/Paper2D/Source/Paper2D/Private/PaperTileMapComponent.cpp
+++ b/Engine/Plugins/2D/Paper2D/Source/Paper2D/Private/PaperTileMapComponent.cpp
@@ -15,6 +15,7 @@
 
 UPaperTileMapComponent::UPaperTileMapComponent(const FObjectInitializer& ObjectInitializer)
 	: Super(ObjectInitializer)
+	, TileMapColor(FLinearColor::White)
 {
 	BodyInstance.SetCollisionEnabled(ECollisionEnabled::QueryAndPhysics);
 
@@ -226,14 +227,14 @@ void UPaperTileMapComponent::RebuildRenderData(FPaperTileMapRenderSceneProxy* Pr
 			continue;
 		}
 
-		FLinearColor DrawColor = FLinearColor::White;
+		FLinearColor DrawColor = TileMapColor * Layer->GetLayerColor();
 #if WITH_EDITORONLY_DATA
 		if (Layer->bHiddenInEditor)
 		{
 			continue;
 		}
 
-		DrawColor.A = Layer->LayerOpacity;
+		DrawColor.A *= Layer->LayerOpacity;
 #endif
 
 		FSpriteDrawCallRecord* CurrentBatch = nullptr;
@@ -463,5 +464,39 @@ UPaperTileLayer* UPaperTileMapComponent::AddNewLayer()
 	return Result;
 }
 
+FLinearColor UPaperTileMapComponent::GetTileMapColor() const
+{
+	return TileMapColor;
+}
+
+void UPaperTileMapComponent::SetTileMapColor(FLinearColor NewColor)
+{
+	TileMapColor = NewColor;
+	MarkRenderStateDirty();
+}
+
+FLinearColor UPaperTileMapComponent::GetLayerColor(int32 Layer) const
+{
+	if (TileMap->TileLayers.IsValidIndex(Layer))
+	{
+		return TileMap->TileLayers[Layer]->GetLayerColor();
+	}
+	else
+	{
+		return FLinearColor::White;
+	}
+}
+
+void UPaperTileMapComponent::SetLayerColor(FLinearColor NewColor, int32 Layer)
+{
+	if (OwnsTileMap())
+	{
+		if (TileMap->TileLayers.IsValidIndex(Layer))
+		{
+			TileMap->TileLayers[Layer]->SetLayerColor(NewColor);
+			MarkRenderStateDirty();
+		}
+	}
+}
 
 #undef LOCTEXT_NAMESPACE
\ No newline at end of file
