Tile Maps: Added the ability to drag-drop tile sets onto the tile selector panel of a tile map editor [UE-13648]

diff --git a/Engine/Plugins/2D/Paper2D/Source/Paper2DEditor/Paper2DEditor.Build.cs b/Engine/Plugins/2D/Paper2D/Source/Paper2DEditor/Paper2DEditor.Build.cs
--- a/Engine/Plugins/2D/Paper2D/Source/Paper2DEditor/Paper2DEditor.Build.cs
+++ b/Engine/Plugins/2D/Paper2D/Source/Paper2DEditor/Paper2DEditor.Build.cs
@@ -28,7 +28,8 @@ public class Paper2DEditor : ModuleRules
 				"ContentBrowser",
 				"WorkspaceMenuStructure",
 				"EditorStyle",
-				"MeshPaint"
+				"MeshPaint",
+				"EditorWidgets"
 			}
 			);
 
diff --git a/Engine/Plugins/2D/Paper2D/Source/Paper2DEditor/Private/TileMapEditing/TileMapEdModeToolkit.cpp b/Engine/Plugins/2D/Paper2D/Source/Paper2DEditor/Private/TileMapEditing/TileMapEdModeToolkit.cpp
--- a/Engine/Plugins/2D/Paper2D/Source/Paper2DEditor/Private/TileMapEditing/TileMapEdModeToolkit.cpp
+++ b/Engine/Plugins/2D/Paper2D/Source/Paper2DEditor/Private/TileMapEditing/TileMapEdModeToolkit.cpp
@@ -7,6 +7,7 @@
 #include "SContentReference.h"
 #include "TileMapEditorCommands.h"
 #include "Engine/Selection.h"
+#include "SAssetDropTarget.h"
 
 #define LOCTEXT_NAMESPACE "Paper2D"
 
@@ -81,7 +82,12 @@ void FTileMapEdModeToolkit::Init(const TSharedPtr<IToolkitHost>& InitToolkitHost
 		// The palette widget
 		+SOverlay::Slot()
 		[
-			TileSetPalette.ToSharedRef()
+			SNew(SAssetDropTarget)
+			.OnIsAssetAcceptableForDrop(this, &FTileMapEdModeToolkit::OnAssetDraggedOver)
+			.OnAssetDropped(this, &FTileMapEdModeToolkit::OnChangeTileSet)
+			[
+				TileSetPalette.ToSharedRef()
+			]
 		]
 		// The no tile set selected warning text/button
 		+SOverlay::Slot()
@@ -323,6 +329,11 @@ FReply FTileMapEdModeToolkit::ClickedOnTileSetPaletteCornerText()
 	return FReply::Handled();
 }
 
+bool FTileMapEdModeToolkit::OnAssetDraggedOver(const UObject* InObject) const
+{
+	return Cast<UPaperTileSet>(InObject) != nullptr;
+}
+
 //////////////////////////////////////////////////////////////////////////
 
 #undef LOCTEXT_NAMESPACE
diff --git a/Engine/Plugins/2D/Paper2D/Source/Paper2DEditor/Private/TileMapEditing/TileMapEdModeToolkit.h b/Engine/Plugins/2D/Paper2D/Source/Paper2DEditor/Private/TileMapEditing/TileMapEdModeToolkit.h
--- a/Engine/Plugins/2D/Paper2D/Source/Paper2DEditor/Private/TileMapEditing/TileMapEdModeToolkit.h
+++ b/Engine/Plugins/2D/Paper2D/Source/Paper2DEditor/Private/TileMapEditing/TileMapEdModeToolkit.h
@@ -40,6 +40,8 @@ protected:
 	EVisibility GetTileSetPaletteCornerTextVisibility() const;
 	FReply ClickedOnTileSetPaletteCornerText();
 
+	bool OnAssetDraggedOver(const UObject* InObject) const;
+
 private:
 	class FEdModeTileMap* TileMapEditor;
 
