Paper2D: Coding standards fix and removed a misleading comment

diff --git a/Engine/Plugins/2D/Paper2D/Source/Paper2DEditor/Private/SpriteAssetTypeActions.cpp b/Engine/Plugins/2D/Paper2D/Source/Paper2DEditor/Private/SpriteAssetTypeActions.cpp
--- a/Engine/Plugins/2D/Paper2D/Source/Paper2DEditor/Private/SpriteAssetTypeActions.cpp
+++ b/Engine/Plugins/2D/Paper2D/Source/Paper2DEditor/Private/SpriteAssetTypeActions.cpp
@@ -74,7 +74,7 @@ void FSpriteAssetTypeActions::ExecuteCreateFlipbook(TArray<TWeakObjectPtr<UPaper
 
 	for (auto ObjIt = Objects.CreateConstIterator(); ObjIt; ++ObjIt)
 	{
-		UPaperSprite *Object = (*ObjIt).Get();
+		UPaperSprite* Object = (*ObjIt).Get();
 		if (Object && Object->IsValidLowLevel())
 		{
 			AllSprites.Add(Object);
@@ -118,7 +118,7 @@ void FSpriteAssetTypeActions::ExecuteCreateFlipbook(TArray<TWeakObjectPtr<UPaper
 				KeyFrame->FrameRun = 1;
 			}
 
-			AssetToolsModule.Get().CreateUniqueAssetName(NewFlipBookDefaultPath, /*out*/ DefaultSuffix, /*out*/ PackageName, /*out*/ AssetName);
+			AssetToolsModule.Get().CreateUniqueAssetName(NewFlipBookDefaultPath, DefaultSuffix, /*out*/ PackageName, /*out*/ AssetName);
 			const FString PackagePath = FPackageName::GetLongPackagePath(PackageName);
 			if (bOneFlipbookCreated)
 			{
