Tile Maps: Added a context menu command for tile sets to create a new tile map asset [UE-13630]

diff --git a/Engine/Plugins/2D/Paper2D/Source/Paper2DEditor/Private/TileSetAssetTypeActions.cpp b/Engine/Plugins/2D/Paper2D/Source/Paper2DEditor/Private/TileSetAssetTypeActions.cpp
--- a/Engine/Plugins/2D/Paper2D/Source/Paper2DEditor/Private/TileSetAssetTypeActions.cpp
+++ b/Engine/Plugins/2D/Paper2D/Source/Paper2DEditor/Private/TileSetAssetTypeActions.cpp
@@ -3,11 +3,12 @@
 #include "Paper2DEditorPrivatePCH.h"
 #include "TileSetAssetTypeActions.h"
 #include "TileSetEditor.h"
-
+#include "AssetToolsModule.h"
+#include "ContentBrowserModule.h"
+#include "PaperTileMapFactory.h"
 
 #define LOCTEXT_NAMESPACE "AssetTypeActions"
 
-
 //////////////////////////////////////////////////////////////////////////
 // FTileSetAssetTypeActions
 
@@ -45,6 +46,52 @@ uint32 FTileSetAssetTypeActions::GetCategories()
 	return EAssetTypeCategories::Misc;
 }
 
+void FTileSetAssetTypeActions::GetActions(const TArray<UObject*>& InObjects, FMenuBuilder& MenuBuilder)
+{
+	TArray<TWeakObjectPtr<UPaperTileSet>> TileSets = GetTypedWeakObjectPtrs<UPaperTileSet>(InObjects);
+
+	if (TileSets.Num() == 1)
+	{
+		MenuBuilder.AddMenuEntry(
+			LOCTEXT("TileSet_CreateTileMap", "Create Tile Map"),
+			LOCTEXT("TileSet_CreateTileMapTooltip", "Creates a tile map using the selected tile set as a guide for tile size, etc..."),
+			FSlateIcon(FEditorStyle::GetStyleSetName(), "ClassIcon.PaperTileSet"),
+			FUIAction(FExecuteAction::CreateSP(this, &FTileSetAssetTypeActions::ExecuteCreateTileMap, TileSets[0]))
+			);
+	}
+}
+
+void FTileSetAssetTypeActions::ExecuteCreateTileMap(TWeakObjectPtr<UPaperTileSet> TileSetPtr)
+{
+	FAssetToolsModule& AssetToolsModule = FModuleManager::GetModuleChecked<FAssetToolsModule>("AssetTools");
+	FContentBrowserModule& ContentBrowserModule = FModuleManager::LoadModuleChecked<FContentBrowserModule>("ContentBrowser");
+	const FString SuffixToIgnore(TEXT("Set"));
+	const FString TileMapSuffix(TEXT("Map"));
+
+	if (UPaperTileSet* TileSet = TileSetPtr.Get())
+	{
+		// Figure out what to call the new tile map
+		FString EffectiveTileSetName = TileSet->GetName();
+		EffectiveTileSetName.RemoveFromEnd(SuffixToIgnore);
+
+		const FString TileSetPathName = TileSet->GetOutermost()->GetPathName();
+		const FString LongPackagePath = FPackageName::GetLongPackagePath(TileSetPathName);
+
+		const FString NewTileMapDefaultPath = LongPackagePath + TEXT("/") + EffectiveTileSetName;
+
+		// Make sure the name is unique
+		FString AssetName;
+		FString PackageName;
+		AssetToolsModule.Get().CreateUniqueAssetName(NewTileMapDefaultPath, TileMapSuffix, /*out*/ PackageName, /*out*/ AssetName);
+		const FString PackagePath = FPackageName::GetLongPackagePath(PackageName);
+
+		// Create the new tile map
+		UPaperTileMapFactory* TileMapFactory = NewObject<UPaperTileMapFactory>();
+		TileMapFactory->InitialTileSet = TileSet;
+		ContentBrowserModule.Get().CreateNewAsset(AssetName, PackagePath, UPaperTileMap::StaticClass(), TileMapFactory);
+	}
+}
+
 //////////////////////////////////////////////////////////////////////////
 
 #undef LOCTEXT_NAMESPACE
\ No newline at end of file
diff --git a/Engine/Plugins/2D/Paper2D/Source/Paper2DEditor/Private/TileSetAssetTypeActions.h b/Engine/Plugins/2D/Paper2D/Source/Paper2DEditor/Private/TileSetAssetTypeActions.h
--- a/Engine/Plugins/2D/Paper2D/Source/Paper2DEditor/Private/TileSetAssetTypeActions.h
+++ b/Engine/Plugins/2D/Paper2D/Source/Paper2DEditor/Private/TileSetAssetTypeActions.h
@@ -4,6 +4,8 @@
 
 #include "AssetTypeActions_Base.h"
 
+class UPaperTileSet;
+
 class FTileSetAssetTypeActions : public FAssetTypeActions_Base
 {
 public:
@@ -13,5 +15,10 @@ public:
 	virtual UClass* GetSupportedClass() const override;
 	virtual void OpenAssetEditor(const TArray<UObject*>& InObjects, TSharedPtr<class IToolkitHost> EditWithinLevelEditor = TSharedPtr<IToolkitHost>()) override;
 	virtual uint32 GetCategories() override;
+	virtual bool HasActions(const TArray<UObject*>& InObjects) const override { return true; }
+	virtual void GetActions(const TArray<UObject*>& InObjects, FMenuBuilder& MenuBuilder) override;
 	// End of IAssetTypeActions interface
-};
\ No newline at end of file
+
+private:
+	void ExecuteCreateTileMap(TWeakObjectPtr<UPaperTileSet> TileSetPtr);
+};
