Tile Maps: Changed the language for rotate/mirror commands to say brush instead of selection

diff --git a/Engine/Plugins/2D/Paper2D/Source/Paper2DEditor/Private/TileMapEditing/TileMapEditorCommands.cpp b/Engine/Plugins/2D/Paper2D/Source/Paper2DEditor/Private/TileMapEditing/TileMapEditorCommands.cpp
--- a/Engine/Plugins/2D/Paper2D/Source/Paper2DEditor/Private/TileMapEditing/TileMapEditorCommands.cpp
+++ b/Engine/Plugins/2D/Paper2D/Source/Paper2DEditor/Private/TileMapEditing/TileMapEditorCommands.cpp
@@ -23,10 +23,10 @@ void FTileMapEditorCommands::RegisterCommands()
 	UI_COMMAND(SetShowPivot, "Pivot", "Display the pivot location of the sprite.", EUserInterfaceActionType::ToggleButton, FInputChord());
 
 	// Selection actions
-	UI_COMMAND(FlipSelectionHorizontally, "Flip selection horizontally", "Flips the selection horizontally", EUserInterfaceActionType::Button, FInputChord(EKeys::X));
-	UI_COMMAND(FlipSelectionVertically, "Flip selection vertically", "Flips the selection vertically", EUserInterfaceActionType::Button, FInputChord(EKeys::Y));
-	UI_COMMAND(RotateSelectionCW, "Rotate selection clockwise", "Rotates the selection clockwise", EUserInterfaceActionType::Button, FInputChord(EKeys::Z));
-	UI_COMMAND(RotateSelectionCCW, "Rotate selection counterclockwise", "Rotates the selection counterclockwise", EUserInterfaceActionType::Button, FInputChord(EKeys::Z, EModifierKey::Shift));
+	UI_COMMAND(FlipSelectionHorizontally, "Flip brush horizontally", "Flips the brush horizontally", EUserInterfaceActionType::Button, FInputChord(EKeys::X));
+	UI_COMMAND(FlipSelectionVertically, "Flip brush vertically", "Flips the brush vertically", EUserInterfaceActionType::Button, FInputChord(EKeys::Y));
+	UI_COMMAND(RotateSelectionCW, "Rotate brush clockwise", "Rotates the brush clockwise", EUserInterfaceActionType::Button, FInputChord(EKeys::Z));
+	UI_COMMAND(RotateSelectionCCW, "Rotate brush counterclockwise", "Rotates the brush counterclockwise", EUserInterfaceActionType::Button, FInputChord(EKeys::Z, EModifierKey::Shift));
 
 	// Layer actions
 	UI_COMMAND(AddNewLayerAbove, "Add new layer", "Add a new layer above the current layer", EUserInterfaceActionType::Button, FInputChord(EKeys::N, EModifierKey::Control | EModifierKey::Shift));
