Flipbooks: Added SetPlaybackPositionInFrames and GetPlaybackPositionInFrames to UPaperFlipbookComponent [UE-13091]

diff --git a/Engine/Plugins/2D/Paper2D/Source/Paper2D/Classes/PaperFlipbookComponent.h b/Engine/Plugins/2D/Paper2D/Source/Paper2D/Classes/PaperFlipbookComponent.h
--- a/Engine/Plugins/2D/Paper2D/Source/Paper2D/Classes/PaperFlipbookComponent.h
+++ b/Engine/Plugins/2D/Paper2D/Source/Paper2D/Classes/PaperFlipbookComponent.h
@@ -105,11 +105,19 @@ public:
 	UFUNCTION(BlueprintCallable, Category="Components|Flipbook")
 	bool IsReversing() const;
 
-	/** Jump to a position in the flipbook. If bFireEvents is true, event functions will fire, otherwise will not. */
+	/** Jump to a position in the flipbook (expressed in frames). If bFireEvents is true, event functions will fire, otherwise they will not. */
+	UFUNCTION(BlueprintCallable, Category="Components|Flipbook")
+	void SetPlaybackPositionInFrames(int32 NewFramePosition, bool bFireEvents);
+
+	/** Get the current playback position (in frames) of the flipbook */
+	UFUNCTION(BlueprintCallable, Category="Components|Flipbook")
+	int32 GetPlaybackPositionInFrames() const;
+
+	/** Jump to a position in the flipbook (expressed in seconds). If bFireEvents is true, event functions will fire, otherwise they will not. */
 	UFUNCTION(BlueprintCallable, Category="Components|Flipbook")
 	void SetPlaybackPosition(float NewPosition, bool bFireEvents);
 
-	/** Get the current playback position of the flipbook */
+	/** Get the current playback position (in seconds) of the flipbook */
 	UFUNCTION(BlueprintCallable, Category="Components|Flipbook")
 	float GetPlaybackPosition() const;
 
diff --git a/Engine/Plugins/2D/Paper2D/Source/Paper2D/Private/PaperFlipbookComponent.cpp b/Engine/Plugins/2D/Paper2D/Source/Paper2D/Private/PaperFlipbookComponent.cpp
--- a/Engine/Plugins/2D/Paper2D/Source/Paper2D/Private/PaperFlipbookComponent.cpp
+++ b/Engine/Plugins/2D/Paper2D/Source/Paper2D/Private/PaperFlipbookComponent.cpp
@@ -456,6 +456,27 @@ bool UPaperFlipbookComponent::IsReversing() const
 	return bPlaying && bReversePlayback;
 }
 
+void UPaperFlipbookComponent::SetPlaybackPositionInFrames(int32 NewFramePosition, bool bFireEvents)
+{
+	const float Framerate = GetFlipbookFramerate();
+	const float NewTime = (Framerate > 0.0f) ? (NewFramePosition / Framerate) : 0.0f;
+	SetPlaybackPosition(NewTime, bFireEvents);
+}
+
+int32 UPaperFlipbookComponent::GetPlaybackPositionInFrames() const
+{
+	const float Framerate = GetFlipbookFramerate();
+	const int32 NumFrames = GetFlipbookLengthInFrames();
+	if (NumFrames > 0)
+	{
+		return FMath::Clamp<int32>(FMath::TruncToInt(AccumulatedTime * Framerate), 0, NumFrames - 1);
+	}
+	else
+	{
+		return 0;
+	}
+}
+
 void UPaperFlipbookComponent::SetPlaybackPosition(float NewPosition, bool bFireEvents)
 {
 	float OldPosition = AccumulatedTime;
