Tile Maps: Changed tile maps to always display the wireframe bounds of layer 0 when unselected in the editor, to help show where an empty tile map is [UE-7727]

diff --git a/Engine/Plugins/2D/Paper2D/Source/Paper2D/Classes/PaperTileMapComponent.h b/Engine/Plugins/2D/Paper2D/Source/Paper2D/Classes/PaperTileMapComponent.h
--- a/Engine/Plugins/2D/Paper2D/Source/Paper2D/Classes/PaperTileMapComponent.h
+++ b/Engine/Plugins/2D/Paper2D/Source/Paper2D/Classes/PaperTileMapComponent.h
@@ -130,6 +130,9 @@ public:
 	// Note: This will only work on components that own their own tile map (OwnsTileMap returns true), you cannot modify standalone tile map assets
 	UFUNCTION(BlueprintCallable, Category = "Sprite")
 	void SetLayerColor(FLinearColor NewColor, int32 Layer = 0);
+
+	// Returns the wireframe color to use for this component.
+	FLinearColor GetWireframeColor() const;
 };
 
 // Allow the old name to continue to work for one release
diff --git a/Engine/Plugins/2D/Paper2D/Source/Paper2D/Private/PaperTileMapComponent.cpp b/Engine/Plugins/2D/Paper2D/Source/Paper2D/Private/PaperTileMapComponent.cpp
--- a/Engine/Plugins/2D/Paper2D/Source/Paper2D/Private/PaperTileMapComponent.cpp
+++ b/Engine/Plugins/2D/Paper2D/Source/Paper2D/Private/PaperTileMapComponent.cpp
@@ -499,4 +499,9 @@ void UPaperTileMapComponent::SetLayerColor(FLinearColor NewColor, int32 Layer)
 	}
 }
 
+FLinearColor UPaperTileMapComponent::GetWireframeColor() const
+{
+	return TileMapColor;
+}
+
 #undef LOCTEXT_NAMESPACE
\ No newline at end of file
diff --git a/Engine/Plugins/2D/Paper2D/Source/Paper2D/Private/PaperTileMapRenderSceneProxy.cpp b/Engine/Plugins/2D/Paper2D/Source/Paper2D/Private/PaperTileMapRenderSceneProxy.cpp
--- a/Engine/Plugins/2D/Paper2D/Source/Paper2D/Private/PaperTileMapRenderSceneProxy.cpp
+++ b/Engine/Plugins/2D/Paper2D/Source/Paper2D/Private/PaperTileMapRenderSceneProxy.cpp
@@ -20,12 +20,30 @@ FPaperTileMapRenderSceneProxy::FPaperTileMapRenderSceneProxy(const UPaperTileMap
 {
 	if (const UPaperTileMapComponent* InTileComponent = Cast<const UPaperTileMapComponent>(InComponent))
 	{
+		WireframeColor = InTileComponent->GetWireframeColor();
 		TileMap = InTileComponent->TileMap;
 		Material = InTileComponent->GetMaterial(0);
 		MaterialRelevance = InTileComponent->GetMaterialRelevance(GetScene().GetFeatureLevel());
 	}
 }
 
+void FPaperTileMapRenderSceneProxy::DrawBoundsForLayer(FPrimitiveDrawInterface* PDI, const FLinearColor& Color, int32 LayerIndex) const
+{
+	// Slight depth bias so that the wireframe grid overlay doesn't z-fight with the tiles themselves
+	const float DepthBias = 0.0001f;
+
+	const FMatrix& LocalToWorld = GetLocalToWorld();
+	const FVector TL(LocalToWorld.TransformPosition(TileMap->GetTilePositionInLocalSpace(0, 0, LayerIndex)));
+	const FVector TR(LocalToWorld.TransformPosition(TileMap->GetTilePositionInLocalSpace(TileMap->MapWidth, 0, LayerIndex)));
+	const FVector BL(LocalToWorld.TransformPosition(TileMap->GetTilePositionInLocalSpace(0, TileMap->MapHeight, LayerIndex)));
+	const FVector BR(LocalToWorld.TransformPosition(TileMap->GetTilePositionInLocalSpace(TileMap->MapWidth, TileMap->MapHeight, LayerIndex)));
+
+	PDI->DrawLine(TL, TR, Color, SDPG_Foreground, 0.0f, DepthBias);
+	PDI->DrawLine(TR, BR, Color, SDPG_Foreground, 0.0f, DepthBias);
+	PDI->DrawLine(BR, BL, Color, SDPG_Foreground, 0.0f, DepthBias);
+	PDI->DrawLine(BL, TL, Color, SDPG_Foreground, 0.0f, DepthBias);
+}
+
 void FPaperTileMapRenderSceneProxy::GetDynamicMeshElements(const TArray<const FSceneView*>& Views, const FSceneViewFamily& ViewFamily, uint32 VisibilityMap, FMeshElementCollector& Collector) const
 {
 	QUICK_SCOPE_CYCLE_COUNTER(STAT_FPaperTileMapRenderSceneProxy_GetDynamicMeshElements);
@@ -45,8 +63,6 @@ void FPaperTileMapRenderSceneProxy::GetDynamicMeshElements(const TArray<const FS
 			//@TODO: RenderThread race condition
 			if (TileMap != nullptr)
 			{
-				FColor WireframeColor = FColor(0, 255, 255, 255);
-
 				if ((View->Family->EngineShowFlags.Collision /*@TODO: && bIsCollisionEnabled*/) && AllowDebugViewmodes())
 				{
 					if (UBodySetup2D* BodySetup2D = Cast<UBodySetup2D>(TileMap->BodySetup))
@@ -125,15 +141,7 @@ void FPaperTileMapRenderSceneProxy::GetDynamicMeshElements(const TArray<const FS
 						{
 							if (LayerIndex != SelectedLayerIndex)
 							{
-								const FVector TL(LocalToWorld.TransformPosition(TileMap->GetTilePositionInLocalSpace(0, 0, LayerIndex)));
-								const FVector TR(LocalToWorld.TransformPosition(TileMap->GetTilePositionInLocalSpace(TileMap->MapWidth, 0, LayerIndex)));
-								const FVector BL(LocalToWorld.TransformPosition(TileMap->GetTilePositionInLocalSpace(0, TileMap->MapHeight, LayerIndex)));
-								const FVector BR(LocalToWorld.TransformPosition(TileMap->GetTilePositionInLocalSpace(TileMap->MapWidth, TileMap->MapHeight, LayerIndex)));
-
-								PDI->DrawLine(TL, TR, OverrideColor, DPG, 0.0f, DepthBias);
-								PDI->DrawLine(TR, BR, OverrideColor, DPG, 0.0f, DepthBias);
-								PDI->DrawLine(BR, BL, OverrideColor, DPG, 0.0f, DepthBias);
-								PDI->DrawLine(BL, TL, OverrideColor, DPG, 0.0f, DepthBias);
+								DrawBoundsForLayer(PDI, OverrideColor, LayerIndex);
 							}
 						}
 
@@ -164,6 +172,11 @@ void FPaperTileMapRenderSceneProxy::GetDynamicMeshElements(const TArray<const FS
 							}
 						}
 					}
+					else
+					{
+						// Draw layer 0 even when not selected, so you can see where the tile map is in the editor
+						DrawBoundsForLayer(PDI, WireframeColor, 0);
+					}
 				}
 #endif
 			}
diff --git a/Engine/Plugins/2D/Paper2D/Source/Paper2D/Private/PaperTileMapRenderSceneProxy.h b/Engine/Plugins/2D/Paper2D/Source/Paper2D/Private/PaperTileMapRenderSceneProxy.h
--- a/Engine/Plugins/2D/Paper2D/Source/Paper2D/Private/PaperTileMapRenderSceneProxy.h
+++ b/Engine/Plugins/2D/Paper2D/Source/Paper2D/Private/PaperTileMapRenderSceneProxy.h
@@ -20,6 +20,10 @@ public:
 	{
 		BatchedSprites = InBatchedSprites;
 	}
+
+protected:
+	void DrawBoundsForLayer(FPrimitiveDrawInterface* PDI, const FLinearColor& Color, int32 LayerIndex) const;
+
 protected:
 	//@TODO: Not thread safe
 	const class UPaperTileMap* TileMap;
