Flipbooks: Marked a method that had been deprecated for Blueprints in 4.5 as deprecated for C++ too

diff --git a/Engine/Plugins/2D/Paper2D/Source/Paper2D/Classes/PaperFlipbookComponent.h b/Engine/Plugins/2D/Paper2D/Source/Paper2D/Classes/PaperFlipbookComponent.h
--- a/Engine/Plugins/2D/Paper2D/Source/Paper2D/Classes/PaperFlipbookComponent.h
+++ b/Engine/Plugins/2D/Paper2D/Source/Paper2D/Classes/PaperFlipbookComponent.h
@@ -69,7 +69,7 @@ public:
 	virtual UPaperFlipbook* GetFlipbook();
 
 	/** DEPRECATED! Use GetMaterial() instead.  Gets the material used by this instance */
-	//@TODO: DEPRECATED(4.5, "GetSpriteMaterial has been replaced by GetMaterial")
+	DEPRECATED(4.5, "GetSpriteMaterial has been replaced by GetMaterial")
 	UFUNCTION(BlueprintCallable, Category="Sprite", meta=(DeprecatedFunction, DeprecationMessage="Use GetMaterial instead"))
 	UMaterialInterface* GetSpriteMaterial() const;
 
