Paper2D: Added a configurable background color for the sprite and flipbook editors (in Editor Preferences..Content Editors) [UE-13640]

diff --git a/Engine/Plugins/2D/Paper2D/Source/Paper2DEditor/Private/FlipbookEditor/FlipbookEditorSettings.cpp b/Engine/Plugins/2D/Paper2D/Source/Paper2DEditor/Private/FlipbookEditor/FlipbookEditorSettings.cpp
new file mode 100644
--- /dev/null
+++ b/Engine/Plugins/2D/Paper2D/Source/Paper2DEditor/Private/FlipbookEditor/FlipbookEditorSettings.cpp
@@ -0,0 +1,12 @@
+// Copyright 1998-2015 Epic Games, Inc. All Rights Reserved.
+
+#include "Paper2DEditorPrivatePCH.h"
+#include "FlipbookEditorSettings.h"
+
+//////////////////////////////////////////////////////////////////////////
+// UFlipbookEditorSettings
+
+UFlipbookEditorSettings::UFlipbookEditorSettings()
+	: BackgroundColor(55, 55, 55)
+{
+}
diff --git a/Engine/Plugins/2D/Paper2D/Source/Paper2DEditor/Private/FlipbookEditor/FlipbookEditorSettings.h b/Engine/Plugins/2D/Paper2D/Source/Paper2DEditor/Private/FlipbookEditor/FlipbookEditorSettings.h
new file mode 100644
--- /dev/null
+++ b/Engine/Plugins/2D/Paper2D/Source/Paper2DEditor/Private/FlipbookEditor/FlipbookEditorSettings.h
@@ -0,0 +1,19 @@
+// Copyright 1998-2015 Epic Games, Inc. All Rights Reserved.
+
+#pragma once
+
+#include "FlipbookEditorSettings.generated.h"
+
+// Settings for the Paper2D sprite editor
+UCLASS(config=EditorUserSettings)
+class UFlipbookEditorSettings : public UObject
+{
+	GENERATED_BODY()
+
+public:
+	UFlipbookEditorSettings();
+
+	/** Background color in the flipbook editor */
+	UPROPERTY(config, EditAnywhere, Category=Background, meta=(HideAlphaChannel))
+	FColor BackgroundColor;
+};
diff --git a/Engine/Plugins/2D/Paper2D/Source/Paper2DEditor/Private/FlipbookEditor/FlipbookEditorViewportClient.cpp b/Engine/Plugins/2D/Paper2D/Source/Paper2DEditor/Private/FlipbookEditor/FlipbookEditorViewportClient.cpp
--- a/Engine/Plugins/2D/Paper2D/Source/Paper2DEditor/Private/FlipbookEditor/FlipbookEditorViewportClient.cpp
+++ b/Engine/Plugins/2D/Paper2D/Source/Paper2DEditor/Private/FlipbookEditor/FlipbookEditorViewportClient.cpp
@@ -8,6 +8,7 @@
 #include "ScopedTransaction.h"
 #include "CanvasTypes.h"
 #include "PaperEditorShared/SocketEditing.h"
+#include "FlipbookEditorSettings.h"
 
 #define LOCTEXT_NAMESPACE "FlipbookEditor"
 
@@ -116,7 +117,7 @@ bool FFlipbookEditorViewportClient::InputKey(FViewport* Viewport, int32 Controll
 
 FLinearColor FFlipbookEditorViewportClient::GetBackgroundColor() const
 {
-	return FEditorViewportClient::GetBackgroundColor();
+	return GetDefault<UFlipbookEditorSettings>()->BackgroundColor;
 }
 
 //////////////////////////////////////////////////////////////////////////
diff --git a/Engine/Plugins/2D/Paper2D/Source/Paper2DEditor/Private/Paper2DEditorModule.cpp b/Engine/Plugins/2D/Paper2D/Source/Paper2DEditor/Private/Paper2DEditorModule.cpp
--- a/Engine/Plugins/2D/Paper2D/Source/Paper2DEditor/Private/Paper2DEditorModule.cpp
+++ b/Engine/Plugins/2D/Paper2D/Source/Paper2DEditor/Private/Paper2DEditorModule.cpp
@@ -7,6 +7,7 @@
 #include "PropertyEditorModule.h"
 #include "PaperStyle.h"
 #include "PaperEditorCommands.h"
+#include "PaperEditorShared/SpriteGeometryEditMode.h"
 
 #include "AssetEditorToolkit.h"
 #include "ModuleManager.h"
@@ -19,14 +20,14 @@
 #include "SpriteEditor/SpriteDetailsCustomization.h"
 #include "SpriteEditor/SpriteComponentDetailsCustomization.h"
 #include "SpriteEditor/SpritePolygonCollectionCustomization.h"
-
-#include "PaperEditorShared/SpriteGeometryEditMode.h"
+#include "SpriteEditor/SpriteEditorSettings.h"
 
 // Flipbook support
 #include "FlipbookAssetTypeActions.h"
 #include "PaperFlipbookAssetBroker.h"
 #include "PaperFlipbookThumbnailRenderer.h"
 #include "FlipbookEditor/FlipbookComponentDetailsCustomization.h"
+#include "FlipbookEditor/FlipbookEditorSettings.h"
 
 // Tile set support
 #include "TileSetAssetTypeActions.h"
@@ -258,8 +259,17 @@ private:
 			SettingsModule->RegisterSettings("Project", "Plugins", "Paper2D",
 				LOCTEXT("RuntimeSettingsName", "Paper 2D"),
 				LOCTEXT("RuntimeSettingsDescription", "Configure the Paper 2D plugin"),
-				GetMutableDefault<UPaperRuntimeSettings>()
-			);
+				GetMutableDefault<UPaperRuntimeSettings>());
+
+			SettingsModule->RegisterSettings("Editor", "ContentEditors", "SpriteEditor",
+				LOCTEXT("SpriteEditorSettingsName", "Sprite Editor"),
+				LOCTEXT("SpriteEditorSettingsDescription", "Configure the look and feel of the Sprite Editor."),
+				GetMutableDefault<USpriteEditorSettings>());
+
+			SettingsModule->RegisterSettings("Editor", "ContentEditors", "FlipbookEditor",
+				LOCTEXT("FlipbookEditorSettingsName", "Flipbook Editor"),
+				LOCTEXT("FlipbookEditorSettingsDescription", "Configure the look and feel of the Flipbook Editor."),
+				GetMutableDefault<UFlipbookEditorSettings>());
 		}
 	}
 
@@ -267,6 +277,8 @@ private:
 	{
 		if (ISettingsModule* SettingsModule = FModuleManager::GetModulePtr<ISettingsModule>("Settings"))
 		{
+			SettingsModule->UnregisterSettings("Editor", "ContentEditors", "FlipbookEditor");
+			SettingsModule->UnregisterSettings("Editor", "ContentEditors", "SpriteEditor");
 			SettingsModule->UnregisterSettings("Project", "Plugins", "Paper2D");
 		}
 	}
diff --git a/Engine/Plugins/2D/Paper2D/Source/Paper2DEditor/Private/SpriteEditor/SpriteEditorSettings.cpp b/Engine/Plugins/2D/Paper2D/Source/Paper2DEditor/Private/SpriteEditor/SpriteEditorSettings.cpp
new file mode 100644
--- /dev/null
+++ b/Engine/Plugins/2D/Paper2D/Source/Paper2DEditor/Private/SpriteEditor/SpriteEditorSettings.cpp
@@ -0,0 +1,12 @@
+// Copyright 1998-2015 Epic Games, Inc. All Rights Reserved.
+
+#include "Paper2DEditorPrivatePCH.h"
+#include "SpriteEditorSettings.h"
+
+//////////////////////////////////////////////////////////////////////////
+// USpriteEditorSettings
+
+USpriteEditorSettings::USpriteEditorSettings()
+	: BackgroundColor(55, 55, 55)
+{
+}
diff --git a/Engine/Plugins/2D/Paper2D/Source/Paper2DEditor/Private/SpriteEditor/SpriteEditorSettings.h b/Engine/Plugins/2D/Paper2D/Source/Paper2DEditor/Private/SpriteEditor/SpriteEditorSettings.h
new file mode 100644
--- /dev/null
+++ b/Engine/Plugins/2D/Paper2D/Source/Paper2DEditor/Private/SpriteEditor/SpriteEditorSettings.h
@@ -0,0 +1,19 @@
+// Copyright 1998-2015 Epic Games, Inc. All Rights Reserved.
+
+#pragma once
+
+#include "SpriteEditorSettings.generated.h"
+
+// Settings for the Paper2D sprite editor
+UCLASS(config=EditorUserSettings)
+class USpriteEditorSettings : public UObject
+{
+	GENERATED_BODY()
+
+public:
+	USpriteEditorSettings();
+
+	/** Background color in the sprite editor */
+	UPROPERTY(config, EditAnywhere, Category=Background, meta=(HideAlphaChannel))
+	FColor BackgroundColor;
+};
diff --git a/Engine/Plugins/2D/Paper2D/Source/Paper2DEditor/Private/SpriteEditor/SpriteEditorViewportClient.cpp b/Engine/Plugins/2D/Paper2D/Source/Paper2DEditor/Private/SpriteEditor/SpriteEditorViewportClient.cpp
--- a/Engine/Plugins/2D/Paper2D/Source/Paper2DEditor/Private/SpriteEditor/SpriteEditorViewportClient.cpp
+++ b/Engine/Plugins/2D/Paper2D/Source/Paper2DEditor/Private/SpriteEditor/SpriteEditorViewportClient.cpp
@@ -23,6 +23,7 @@
 #include "SNotificationList.h"
 #include "NotificationManager.h"
 
+#include "SpriteEditorSettings.h"
 #include "PaperSpriteFactory.h"
 
 #define LOCTEXT_NAMESPACE "SpriteEditor"
@@ -948,7 +949,7 @@ void FSpriteEditorViewportClient::TrackingStopped()
 
 FLinearColor FSpriteEditorViewportClient::GetBackgroundColor() const
 {
-	return FEditorViewportClient::GetBackgroundColor();
+	return GetDefault<USpriteEditorSettings>()->BackgroundColor;
 }
 
 FVector2D FSpriteEditorViewportClient::SelectedItemConvertWorldSpaceDeltaToLocalSpace(const FVector& WorldSpaceDelta) const
