Paper2D: Added a default setting for the 'show grid' option in the sprite editor

diff --git a/Engine/Plugins/2D/Paper2D/Source/Paper2DEditor/Private/SpriteEditor/SpriteEditorSettings.cpp b/Engine/Plugins/2D/Paper2D/Source/Paper2DEditor/Private/SpriteEditor/SpriteEditorSettings.cpp
--- a/Engine/Plugins/2D/Paper2D/Source/Paper2DEditor/Private/SpriteEditor/SpriteEditorSettings.cpp
+++ b/Engine/Plugins/2D/Paper2D/Source/Paper2DEditor/Private/SpriteEditor/SpriteEditorSettings.cpp
@@ -8,5 +8,6 @@
 
 USpriteEditorSettings::USpriteEditorSettings()
 	: BackgroundColor(55, 55, 55)
+	, bShowGridByDefault(false)
 {
 }
diff --git a/Engine/Plugins/2D/Paper2D/Source/Paper2DEditor/Private/SpriteEditor/SpriteEditorSettings.h b/Engine/Plugins/2D/Paper2D/Source/Paper2DEditor/Private/SpriteEditor/SpriteEditorSettings.h
--- a/Engine/Plugins/2D/Paper2D/Source/Paper2DEditor/Private/SpriteEditor/SpriteEditorSettings.h
+++ b/Engine/Plugins/2D/Paper2D/Source/Paper2DEditor/Private/SpriteEditor/SpriteEditorSettings.h
@@ -16,4 +16,8 @@ public:
 	/** Background color in the sprite editor */
 	UPROPERTY(config, EditAnywhere, Category=Background, meta=(HideAlphaChannel))
 	FColor BackgroundColor;
+
+	/** Should the grid be shown by default when the editor is opened? */
+	UPROPERTY(config, EditAnywhere, Category=Background)
+	bool bShowGridByDefault;
 };
diff --git a/Engine/Plugins/2D/Paper2D/Source/Paper2DEditor/Private/SpriteEditor/SpriteEditorViewportClient.cpp b/Engine/Plugins/2D/Paper2D/Source/Paper2DEditor/Private/SpriteEditor/SpriteEditorViewportClient.cpp
--- a/Engine/Plugins/2D/Paper2D/Source/Paper2DEditor/Private/SpriteEditor/SpriteEditorViewportClient.cpp
+++ b/Engine/Plugins/2D/Paper2D/Source/Paper2DEditor/Private/SpriteEditor/SpriteEditorViewportClient.cpp
@@ -87,7 +87,7 @@ FSpriteEditorViewportClient::FSpriteEditorViewportClient(TWeakPtr<FSpriteEditor>
 	bShowPivot = true;
 	bShowRelatedSprites = true;
 
-	DrawHelper.bDrawGrid = false;
+	DrawHelper.bDrawGrid = GetDefault<USpriteEditorSettings>()->bShowGridByDefault;
 
 	EngineShowFlags.DisableAdvancedFeatures();
 	EngineShowFlags.CompositeEditorPrimitives = true;
