Paper2D: Added a default setting for the 'show grid' option in the flipbook editor

diff --git a/Engine/Plugins/2D/Paper2D/Source/Paper2DEditor/Private/FlipbookEditor/FlipbookEditorSettings.cpp b/Engine/Plugins/2D/Paper2D/Source/Paper2DEditor/Private/FlipbookEditor/FlipbookEditorSettings.cpp
--- a/Engine/Plugins/2D/Paper2D/Source/Paper2DEditor/Private/FlipbookEditor/FlipbookEditorSettings.cpp
+++ b/Engine/Plugins/2D/Paper2D/Source/Paper2DEditor/Private/FlipbookEditor/FlipbookEditorSettings.cpp
@@ -8,5 +8,6 @@
 
 UFlipbookEditorSettings::UFlipbookEditorSettings()
 	: BackgroundColor(55, 55, 55)
+	, bShowGridByDefault(false)
 {
 }
diff --git a/Engine/Plugins/2D/Paper2D/Source/Paper2DEditor/Private/FlipbookEditor/FlipbookEditorSettings.h b/Engine/Plugins/2D/Paper2D/Source/Paper2DEditor/Private/FlipbookEditor/FlipbookEditorSettings.h
--- a/Engine/Plugins/2D/Paper2D/Source/Paper2DEditor/Private/FlipbookEditor/FlipbookEditorSettings.h
+++ b/Engine/Plugins/2D/Paper2D/Source/Paper2DEditor/Private/FlipbookEditor/FlipbookEditorSettings.h
@@ -16,4 +16,8 @@ public:
 	/** Background color in the flipbook editor */
 	UPROPERTY(config, EditAnywhere, Category=Background, meta=(HideAlphaChannel))
 	FColor BackgroundColor;
+
+	/** Should the grid be shown by default when the editor is opened? */
+	UPROPERTY(config, EditAnywhere, Category = Background)
+	bool bShowGridByDefault;
 };
diff --git a/Engine/Plugins/2D/Paper2D/Source/Paper2DEditor/Private/FlipbookEditor/FlipbookEditorViewportClient.cpp b/Engine/Plugins/2D/Paper2D/Source/Paper2DEditor/Private/FlipbookEditor/FlipbookEditorViewportClient.cpp
--- a/Engine/Plugins/2D/Paper2D/Source/Paper2DEditor/Private/FlipbookEditor/FlipbookEditorViewportClient.cpp
+++ b/Engine/Plugins/2D/Paper2D/Source/Paper2DEditor/Private/FlipbookEditor/FlipbookEditorViewportClient.cpp
@@ -32,7 +32,7 @@ FFlipbookEditorViewportClient::FFlipbookEditorViewportClient(const TAttribute<UP
 
 	bShowPivot = false;
 	bShowSockets = true;
-	DrawHelper.bDrawGrid = false;
+	DrawHelper.bDrawGrid = GetDefault<UFlipbookEditorSettings>()->bShowGridByDefault;
 
 	EngineShowFlags.DisableAdvancedFeatures();
 	EngineShowFlags.CompositeEditorPrimitives = true;
