Paper2D: Added tile map and tile set editor settings, controlling the default background color for new assets and the default 'show grid' setting for the respective editors

diff --git a/Engine/Plugins/2D/Paper2D/Source/Paper2DEditor/Private/Paper2DEditorModule.cpp b/Engine/Plugins/2D/Paper2D/Source/Paper2DEditor/Private/Paper2DEditorModule.cpp
--- a/Engine/Plugins/2D/Paper2D/Source/Paper2DEditor/Private/Paper2DEditorModule.cpp
+++ b/Engine/Plugins/2D/Paper2D/Source/Paper2DEditor/Private/Paper2DEditorModule.cpp
@@ -32,12 +32,14 @@
 // Tile set support
 #include "TileSetAssetTypeActions.h"
 #include "PaperTileSetThumbnailRenderer.h"
+#include "TileSetEditor/TileSetEditorSettings.h"
 
 // Tile map support
 #include "TileMapEditing/TileMapAssetTypeActions.h"
 #include "TileMapEditing/PaperTileMapAssetBroker.h"
 #include "TileMapEditing/EdModeTileMap.h"
 #include "TileMapEditing/PaperTileMapDetailsCustomization.h"
+#include "TileMapEditing/TileMapEditorSettings.h"
 
 // Atlas support
 #include "Atlasing/AtlasAssetTypeActions.h"
@@ -270,6 +272,16 @@ private:
 				LOCTEXT("FlipbookEditorSettingsName", "Flipbook Editor"),
 				LOCTEXT("FlipbookEditorSettingsDescription", "Configure the look and feel of the Flipbook Editor."),
 				GetMutableDefault<UFlipbookEditorSettings>());
+
+			SettingsModule->RegisterSettings("Editor", "ContentEditors", "TileMapEditor",
+				LOCTEXT("TileMapEditorSettingsName", "Tile Map Editor"),
+				LOCTEXT("TileMapEditorSettingsDescription", "Configure the look and feel of the Tile Map Editor."),
+				GetMutableDefault<UFlipbookEditorSettings>());
+
+			SettingsModule->RegisterSettings("Editor", "ContentEditors", "TileSetEditor",
+				LOCTEXT("TileSetEditorSettingsName", "Tile Set Editor"),
+				LOCTEXT("TileSetEditorSettingsDescription", "Configure the look and feel of the Tile Set Editor."),
+				GetMutableDefault<UTileSetEditorSettings>());
 		}
 	}
 
@@ -277,6 +289,8 @@ private:
 	{
 		if (ISettingsModule* SettingsModule = FModuleManager::GetModulePtr<ISettingsModule>("Settings"))
 		{
+			SettingsModule->UnregisterSettings("Editor", "ContentEditors", "TileSetEditor");
+			SettingsModule->UnregisterSettings("Editor", "ContentEditors", "TileMapEditor");
 			SettingsModule->UnregisterSettings("Editor", "ContentEditors", "FlipbookEditor");
 			SettingsModule->UnregisterSettings("Editor", "ContentEditors", "SpriteEditor");
 			SettingsModule->UnregisterSettings("Project", "Plugins", "Paper2D");
diff --git a/Engine/Plugins/2D/Paper2D/Source/Paper2DEditor/Private/PaperTileMapFactory.cpp b/Engine/Plugins/2D/Paper2D/Source/Paper2DEditor/Private/PaperTileMapFactory.cpp
--- a/Engine/Plugins/2D/Paper2D/Source/Paper2DEditor/Private/PaperTileMapFactory.cpp
+++ b/Engine/Plugins/2D/Paper2D/Source/Paper2DEditor/Private/PaperTileMapFactory.cpp
@@ -2,6 +2,7 @@
 
 #include "Paper2DEditorPrivatePCH.h"
 #include "PaperTileMapFactory.h"
+#include "TileMapEditing/TileMapEditorSettings.h"
 
 #define LOCTEXT_NAMESPACE "Paper2D"
 
@@ -27,6 +28,7 @@ UObject* UPaperTileMapFactory::FactoryCreateNew(UClass* Class, UObject* InParent
 		NewTileMap->SelectedTileSet = InitialTileSet;
 	}
 
+	NewTileMap->BackgroundColor = GetDefault<UTileMapEditorSettings>()->DefaultBackgroundColor;
 	NewTileMap->InitializeNewEmptyTileMap();
 
 	return NewTileMap;
diff --git a/Engine/Plugins/2D/Paper2D/Source/Paper2DEditor/Private/PaperTileSetFactory.cpp b/Engine/Plugins/2D/Paper2D/Source/Paper2DEditor/Private/PaperTileSetFactory.cpp
--- a/Engine/Plugins/2D/Paper2D/Source/Paper2DEditor/Private/PaperTileSetFactory.cpp
+++ b/Engine/Plugins/2D/Paper2D/Source/Paper2DEditor/Private/PaperTileSetFactory.cpp
@@ -2,6 +2,7 @@
 
 #include "Paper2DEditorPrivatePCH.h"
 #include "PaperTileSetFactory.h"
+#include "TileSetEditor/TileSetEditorSettings.h"
 
 #define LOCTEXT_NAMESPACE "Paper2D"
 
@@ -18,8 +19,9 @@ UPaperTileSetFactory::UPaperTileSetFactory(const FObjectInitializer& ObjectIniti
 
 UObject* UPaperTileSetFactory::FactoryCreateNew(UClass* Class, UObject* InParent, FName Name, EObjectFlags Flags, UObject* Context, FFeedbackContext* Warn)
 {
-	auto NewTileSet = NewObject<UPaperTileSet>(InParent, Class, Name, Flags | RF_Transactional);
+	UPaperTileSet* NewTileSet = NewObject<UPaperTileSet>(InParent, Class, Name, Flags | RF_Transactional);
 	NewTileSet->TileSheet = InitialTexture;
+	NewTileSet->BackgroundColor = GetDefault<UTileSetEditorSettings>()->DefaultBackgroundColor;
 
 	return NewTileSet;
 }
diff --git a/Engine/Plugins/2D/Paper2D/Source/Paper2DEditor/Private/TileMapEditing/TileMapEditorSettings.cpp b/Engine/Plugins/2D/Paper2D/Source/Paper2DEditor/Private/TileMapEditing/TileMapEditorSettings.cpp
new file mode 100644
--- /dev/null
+++ b/Engine/Plugins/2D/Paper2D/Source/Paper2DEditor/Private/TileMapEditing/TileMapEditorSettings.cpp
@@ -0,0 +1,13 @@
+// Copyright 1998-2015 Epic Games, Inc. All Rights Reserved.
+
+#include "Paper2DEditorPrivatePCH.h"
+#include "TileMapEditorSettings.h"
+
+//////////////////////////////////////////////////////////////////////////
+// UTileMapEditorSettings
+
+UTileMapEditorSettings::UTileMapEditorSettings()
+	: DefaultBackgroundColor(55, 55, 55)
+	, bShowGridByDefault(false)
+{
+}
diff --git a/Engine/Plugins/2D/Paper2D/Source/Paper2DEditor/Private/TileMapEditing/TileMapEditorSettings.h b/Engine/Plugins/2D/Paper2D/Source/Paper2DEditor/Private/TileMapEditing/TileMapEditorSettings.h
new file mode 100644
--- /dev/null
+++ b/Engine/Plugins/2D/Paper2D/Source/Paper2DEditor/Private/TileMapEditing/TileMapEditorSettings.h
@@ -0,0 +1,23 @@
+// Copyright 1998-2015 Epic Games, Inc. All Rights Reserved.
+
+#pragma once
+
+#include "TileMapEditorSettings.generated.h"
+
+// Settings for the Paper2D tile map editor
+UCLASS(config=EditorUserSettings)
+class UTileMapEditorSettings : public UObject
+{
+	GENERATED_BODY()
+
+public:
+	UTileMapEditorSettings();
+
+	/** Default background color for new tile map assets */
+	UPROPERTY(config, EditAnywhere, Category=Background, meta=(HideAlphaChannel))
+	FColor DefaultBackgroundColor;
+
+	/** Should the grid be shown by default when the editor is opened? */
+	UPROPERTY(config, EditAnywhere, Category=Background)
+	bool bShowGridByDefault;
+};
diff --git a/Engine/Plugins/2D/Paper2D/Source/Paper2DEditor/Private/TileMapEditing/TileMapEditorViewportClient.cpp b/Engine/Plugins/2D/Paper2D/Source/Paper2DEditor/Private/TileMapEditing/TileMapEditorViewportClient.cpp
--- a/Engine/Plugins/2D/Paper2D/Source/Paper2DEditor/Private/TileMapEditing/TileMapEditorViewportClient.cpp
+++ b/Engine/Plugins/2D/Paper2D/Source/Paper2DEditor/Private/TileMapEditing/TileMapEditorViewportClient.cpp
@@ -10,6 +10,7 @@
 #include "Runtime/Engine/Public/ComponentReregisterContext.h"
 #include "CanvasTypes.h"
 #include "Engine/Selection.h"
+#include "TileMapEditorSettings.h"
 
 #define LOCTEXT_NAMESPACE "TileMapEditor"
 
@@ -35,7 +36,7 @@ FTileMapEditorViewportClient::FTileMapEditorViewportClient(TWeakPtr<FTileMapEdit
 	bManipulationDirtiedSomething = false;
 	ScopedTransaction = nullptr;
 
-	DrawHelper.bDrawGrid = false;
+	DrawHelper.bDrawGrid = GetDefault<UTileMapEditorSettings>()->bShowGridByDefault;
 	DrawHelper.bDrawPivot = false;
 	bShowPivot = true;
 
@@ -127,7 +128,7 @@ FLinearColor FTileMapEditorViewportClient::GetBackgroundColor() const
 	}
 	else
 	{
-		return FEditorViewportClient::GetBackgroundColor();
+		return GetDefault<UTileMapEditorSettings>()->DefaultBackgroundColor;
 	}
 }
 
diff --git a/Engine/Plugins/2D/Paper2D/Source/Paper2DEditor/Private/TileSetEditor/SingleTileEditorViewportClient.cpp b/Engine/Plugins/2D/Paper2D/Source/Paper2DEditor/Private/TileSetEditor/SingleTileEditorViewportClient.cpp
--- a/Engine/Plugins/2D/Paper2D/Source/Paper2DEditor/Private/TileSetEditor/SingleTileEditorViewportClient.cpp
+++ b/Engine/Plugins/2D/Paper2D/Source/Paper2DEditor/Private/TileSetEditor/SingleTileEditorViewportClient.cpp
@@ -5,6 +5,7 @@
 #include "SpriteEditor/SpriteEditorCommands.h"
 #include "PaperEditorShared/SpriteGeometryEditMode.h"
 #include "ScopedTransaction.h"
+#include "TileSetEditorSettings.h"
 
 #define LOCTEXT_NAMESPACE "TileSetEditor"
 
@@ -24,7 +25,7 @@ FSingleTileEditorViewportClient::FSingleTileEditorViewportClient(UPaperTileSet*
 	// The tile map editor fully supports mode tools and isn't doing any incompatible stuff with the Widget
 	Widget->SetUsesEditorModeTools(ModeTools);
 
-	DrawHelper.bDrawGrid = true;//@TODO:
+	DrawHelper.bDrawGrid = GetDefault<UTileSetEditorSettings>()->bShowGridByDefault;
 	DrawHelper.bDrawPivot = false;
 
 	PreviewScene = &OwnedPreviewScene;
diff --git a/Engine/Plugins/2D/Paper2D/Source/Paper2DEditor/Private/TileSetEditor/TileSetEditorSettings.cpp b/Engine/Plugins/2D/Paper2D/Source/Paper2DEditor/Private/TileSetEditor/TileSetEditorSettings.cpp
new file mode 100644
--- /dev/null
+++ b/Engine/Plugins/2D/Paper2D/Source/Paper2DEditor/Private/TileSetEditor/TileSetEditorSettings.cpp
@@ -0,0 +1,13 @@
+// Copyright 1998-2015 Epic Games, Inc. All Rights Reserved.
+
+#include "Paper2DEditorPrivatePCH.h"
+#include "TileSetEditorSettings.h"
+
+//////////////////////////////////////////////////////////////////////////
+// UTileSetEditorSettings
+
+UTileSetEditorSettings::UTileSetEditorSettings()
+	: DefaultBackgroundColor(0, 0, 127)
+	, bShowGridByDefault(true)
+{
+}
diff --git a/Engine/Plugins/2D/Paper2D/Source/Paper2DEditor/Private/TileSetEditor/TileSetEditorSettings.h b/Engine/Plugins/2D/Paper2D/Source/Paper2DEditor/Private/TileSetEditor/TileSetEditorSettings.h
new file mode 100644
--- /dev/null
+++ b/Engine/Plugins/2D/Paper2D/Source/Paper2DEditor/Private/TileSetEditor/TileSetEditorSettings.h
@@ -0,0 +1,23 @@
+// Copyright 1998-2015 Epic Games, Inc. All Rights Reserved.
+
+#pragma once
+
+#include "TileSetEditorSettings.generated.h"
+
+// Settings for the Paper2D tile set editor
+UCLASS(config=EditorUserSettings)
+class UTileSetEditorSettings : public UObject
+{
+	GENERATED_BODY()
+
+public:
+	UTileSetEditorSettings();
+
+	/** Default background color for new tile set assets */
+	UPROPERTY(config, EditAnywhere, Category=Background, meta=(HideAlphaChannel))
+	FColor DefaultBackgroundColor;
+
+	/** Should the grid be shown by default when the editor is opened? */
+	UPROPERTY(config, EditAnywhere, Category="Tile Editor")
+	bool bShowGridByDefault;
+};
