Paper2D: Style cleanup

diff --git a/Engine/Plugins/2D/Paper2D/Source/Paper2DEditor/Classes/PaperSpriteFactory.h b/Engine/Plugins/2D/Paper2D/Source/Paper2DEditor/Classes/PaperSpriteFactory.h
--- a/Engine/Plugins/2D/Paper2D/Source/Paper2DEditor/Classes/PaperSpriteFactory.h
+++ b/Engine/Plugins/2D/Paper2D/Source/Paper2DEditor/Classes/PaperSpriteFactory.h
@@ -21,13 +21,8 @@ class UPaperSpriteFactory : public UFactory
 	FVector2D InitialSourceUV;
 	FVector2D InitialSourceDimension;
 
-
 	// UFactory interface
 	virtual bool ConfigureProperties() override;
 	virtual UObject* FactoryCreateNew(UClass* Class, UObject* InParent, FName Name, EObjectFlags Flags, UObject* Context, FFeedbackContext* Warn) override;
 	// End of UFactory interface
-
-protected:
-	/** A pointer to the window that is asking the user to select a parent class */
-//	TSharedPtr<SWindow> PickerWindow;
 };
diff --git a/Engine/Plugins/2D/Paper2D/Source/Paper2DEditor/Private/ExtractSprites/SPaperExtractSpritesDialog.cpp b/Engine/Plugins/2D/Paper2D/Source/Paper2DEditor/Private/ExtractSprites/SPaperExtractSpritesDialog.cpp
--- a/Engine/Plugins/2D/Paper2D/Source/Paper2DEditor/Private/ExtractSprites/SPaperExtractSpritesDialog.cpp
+++ b/Engine/Plugins/2D/Paper2D/Source/Paper2DEditor/Private/ExtractSprites/SPaperExtractSpritesDialog.cpp
@@ -425,7 +425,7 @@ void SPaperExtractSpritesDialog::CreateExtractedSprites()
 	TArray<UObject*> ObjectsToSync;
 
 	// Create the factory used to generate the sprite
-	auto SpriteFactory = NewObject<UPaperSpriteFactory>();
+	UPaperSpriteFactory* SpriteFactory = NewObject<UPaperSpriteFactory>();
 	SpriteFactory->InitialTexture = SourceTexture;
 
 	// Create the sprite
