Paper2D: Added import settings and auto-detection for twinned diffuse+normal maps
- It can handle things like Texture + Texture_N, Texture_D + Texture_N, etc..., where "_D" and "_N" are configurable in the editor preferences
- Newly created sprites with twinned normal maps will default to using a lit material instead of an unlit material (the specific material is configurable in editor preferences as well)
- Works well with sprite sheets exported from Texture Packer using the "Pack with same layout" option
[UE-10417]

diff --git a/Engine/Plugins/2D/Paper2D/Source/Paper2D/Classes/SpriteEditorOnlyTypes.h b/Engine/Plugins/2D/Paper2D/Source/Paper2D/Classes/SpriteEditorOnlyTypes.h
--- a/Engine/Plugins/2D/Paper2D/Source/Paper2D/Classes/SpriteEditorOnlyTypes.h
+++ b/Engine/Plugins/2D/Paper2D/Source/Paper2D/Classes/SpriteEditorOnlyTypes.h
@@ -256,6 +256,9 @@ struct FSpriteAssetInitParameters
 		, Offset(FVector2D::ZeroVector)
 		, Dimension(FVector2D::ZeroVector)
 		, bNewlyCreated(false)
+		, OpaqueMaterialOverride(nullptr)
+		, TranslucentMaterialOverride(nullptr)
+		, MaskedMaterialOverride(nullptr)
 	{
 	}
 
@@ -279,6 +282,9 @@ public:
 	// The texture to use
 	UTexture2D* Texture;
 
+	// Additional textures to use
+	TArray<UTexture2D*> AdditionalTextures;
+
 	// The offset within the texture (in pixels)
 	FVector2D Offset;
 
@@ -287,6 +293,15 @@ public:
 
 	// Is this sprite newly created (should we pull the default pixels/uu and materials from the project settings)?
 	bool bNewlyCreated;
+
+	// The material to use for opaque sprites if set
+	UMaterialInterface* OpaqueMaterialOverride;
+
+	// The material to use for translucent sprites if set
+	UMaterialInterface* TranslucentMaterialOverride;
+
+	// The material to use for masked sprites if set
+	UMaterialInterface* MaskedMaterialOverride;
 };
 
 UENUM()
diff --git a/Engine/Plugins/2D/Paper2D/Source/Paper2D/Private/PaperSprite.cpp b/Engine/Plugins/2D/Paper2D/Source/Paper2D/Private/PaperSprite.cpp
--- a/Engine/Plugins/2D/Paper2D/Source/Paper2D/Private/PaperSprite.cpp
+++ b/Engine/Plugins/2D/Paper2D/Source/Paper2D/Private/PaperSprite.cpp
@@ -1371,20 +1371,32 @@ void UPaperSprite::InitializeSprite(const FSpriteAssetInitParameters& InitParams
 
 		if (bUseMaskedTexture)
 		{
-			if (UMaterialInterface* MaskedMaterial = LoadObject<UMaterialInterface>(nullptr, *DefaultSettings->DefaultMaskedMaterialName.ToString(), nullptr, LOAD_None, nullptr))
+			if (InitParams.MaskedMaterialOverride != nullptr)
+			{
+				DefaultMaterial = InitParams.MaskedMaterialOverride;
+			}
+			else if (UMaterialInterface* MaskedMaterial = Cast<UMaterialInterface>(DefaultSettings->DefaultMaskedMaterialName.TryLoad()))
 			{
 				DefaultMaterial = MaskedMaterial;
 			}
 		}
 		else
 		{
-			if (UMaterialInterface* TranslucentMaterial = LoadObject<UMaterialInterface>(nullptr, *DefaultSettings->DefaultTranslucentMaterialName.ToString(), nullptr, LOAD_None, nullptr))
+			if (InitParams.TranslucentMaterialOverride != nullptr)
+			{
+				DefaultMaterial = InitParams.TranslucentMaterialOverride;
+			}
+			else if (UMaterialInterface* TranslucentMaterial = Cast<UMaterialInterface>(DefaultSettings->DefaultTranslucentMaterialName.TryLoad()))
 			{
 				DefaultMaterial = TranslucentMaterial;
 			}
 		}
 
-		if (UMaterialInterface* OpaqueMaterial = LoadObject<UMaterialInterface>(nullptr, *DefaultSettings->DefaultOpaqueMaterialName.ToString(), nullptr, LOAD_None, nullptr))
+		if (InitParams.OpaqueMaterialOverride != nullptr)
+		{
+			AlternateMaterial = InitParams.OpaqueMaterialOverride;
+		}
+		else if (UMaterialInterface* OpaqueMaterial = Cast<UMaterialInterface>(DefaultSettings->DefaultOpaqueMaterialName.TryLoad()))
 		{
 			AlternateMaterial = OpaqueMaterial;
 		}
@@ -1399,6 +1411,8 @@ void UPaperSprite::InitializeSprite(const FSpriteAssetInitParameters& InitParams
 	{
 		SourceTextureDimension.Set(0, 0);
 	}
+	AdditionalSourceTextures = InitParams.AdditionalTextures;
+
 	SourceUV = InitParams.Offset;
 	SourceDimension = InitParams.Dimension;
 
diff --git a/Engine/Plugins/2D/Paper2D/Source/Paper2DEditor/Classes/PaperImporterSettings.h b/Engine/Plugins/2D/Paper2D/Source/Paper2DEditor/Classes/PaperImporterSettings.h
new file mode 100644
--- /dev/null
+++ b/Engine/Plugins/2D/Paper2D/Source/Paper2DEditor/Classes/PaperImporterSettings.h
@@ -0,0 +1,47 @@
+// Copyright 1998-2015 Epic Games, Inc. All Rights Reserved.
+
+#pragma once
+
+#include "PaperImporterSettings.generated.h"
+
+/**
+ * Implements the settings for imported Paper2D assets, such as sprite sheet textures.
+ */
+UCLASS(config=Editor, defaultconfig)
+class PAPER2DEDITOR_API UPaperImporterSettings : public UObject
+{
+	GENERATED_BODY()
+
+public:
+	// A list of default suffixes to use when looking for associated normal maps while importing sprites or creating sprites from textures
+	UPROPERTY(config, EditAnywhere, Category=ImportSettings)
+	TArray<FString> NormalMapTextureSuffixes;
+
+	// The default suffix to remove (if present) from a texture name before looking for an associated normal map using NormalMapTextureSuffixes
+	UPROPERTY(config, EditAnywhere, Category=ImportSettings)
+	TArray<FString> BaseMapTextureSuffixes;
+
+	// The default masked material for newly created sprites (masked means binary opacity: things are either opaque or see-thru, with nothing in between)
+	UPROPERTY(config, EditAnywhere, Category=Settings, meta=(AllowedClasses="MaterialInterface", DisplayName="Lit Default Masked Material"))
+	FStringAssetReference LitDefaultMaskedMaterialName;
+
+	// The default translucent material for newly created sprites (translucent means smooth opacity which can vary continuously from 0..1, but translucent rendering is more expensive that opaque or masked rendering and has different sorting rules)
+	UPROPERTY(config, EditAnywhere, Category=Settings, meta=(AllowedClasses="MaterialInterface", DisplayName="Lit Default Translucent Material"))
+	FStringAssetReference LitDefaultTranslucentMaterialName;
+
+	// The default opaque material for newly created sprites
+	UPROPERTY(config, EditAnywhere, Category=Settings, meta=(AllowedClasses="MaterialInterface", DisplayName="Lit Default Opaque Material"))
+	FStringAssetReference LitDefaultOpaqueMaterialName;
+
+public:
+	UPaperImporterSettings();
+
+	// Removes the suffix from the specified name if it matches something in BaseMapTextureSuffixes
+	FString RemoveSuffixFromBaseMapName(const FString& InName) const;
+
+	// Generates names to test for a normal map using NormalMapTextureSuffixes
+	void GenerateNormalMapNamesToTest(const FString& InRoot, TArray<FString>& InOutNames) const;
+
+	// Pushes the lit materials
+	void PopulateMaterialsIntoInitParams(FSpriteAssetInitParameters& InitParams) const;
+};
diff --git a/Engine/Plugins/2D/Paper2D/Source/Paper2DEditor/Private/Paper2DEditorModule.cpp b/Engine/Plugins/2D/Paper2D/Source/Paper2DEditor/Private/Paper2DEditorModule.cpp
--- a/Engine/Plugins/2D/Paper2D/Source/Paper2DEditor/Private/Paper2DEditorModule.cpp
+++ b/Engine/Plugins/2D/Paper2D/Source/Paper2DEditor/Private/Paper2DEditorModule.cpp
@@ -12,6 +12,7 @@
 #include "AssetEditorToolkit.h"
 #include "ModuleManager.h"
 #include "ContentBrowserExtensions/ContentBrowserExtensions.h"
+#include "PaperImporterSettings.h"
 
 // Sprite support
 #include "SpriteAssetTypeActions.h"
@@ -282,6 +283,11 @@ private:
 				LOCTEXT("TileSetEditorSettingsName", "Tile Set Editor"),
 				LOCTEXT("TileSetEditorSettingsDescription", "Configure the look and feel of the Tile Set Editor."),
 				GetMutableDefault<UTileSetEditorSettings>());
+
+			SettingsModule->RegisterSettings("Editor", "General", "Paper2DImport",
+				LOCTEXT("PaperImporterSettingsName", "Paper2D - Import"),
+				LOCTEXT("PaperImporterSettingsDescription", "Configure how assets get imported or converted to sprites."),
+				GetMutableDefault<UPaperImporterSettings>());
 		}
 	}
 
@@ -289,6 +295,7 @@ private:
 	{
 		if (ISettingsModule* SettingsModule = FModuleManager::GetModulePtr<ISettingsModule>("Settings"))
 		{
+			SettingsModule->UnregisterSettings("Editor", "General", "Paper2DImport");
 			SettingsModule->UnregisterSettings("Editor", "ContentEditors", "TileSetEditor");
 			SettingsModule->UnregisterSettings("Editor", "ContentEditors", "TileMapEditor");
 			SettingsModule->UnregisterSettings("Editor", "ContentEditors", "FlipbookEditor");
diff --git a/Engine/Plugins/2D/Paper2D/Source/Paper2DEditor/Private/PaperImporterSettings.cpp b/Engine/Plugins/2D/Paper2D/Source/Paper2DEditor/Private/PaperImporterSettings.cpp
new file mode 100644
--- /dev/null
+++ b/Engine/Plugins/2D/Paper2D/Source/Paper2DEditor/Private/PaperImporterSettings.cpp
@@ -0,0 +1,53 @@
+// Copyright 1998-2015 Epic Games, Inc. All Rights Reserved.
+
+#include "Paper2DEditorPrivatePCH.h"
+#include "PaperImporterSettings.h"
+
+//////////////////////////////////////////////////////////////////////////
+// UPaperImporterSettings
+
+UPaperImporterSettings::UPaperImporterSettings()
+{
+	NormalMapTextureSuffixes.Add(TEXT("_N"));
+	NormalMapTextureSuffixes.Add(TEXT("_Normal"));
+
+	BaseMapTextureSuffixes.Add(TEXT("_D"));
+	BaseMapTextureSuffixes.Add(TEXT("_Diffuse"));
+
+	if (!IsRunningCommandlet())
+	{
+		LitDefaultMaskedMaterialName = FStringAssetReference("/Paper2D/MaskedLitSpriteMaterial.MaskedLitSpriteMaterial");
+		LitDefaultTranslucentMaterialName = FStringAssetReference("/Paper2D/TranslucentLitSpriteMaterial.TranslucentLitSpriteMaterial");
+		LitDefaultOpaqueMaterialName = FStringAssetReference("/Paper2D/OpaqueLitSpriteMaterial.OpaqueLitSpriteMaterial");
+	}
+}
+
+FString UPaperImporterSettings::RemoveSuffixFromBaseMapName(const FString& InName) const
+{
+	FString Result = InName;
+
+	for (const FString& PossibleSuffix : BaseMapTextureSuffixes)
+	{
+		if (Result.RemoveFromEnd(PossibleSuffix))
+		{
+			return Result;
+		}
+	}
+
+	return Result;
+}
+
+void UPaperImporterSettings::GenerateNormalMapNamesToTest(const FString& InRoot, TArray<FString>& InOutNames) const
+{
+	for (const FString& PossibleSuffix : NormalMapTextureSuffixes)
+	{
+		InOutNames.Add(InRoot + PossibleSuffix);
+	}
+}
+
+void UPaperImporterSettings::PopulateMaterialsIntoInitParams(FSpriteAssetInitParameters& InitParams) const
+{
+	InitParams.MaskedMaterialOverride = Cast<UMaterialInterface>(LitDefaultMaskedMaterialName.TryLoad());
+	InitParams.OpaqueMaterialOverride = Cast<UMaterialInterface>(LitDefaultOpaqueMaterialName.TryLoad());
+	InitParams.TranslucentMaterialOverride = Cast<UMaterialInterface>(LitDefaultTranslucentMaterialName.TryLoad());
+}
diff --git a/Engine/Plugins/2D/Paper2D/Source/Paper2DEditor/Private/PaperSpriteFactory.cpp b/Engine/Plugins/2D/Paper2D/Source/Paper2DEditor/Private/PaperSpriteFactory.cpp
--- a/Engine/Plugins/2D/Paper2D/Source/Paper2DEditor/Private/PaperSpriteFactory.cpp
+++ b/Engine/Plugins/2D/Paper2D/Source/Paper2DEditor/Private/PaperSpriteFactory.cpp
@@ -1,7 +1,10 @@
 // Copyright 1998-2015 Epic Games, Inc. All Rights Reserved.
 
 #include "Paper2DEditorPrivatePCH.h"
+#include "PaperImporterSettings.h"
 #include "PaperSpriteFactory.h"
+#include "AssetRegistryModule.h"
+#include "PackageTools.h"
 
 #define LOCTEXT_NAMESPACE "Paper2D"
 
@@ -31,6 +34,7 @@ UObject* UPaperSpriteFactory::FactoryCreateNew(UClass* Class, UObject* InParent,
 
 	FSpriteAssetInitParameters SpriteInitParams;
 	SpriteInitParams.bNewlyCreated = true;
+
 	if (bUseSourceRegion)
 	{
 		SpriteInitParams.Texture = InitialTexture;
@@ -41,6 +45,40 @@ UObject* UPaperSpriteFactory::FactoryCreateNew(UClass* Class, UObject* InParent,
 	{
 		SpriteInitParams.SetTextureAndFill(InitialTexture);
 	}
+
+	if (InitialTexture != nullptr)
+	{
+		// Look for an associated normal map to go along with the base map
+		const FString SanitizedBasePackageName = PackageTools::SanitizePackageName(InitialTexture->GetOutermost()->GetName());
+		const FString PackagePath = FPackageName::GetLongPackagePath(SanitizedBasePackageName);
+		FAssetRegistryModule& AssetRegistryModule = FModuleManager::Get().LoadModuleChecked<FAssetRegistryModule>(TEXT("AssetRegistry"));
+
+		const UPaperImporterSettings* ImporterSettings = GetDefault<UPaperImporterSettings>();
+
+		const FString NormalMapNameNoSuffix = ImporterSettings->RemoveSuffixFromBaseMapName(InitialTexture->GetName());
+
+		TArray<FString> NamesToTest;
+		ImporterSettings->GenerateNormalMapNamesToTest(NormalMapNameNoSuffix, /*inout*/ NamesToTest);
+		ImporterSettings->GenerateNormalMapNamesToTest(InitialTexture->GetName(), /*inout*/ NamesToTest);
+
+		// Test each name for an existing asset
+		for (const FString& NameToTest : NamesToTest)
+		{
+			const FString ObjectPathToTest = PackagePath / (NameToTest + FString(TEXT(".")) + NameToTest);
+			FAssetData AssetData = AssetRegistryModule.Get().GetAssetByObjectPath(*ObjectPathToTest);
+
+			if (AssetData.IsValid())
+			{
+				if (UTexture2D* NormalMapTexture = Cast<UTexture2D>(AssetData.GetAsset()))
+				{
+					SpriteInitParams.AdditionalTextures.Add(NormalMapTexture);
+					ImporterSettings->PopulateMaterialsIntoInitParams(/*inout*/ SpriteInitParams);
+					break;
+				}
+			}
+		}
+	}
+
 	NewSprite->InitializeSprite(SpriteInitParams);
 
 	return NewSprite;
diff --git a/Engine/Plugins/2D/Paper2D/Source/PaperSpriteSheetImporter/Classes/PaperSpriteSheetImportFactory.h b/Engine/Plugins/2D/Paper2D/Source/PaperSpriteSheetImporter/Classes/PaperSpriteSheetImportFactory.h
deleted file mode 100644
diff --git a/Engine/Plugins/2D/Paper2D/Source/PaperSpriteSheetImporter/Private/PaperJsonSpriteSheetImporter.cpp b/Engine/Plugins/2D/Paper2D/Source/PaperSpriteSheetImporter/Private/PaperJsonSpriteSheetImporter.cpp
--- a/Engine/Plugins/2D/Paper2D/Source/PaperSpriteSheetImporter/Private/PaperJsonSpriteSheetImporter.cpp
+++ b/Engine/Plugins/2D/Paper2D/Source/PaperSpriteSheetImporter/Private/PaperJsonSpriteSheetImporter.cpp
@@ -10,6 +10,7 @@
 #include "PackageTools.h"
 #include "ObjectTools.h"
 #include "PaperSpriteSheet.h"
+#include "PaperImporterSettings.h"
 
 TSharedPtr<FJsonObject> ParseJSON(const FString& FileContents, const FString& NameForErrors)
 {
@@ -314,18 +315,20 @@ static ESpritePivotMode::Type GetBestPivotType(FVector2D JsonPivot)
 	}
 }
 
-// FPaperJsonSpriteSheetImporter
 ////////////////////////////////////////////
+// FPaperJsonSpriteSheetImporter
 
-FPaperJsonSpriteSheetImporter::FPaperJsonSpriteSheetImporter() : bIsReimporting(false)
+FPaperJsonSpriteSheetImporter::FPaperJsonSpriteSheetImporter()
+	: ImageTexture(nullptr)
+	, NormalMapTexture(nullptr)
+	, bIsReimporting(false)
+	, ExistingBaseTexture(nullptr)
+	, ExistingNormalMapTexture(nullptr)
 {
 }
 
-void FPaperJsonSpriteSheetImporter::SetReimportData(const FString& InExistingTextureName, UTexture2D* InExistingTexture, const TArray<FString>& ExistingSpriteNames, const TArray< TAssetPtr<class UPaperSprite> >& ExistingSpriteAssetPtrs)
+void FPaperJsonSpriteSheetImporter::SetReimportData(const TArray<FString>& ExistingSpriteNames, const TArray< TAssetPtr<class UPaperSprite> >& ExistingSpriteAssetPtrs)
 {
-	ExistingTextureName = InExistingTextureName;
-	ExistingTexture = InExistingTexture;
-
 	check(ExistingSpriteNames.Num() == ExistingSpriteAssetPtrs.Num());
 	if (ExistingSpriteNames.Num() == ExistingSpriteAssetPtrs.Num())
 	{
@@ -348,7 +351,7 @@ void FPaperJsonSpriteSheetImporter::SetReimportData(const FString& InExistingTex
 
 bool FPaperJsonSpriteSheetImporter::Import(TSharedPtr<FJsonObject> SpriteDescriptorObject, const FString& NameForErrors)
 {
-	bool bLoadedSuccessfully = ParseMetaBlock(NameForErrors, SpriteDescriptorObject, Image);
+	bool bLoadedSuccessfully = ParseMetaBlock(NameForErrors, SpriteDescriptorObject, /*out*/ ImageName);
 	if (bLoadedSuccessfully)
 	{
 		TSharedPtr<FJsonObject> ObjectFrameBlock = FPaperJSONHelpers::ReadObject(SpriteDescriptorObject, TEXT("frames"));
@@ -373,7 +376,7 @@ bool FPaperJsonSpriteSheetImporter::Import(TSharedPtr<FJsonObject> SpriteDescrip
 			}
 		}
 
-		if (bLoadedSuccessfully && Frames.Num() == 0)
+		if (bLoadedSuccessfully && (Frames.Num() == 0))
 		{
 			UE_LOG(LogPaperSpriteSheetImporter, Warning, TEXT("Failed to parse sprite descriptor file '%s'.  No frames loaded"), *NameForErrors);
 			bLoadedSuccessfully = false;
@@ -400,73 +403,113 @@ bool FPaperJsonSpriteSheetImporter::ImportFromArchive(FArchive* Archive, const F
 bool FPaperJsonSpriteSheetImporter::ImportTextures(const FString& LongPackagePath, const FString& SourcePath)
 {
 	bool bLoadedSuccessfully = true;
-	bool bFoundExistingTexture = false;
-	const FString TargetSubPath = LongPackagePath + TEXT("/Textures");
+	const FString TargetSubPath = LongPackagePath / TEXT("Textures");
 
-	const FString SourceSheetImageFilename = FPaths::Combine(*SourcePath, *Image);
-	TArray<FString> SheetFileNames;
-	SheetFileNames.Add(SourceSheetImageFilename);
-
-	if (bIsReimporting && (ExistingTextureName == Image) && (ExistingTexture != nullptr))
+	// Load the base texture
+	const FString SourceSheetImageFilename = FPaths::Combine(*SourcePath, *ImageName);
+	ImageTexture = ImportOrReimportTexture((bIsReimporting && (ExistingBaseTextureName == ImageName)) ? ExistingBaseTexture : nullptr, SourceSheetImageFilename, TargetSubPath);
+	if (ImageTexture == nullptr)
 	{
-		ImageTexture = ExistingTexture;
-		bFoundExistingTexture = true;
+		UE_LOG(LogPaperSpriteSheetImporter, Warning, TEXT("Failed to import sprite sheet image '%s'."), *SourceSheetImageFilename);
+		bLoadedSuccessfully = false;
 	}
 
-	// Import the texture if we need to
-	bool bNeedsImport = !bIsReimporting || (bIsReimporting && !bFoundExistingTexture);
-
-	if (bIsReimporting && bFoundExistingTexture)
+	// Try reimporting the normal map
+	// Note: We are checking to see if the *base* texture has been renamed, since the .JSON doesn't actually store a name for the normal map.
+	// If the base name has changed, we start from scratch for the normal map too, rather than reimport it even if the old computed one still exists
+	if (bIsReimporting && (ExistingBaseTextureName == ImageName) && (ExistingNormalMapTexture != nullptr))
 	{
-		if (!FReimportManager::Instance()->Reimport(ImageTexture, /*bAskForNewFileIfMissing=*/ true))
+		if (FReimportManager::Instance()->Reimport(ExistingNormalMapTexture, /*bAskForNewFileIfMissing=*/ true))
 		{
-			bNeedsImport = true;
+			NormalMapTexture = ExistingNormalMapTexture;
+			ComputedNormalMapName = ExistingNormalMapTextureName;
 		}
 	}
 
-	if (bNeedsImport)
+	// If we weren't reimporting (or failed the reimport), try scanning for a normal map (which may not exist, and that is not an error)
+	if (NormalMapTexture == nullptr)
 	{
-		FAssetToolsModule& AssetToolsModule = FModuleManager::GetModuleChecked<FAssetToolsModule>("AssetTools");
+		const UPaperImporterSettings* ImporterSettings = GetDefault<UPaperImporterSettings>();
 
-		UReimportTextureFactory* ReimportFactory = nullptr;
+		// Create a list of names to test of the form [ImageName-[BaseMapSuffix]][NormalMapSuffix] or [ImageName][NormalMapSuffix], preferring the former
+		const FString ImageNameNoExtension = FPaths::GetBaseFilename(ImageName);
+		const FString ImageTypeExtension = FPaths::GetExtension(ImageName, /*bIncludeDot=*/ true);
+		const FString NormalMapNameNoSuffix = ImporterSettings->RemoveSuffixFromBaseMapName(ImageNameNoExtension);
 
-		TArray<UObject*> ImportedSheets = AssetToolsModule.Get().ImportAssets(SheetFileNames, TargetSubPath, ReimportFactory);
+		TArray<FString> NamesToTest;
+		ImporterSettings->GenerateNormalMapNamesToTest(NormalMapNameNoSuffix, /*inout*/ NamesToTest);
+		ImporterSettings->GenerateNormalMapNamesToTest(ImageNameNoExtension, /*inout*/ NamesToTest);
 
-		UTexture2D* ImportedTexture = (ImportedSheets.Num() > 0) ? Cast<UTexture2D>(ImportedSheets[0]) : nullptr;
-
-		if (ImportedTexture == nullptr)
-		{
-			UE_LOG(LogPaperSpriteSheetImporter, Warning, TEXT("Failed to import sprite sheet image '%s'."), *SourceSheetImageFilename);
-			bLoadedSuccessfully = false;
-		}
-		else
+		// Test each name for a file we can try to import
+		for (const FString& NameToTestNoExtension : NamesToTest)
 		{
-			// Change the compression settings
-			//@TODO: Should we always be doing this (particularly the TF_Nearest seems potentially undesirable; maybe make it a property of the sprite sheet asset?)
-			ImportedTexture->Modify();
-			ImportedTexture->LODGroup = TEXTUREGROUP_UI;
-			ImportedTexture->CompressionSettings = TC_EditorIcon;
-			ImportedTexture->Filter = TF_Nearest;
-			ImportedTexture->PostEditChange();
-		}
+			const FString NameToTest = NameToTestNoExtension + ImageTypeExtension;
+			const FString NormalMapSourceImageFilename = FPaths::Combine(*SourcePath, *NameToTest);
 
-		ImageTexture = ImportedTexture;
+			if (FPaths::FileExists(NormalMapSourceImageFilename))
+			{
+				NormalMapTexture = ImportTexture(NormalMapSourceImageFilename, TargetSubPath);
+				if (NormalMapTexture != nullptr)
+				{
+					ComputedNormalMapName = NameToTest;
+				}
+				break;
+			}
+		}
 	}
 
 	return bLoadedSuccessfully;
 }
 
-UPaperSprite* FPaperJsonSpriteSheetImporter::FindExistingSprite(const FString& Name)
+UTexture2D* FPaperJsonSpriteSheetImporter::ImportOrReimportTexture(UTexture2D* ExistingTexture, const FString& TextureSourcePath, const FString& DestinationAssetFolder)
 {
-	UPaperSprite**ExistingSprite = ExistingSprites.Find(Name);
-	if (ExistingSprite != nullptr)
+	UTexture2D* ResultTexture = nullptr;
+
+	// Try reimporting if we have an existing texture
+	if (ExistingTexture != nullptr)
 	{
-		return *ExistingSprite;
+		if (FReimportManager::Instance()->Reimport(ExistingTexture, /*bAskForNewFileIfMissing=*/ true))
+		{
+			ResultTexture = ExistingTexture;
+		}
 	}
-	else
+
+	// If that fails, import the original textures
+	if (ResultTexture == nullptr)
 	{
-		return nullptr;
+		ResultTexture = ImportTexture(TextureSourcePath, DestinationAssetFolder);
+	}
+
+	return ResultTexture;
+}
+
+UTexture2D* FPaperJsonSpriteSheetImporter::ImportTexture(const FString& TextureSourcePath, const FString& DestinationAssetFolder)
+{
+	FAssetToolsModule& AssetToolsModule = FModuleManager::GetModuleChecked<FAssetToolsModule>("AssetTools");
+
+	TArray<FString> TextureFileNames;
+	TextureFileNames.Add(TextureSourcePath);
+	TArray<UObject*> ImportedSheets = AssetToolsModule.Get().ImportAssets(TextureFileNames, DestinationAssetFolder);
+
+	UTexture2D* ImportedTexture = (ImportedSheets.Num() > 0) ? Cast<UTexture2D>(ImportedSheets[0]) : nullptr;
+
+	if (ImportedTexture != nullptr)
+	{
+		// Change the compression settings
+		//@TODO: Should we always be doing this (particularly the TF_Nearest seems potentially undesirable; maybe make it a property of the sprite sheet asset?)
+		ImportedTexture->Modify();
+		ImportedTexture->LODGroup = TEXTUREGROUP_UI;
+		ImportedTexture->CompressionSettings = TC_EditorIcon;
+		ImportedTexture->Filter = TF_Nearest;
+		ImportedTexture->PostEditChange();
 	}
+
+	return ImportedTexture;
+}
+
+UPaperSprite* FPaperJsonSpriteSheetImporter::FindExistingSprite(const FString& Name)
+{
+	return ExistingSprites.FindRef(Name);
 }
 
 bool FPaperJsonSpriteSheetImporter::PerformImport(const FString& LongPackagePath, EObjectFlags Flags, UPaperSpriteSheet* SpriteSheet)
@@ -523,6 +566,16 @@ bool FPaperJsonSpriteSheetImporter::PerformImport(const FString& LongPackagePath
 		TargetSprite->Modify();
 		FSpriteAssetInitParameters SpriteInitParams;
 		SpriteInitParams.Texture = ImageTexture;
+
+		if (NormalMapTexture != nullptr)
+		{
+			// Put the normal map into the additional textures array and ask for a lit material instead of unlit
+			SpriteInitParams.AdditionalTextures.Add(NormalMapTexture);
+
+			const UPaperImporterSettings* ImporterSettings = GetDefault<UPaperImporterSettings>();
+			ImporterSettings->PopulateMaterialsIntoInitParams(/*inout*/ SpriteInitParams);
+		}
+
 		SpriteInitParams.Offset = Frame.SpritePosInSheet;
 		SpriteInitParams.Dimension = Frame.SpriteSizeInSheet;
 		SpriteInitParams.bNewlyCreated = true;
@@ -555,8 +608,10 @@ bool FPaperJsonSpriteSheetImporter::PerformImport(const FString& LongPackagePath
 		TargetSprite->PostEditChange();
 	}
 
-	SpriteSheet->TextureName = Image;
+	SpriteSheet->TextureName = ImageName;
 	SpriteSheet->Texture = ImageTexture;
+	SpriteSheet->NormalMapTextureName = ComputedNormalMapName;
+	SpriteSheet->NormalMapTexture = NormalMapTexture;
 
 	GWarn->EndSlowTask();
 	return true;
diff --git a/Engine/Plugins/2D/Paper2D/Source/PaperSpriteSheetImporter/Private/PaperJsonSpriteSheetImporter.h b/Engine/Plugins/2D/Paper2D/Source/PaperSpriteSheetImporter/Private/PaperJsonSpriteSheetImporter.h
--- a/Engine/Plugins/2D/Paper2D/Source/PaperSpriteSheetImporter/Private/PaperJsonSpriteSheetImporter.h
+++ b/Engine/Plugins/2D/Paper2D/Source/PaperSpriteSheetImporter/Private/PaperJsonSpriteSheetImporter.h
@@ -41,21 +41,43 @@ public:
 	bool ImportFromArchive(FArchive* Archive, const FString& NameForErrors);
 
 	bool ImportTextures(const FString& LongPackagePath, const FString& SourcePath);
+
 	bool PerformImport(const FString& LongPackagePath, EObjectFlags Flags, UPaperSpriteSheet* SpriteSheet);
 
-	void SetReimportData(const FString& ExistingImageName, UTexture2D* ExistingTexture, const TArray<FString>& ExistingSpriteNames, const TArray< TAssetPtr<class UPaperSprite> >& ExistingSpriteAssetPtrs);
+	void SetReimportData(const TArray<FString>& ExistingSpriteNames, const TArray< TAssetPtr<class UPaperSprite> >& ExistingSpriteAssetPtrs);
+
+	static UTexture2D* ImportOrReimportTexture(UTexture2D* ExistingTexture, const FString& TextureSourcePath, const FString& DestinationAssetFolder);
+	static UTexture2D* ImportTexture(const FString& TextureSourcePath, const FString& DestinationAssetFolder);
 
 protected:
 	bool Import(TSharedPtr<FJsonObject> SpriteDescriptorObject, const FString& NameForErrors);
 	UPaperSprite* FindExistingSprite(const FString& Name);
 
+protected:
 	TArray<FSpriteFrame> Frames;
-	FString Image;
+
+	FString ImageName;
 	UTexture2D* ImageTexture;
 
+	FString ComputedNormalMapName;
+	UTexture2D* NormalMapTexture;
+
+public:
 	bool bIsReimporting;
- 	FString ExistingTextureName;					// Name of the last imported texture
-	UTexture2D* ExistingTexture;					// And the texture itself
-	TMap<FString, UPaperSprite*> ExistingSprites;	// Map of a sprite name (as seen in the importer) -> UPaperSprite
+
+	// The name of the default or diffuse texture during a previous import
+	FString ExistingBaseTextureName;
+
+	// The asset that was created for ExistingBaseTextureName during a previous import
+	UTexture2D* ExistingBaseTexture;
+
+	// The name of the normal map texture during a previous import (if any)
+	FString ExistingNormalMapTextureName;
+
+	// The asset that was created for ExistingNormalMapTextureName  during a previous import (if any)
+	UTexture2D* ExistingNormalMapTexture;
+
+	// Map of a sprite name (as seen in the importer) -> UPaperSprite
+	TMap<FString, UPaperSprite*> ExistingSprites;
 };
 
diff --git a/Engine/Plugins/2D/Paper2D/Source/PaperSpriteSheetImporter/Private/PaperSpriteSheet.h b/Engine/Plugins/2D/Paper2D/Source/PaperSpriteSheetImporter/Private/PaperSpriteSheet.h
--- a/Engine/Plugins/2D/Paper2D/Source/PaperSpriteSheetImporter/Private/PaperSpriteSheet.h
+++ b/Engine/Plugins/2D/Paper2D/Source/PaperSpriteSheetImporter/Private/PaperSpriteSheet.h
@@ -14,16 +14,25 @@ public:
 	UPROPERTY(VisibleAnywhere, Category=Data)
 	TArray<FString> SpriteNames;
 
-	UPROPERTY(VisibleAnywhere, Category = Data)
+	UPROPERTY(VisibleAnywhere, Category=Data)
 	TArray< TAssetPtr<class UPaperSprite> > Sprites;
 
-	// The name of the texture during import
-	UPROPERTY(VisibleAnywhere, Category = Data)
+	// The name of the default or diffuse texture during import
+	UPROPERTY(VisibleAnywhere, Category=Data)
 	FString TextureName;
 
-	UPROPERTY(VisibleAnywhere, Category = Data)
+	// The asset that was created for TextureName
+	UPROPERTY(VisibleAnywhere, Category=Data)
 	UTexture2D* Texture;
 
+	// The name of the normal map texture during import (if any)
+	UPROPERTY(VisibleAnywhere, Category=Data)
+	FString NormalMapTextureName;
+
+	// The asset that was created for NormalMapTextureName (if any)
+	UPROPERTY(VisibleAnywhere, Category=Data)
+	UTexture2D* NormalMapTexture;
+
 #if WITH_EDITORONLY_DATA
 	// Import data for this 
 	UPROPERTY(VisibleAnywhere, Instanced, Category=ImportSettings)
diff --git a/Engine/Plugins/2D/Paper2D/Source/PaperSpriteSheetImporter/Private/PaperSpriteSheetImportFactory.cpp b/Engine/Plugins/2D/Paper2D/Source/PaperSpriteSheetImporter/Private/PaperSpriteSheetImportFactory.cpp
--- a/Engine/Plugins/2D/Paper2D/Source/PaperSpriteSheetImporter/Private/PaperSpriteSheetImportFactory.cpp
+++ b/Engine/Plugins/2D/Paper2D/Source/PaperSpriteSheetImporter/Private/PaperSpriteSheetImportFactory.cpp
@@ -8,7 +8,6 @@
 #include "AssetRegistryModule.h"
 #include "PackageTools.h"
 #include "PaperSpriteSheet.h"
-#include "PaperJsonSpriteSheetImporter.h"
 #include "PaperSpriteSheetImportFactory.h"
 #include "PaperSpriteSheetReimportFactory.h"
 
@@ -56,14 +55,6 @@ UObject* UPaperSpriteSheetImportFactory::FactoryCreateText(UClass* InClass, UObj
 	const FString NameForErrors(InName.ToString());
 	const FString FileContent(BufferEnd - Buffer, Buffer);
 	
-	FPaperJsonSpriteSheetImporter Importer;
-
-	// Are we reimporting?
-	if (UPaperSpriteSheetReimportFactory* CurrentReimportFactory = Cast<UPaperSpriteSheetReimportFactory>(this))
-	{
-		Importer.SetReimportData(CurrentReimportFactory->ExistingTextureName, CurrentReimportFactory->ExistingTexture, CurrentReimportFactory->ExistingSpriteNames, CurrentReimportFactory->ExistingSprites);
-	}
-
 	if (Importer.ImportFromString(FileContent, NameForErrors) &&
 		Importer.ImportTextures(LongPackagePath, CurrentSourcePath))
 	{
diff --git a/Engine/Plugins/2D/Paper2D/Source/PaperSpriteSheetImporter/Private/PaperSpriteSheetImportFactory.h b/Engine/Plugins/2D/Paper2D/Source/PaperSpriteSheetImporter/Private/PaperSpriteSheetImportFactory.h
new file mode 100644
--- /dev/null
+++ b/Engine/Plugins/2D/Paper2D/Source/PaperSpriteSheetImporter/Private/PaperSpriteSheetImportFactory.h
@@ -0,0 +1,23 @@
+// Copyright 1998-2015 Epic Games, Inc. All Rights Reserved.
+
+#pragma once
+
+#include "PaperJsonSpriteSheetImporter.h"
+#include "PaperSpriteSheetImportFactory.generated.h"
+
+// Imports a sprite sheet (and associated paper sprites and textures) from a JSON file exported from Adobe Flash CS6, Texture Packer, or other compatible tool
+UCLASS()
+class UPaperSpriteSheetImportFactory : public UFactory
+{
+	GENERATED_UCLASS_BODY()
+
+	// UFactory interface
+	virtual FText GetToolTip() const override;
+	virtual UObject* FactoryCreateText(UClass* InClass, UObject* InParent, FName InName, EObjectFlags Flags, UObject* Context, const TCHAR* Type, const TCHAR*& Buffer, const TCHAR* BufferEnd, FFeedbackContext* Warn) override;
+	// End of UFactory interface
+
+protected:
+	// The actual import worker, which may be already configured by the more derived reimport factory by the time FactoryCreateText is called
+	FPaperJsonSpriteSheetImporter Importer;
+};
+
diff --git a/Engine/Plugins/2D/Paper2D/Source/PaperSpriteSheetImporter/Private/PaperSpriteSheetReimportFactory.cpp b/Engine/Plugins/2D/Paper2D/Source/PaperSpriteSheetImporter/Private/PaperSpriteSheetReimportFactory.cpp
--- a/Engine/Plugins/2D/Paper2D/Source/PaperSpriteSheetImporter/Private/PaperSpriteSheetReimportFactory.cpp
+++ b/Engine/Plugins/2D/Paper2D/Source/PaperSpriteSheetImporter/Private/PaperSpriteSheetReimportFactory.cpp
@@ -15,7 +15,6 @@ UPaperSpriteSheetReimportFactory::UPaperSpriteSheetReimportFactory(const FObject
 	SupportedClass = UPaperSpriteSheet::StaticClass();
 
 	bCreateNew = false;
-	bText = true;
 }
 
 bool UPaperSpriteSheetReimportFactory::CanReimport(UObject* Obj, TArray<FString>& OutFilenames)
@@ -53,11 +52,14 @@ EReimportResult::Type UPaperSpriteSheetReimportFactory::Reimport(UObject* Obj)
 		return EReimportResult::Failed;
 	}
 
-	ExistingSpriteNames = SpriteSheet->SpriteNames;
-	ExistingSprites = SpriteSheet->Sprites;
-	ExistingTextureName = SpriteSheet->TextureName;
-	ExistingTexture = SpriteSheet->Texture;
+	// Configure the importer with the reimport settings
+	Importer.SetReimportData(SpriteSheet->SpriteNames, SpriteSheet->Sprites);
+	Importer.ExistingBaseTextureName = SpriteSheet->TextureName;
+	Importer.ExistingBaseTexture = SpriteSheet->Texture;
+	Importer.ExistingNormalMapTextureName = SpriteSheet->NormalMapTextureName;
+	Importer.ExistingNormalMapTexture = SpriteSheet->NormalMapTexture;
 
+	// Run the import again
 	EReimportResult::Type Result = EReimportResult::Failed;
 	if (UFactory::StaticImportObject(SpriteSheet->GetClass(), SpriteSheet->GetOuter(), *SpriteSheet->GetName(), RF_Public | RF_Standalone, *Filename, nullptr, this))
 	{
@@ -87,4 +89,6 @@ int32 UPaperSpriteSheetReimportFactory::GetPriority() const
 	return ImportPriority;
 }
 
+//////////////////////////////////////////////////////////////////////////
+
 #undef LOCTEXT_NAMESPACE
diff --git a/Engine/Plugins/2D/Paper2D/Source/PaperSpriteSheetImporter/Private/PaperSpriteSheetReimportFactory.h b/Engine/Plugins/2D/Paper2D/Source/PaperSpriteSheetImporter/Private/PaperSpriteSheetReimportFactory.h
--- a/Engine/Plugins/2D/Paper2D/Source/PaperSpriteSheetImporter/Private/PaperSpriteSheetReimportFactory.h
+++ b/Engine/Plugins/2D/Paper2D/Source/PaperSpriteSheetImporter/Private/PaperSpriteSheetReimportFactory.h
@@ -5,7 +5,7 @@
 #include "PaperSpriteSheetImportFactory.h"
 #include "PaperSpriteSheetReimportFactory.generated.h"
 
-// Imports a paper animated sprite (and associated paper sprites and textures) from an Adobe Flash CS6 exported JSON file
+// Reimports a UPaperSpriteSheet asset
 UCLASS()
 class UPaperSpriteSheetReimportFactory : public UPaperSpriteSheetImportFactory, public FReimportHandler
 {
@@ -17,11 +17,4 @@ class UPaperSpriteSheetReimportFactory : public UPaperSpriteSheetImportFactory,
 	virtual EReimportResult::Type Reimport(UObject* Obj) override;
 	virtual int32 GetPriority() const override;
 	// End FReimportHandler interface
-
-	TArray<FString> ExistingSpriteNames;
-	TArray< TAssetPtr<class UPaperSprite> > ExistingSprites;
-	
-	FString ExistingTextureName;
-	UTexture2D* ExistingTexture;
 };
-
