Paper2D: Removed all uses of platform-specific texture data for sprite/tile map/etc.. building, it now always works off of the imported size, since that is the authoring size in the editor [UE-14171]

diff --git a/Engine/Plugins/2D/Paper2D/Source/Paper2D/Classes/SpriteEditorOnlyTypes.h b/Engine/Plugins/2D/Paper2D/Source/Paper2D/Classes/SpriteEditorOnlyTypes.h
--- a/Engine/Plugins/2D/Paper2D/Source/Paper2D/Classes/SpriteEditorOnlyTypes.h
+++ b/Engine/Plugins/2D/Paper2D/Source/Paper2D/Classes/SpriteEditorOnlyTypes.h
@@ -268,7 +268,7 @@ struct FSpriteAssetInitParameters
 		Texture = InTexture;
 		if (Texture != nullptr)
 		{
-			Dimension = FVector2D(Texture->GetSizeX(), Texture->GetSizeY());
+			Dimension = FVector2D(Texture->GetImportedSize());
 			Offset = FVector2D::ZeroVector;
 		}
 		else
diff --git a/Engine/Plugins/2D/Paper2D/Source/Paper2D/Private/PaperSprite.cpp b/Engine/Plugins/2D/Paper2D/Source/Paper2D/Private/PaperSprite.cpp
--- a/Engine/Plugins/2D/Paper2D/Source/Paper2D/Private/PaperSprite.cpp
+++ b/Engine/Plugins/2D/Paper2D/Source/Paper2D/Private/PaperSprite.cpp
@@ -520,8 +520,8 @@ void UPaperSprite::PostEditChangeProperty(FPropertyChangedEvent& PropertyChanged
 		{
 			// If this is a brand new sprite that didn't have a texture set previously, act like we were factoried with the texture
 			SourceUV = FVector2D::ZeroVector;
-			SourceDimension = FVector2D(SourceTexture->GetSizeX(), SourceTexture->GetSizeY());
-			SourceTextureDimension = FVector2D(SourceTexture->GetSizeX(), SourceTexture->GetSizeY());
+			SourceDimension = FVector2D(SourceTexture->GetImportedSize());
+			SourceTextureDimension = SourceDimension;
 		}
 		bBothModified = true;
 	}
@@ -686,9 +686,9 @@ bool UPaperSprite::NeedRescaleSpriteData()
 	if (UTexture2D* Texture = GetSourceTexture())
 	{
 		Texture->ConditionalPostLoad();
-		FIntPoint TextureSize = Texture->GetImportedSize();
-		bool bTextureSizeIsZero = TextureSize.X == 0 || TextureSize.Y == 0;
-		return !SourceTextureDimension.IsZero() && !bTextureSizeIsZero && (TextureSize.X != SourceTextureDimension.X || TextureSize.Y != SourceTextureDimension.Y);
+		const FIntPoint TextureSize = Texture->GetImportedSize();
+		const bool bTextureSizeIsZero = (TextureSize.X == 0) || (TextureSize.Y == 0);
+		return !SourceTextureDimension.IsZero() && !bTextureSizeIsZero && ((TextureSize.X != SourceTextureDimension.X) || (TextureSize.Y != SourceTextureDimension.Y));
 	}
 
 	return false;
@@ -831,12 +831,7 @@ void UPaperSprite::RebuildRenderData()
 	if (EffectiveTexture)
 	{
 		EffectiveTexture->ConditionalPostLoad();
-		const int32 TextureWidth = EffectiveTexture->GetSizeX();
-		const int32 TextureHeight = EffectiveTexture->GetSizeY();
-		if (ensure((TextureWidth > 0) && (TextureHeight > 0)))
-		{
-			TextureSize = FVector2D(TextureWidth, TextureHeight);
-		}
+		TextureSize = FVector2D(EffectiveTexture->GetImportedSize());
 	}
 	const float InverseWidth = 1.0f / TextureSize.X;
 	const float InverseHeight = 1.0f / TextureSize.Y;
@@ -1784,6 +1779,11 @@ void UPaperSprite::PostLoad()
 	Super::PostLoad();
 
 #if WITH_EDITORONLY_DATA
+	if (UTexture2D* EffectiveTexture = GetBakedTexture())
+	{
+		EffectiveTexture->ConditionalPostLoad();
+	}
+	
 	const int32 PaperVer = GetLinkerCustomVersion(FPaperCustomVersion::GUID);
 
 	bool bRebuildCollision = false;
diff --git a/Engine/Plugins/2D/Paper2D/Source/Paper2D/Private/PaperTileMapComponent.cpp b/Engine/Plugins/2D/Paper2D/Source/Paper2D/Private/PaperTileMapComponent.cpp
--- a/Engine/Plugins/2D/Paper2D/Source/Paper2D/Private/PaperTileMapComponent.cpp
+++ b/Engine/Plugins/2D/Paper2D/Source/Paper2D/Private/PaperTileMapComponent.cpp
@@ -305,7 +305,8 @@ void UPaperTileMapComponent::RebuildRenderData(FPaperTileMapRenderSceneProxy* Pr
 
 					if (SourceTexture != LastSourceTexture)
 					{
-						InverseTextureSize = FVector2D(1.0f / SourceTexture->GetSizeX(), 1.0f / SourceTexture->GetSizeY());
+						const FVector2D TextureSize(SourceTexture->GetImportedSize());
+						InverseTextureSize = FVector2D(1.0f / TextureSize.X, 1.0f / TextureSize.Y);
 
 						if (TileInfo.TileSet != nullptr)
 						{
diff --git a/Engine/Plugins/2D/Paper2D/Source/Paper2D/Private/PaperTileSet.cpp b/Engine/Plugins/2D/Paper2D/Source/Paper2D/Private/PaperTileSet.cpp
--- a/Engine/Plugins/2D/Paper2D/Source/Paper2D/Private/PaperTileSet.cpp
+++ b/Engine/Plugins/2D/Paper2D/Source/Paper2D/Private/PaperTileSet.cpp
@@ -22,8 +22,9 @@ int32 UPaperTileSet::GetTileCount() const
 	if (TileSheet != nullptr)
 	{
 		checkSlow((TileWidth > 0) && (TileHeight > 0));
-		const int32 TextureWidth = TileSheet->GetSizeX();
-		const int32 TextureHeight = TileSheet->GetSizeY();
+		const FIntPoint TextureSize = TileSheet->GetImportedSize();
+		const int32 TextureWidth = TextureSize.X;
+		const int32 TextureHeight = TextureSize.Y;
 
 		const int32 CellsX = (TextureWidth - (Margin * 2) + Spacing) / (TileWidth + Spacing);
 		const int32 CellsY = (TextureHeight - (Margin * 2) + Spacing) / (TileHeight + Spacing);
@@ -41,7 +42,7 @@ int32 UPaperTileSet::GetTileCountX() const
 	if (TileSheet != nullptr)
 	{
 		checkSlow(TileWidth > 0);
-		const int32 TextureWidth = TileSheet->GetSizeX();
+		const int32 TextureWidth = TileSheet->GetImportedSize().X;
 		const int32 CellsX = (TextureWidth - (Margin * 2) + Spacing) / (TileWidth + Spacing);
 		return CellsX;
 	}
@@ -56,7 +57,7 @@ int32 UPaperTileSet::GetTileCountY() const
 	if (TileSheet != nullptr)
 	{
 		checkSlow(TileHeight > 0);
-		const int32 TextureHeight = TileSheet->GetSizeY();
+		const int32 TextureHeight = TileSheet->GetImportedSize().Y;
 		const int32 CellsY = (TextureHeight - (Margin * 2) + Spacing) / (TileHeight + Spacing);
 		return CellsY;
 	}
diff --git a/Engine/Plugins/2D/Paper2D/Source/Paper2DEditor/Private/Atlasing/PaperAtlasGenerator.cpp b/Engine/Plugins/2D/Paper2D/Source/Paper2DEditor/Private/Atlasing/PaperAtlasGenerator.cpp
--- a/Engine/Plugins/2D/Paper2D/Source/Paper2DEditor/Private/Atlasing/PaperAtlasGenerator.cpp
+++ b/Engine/Plugins/2D/Paper2D/Source/Paper2DEditor/Private/Atlasing/PaperAtlasGenerator.cpp
@@ -185,11 +185,11 @@ void FPaperAtlasGenerator::HandleAssetChangedEvent(UPaperSpriteAtlas* Atlas)
 
 		// An atlas is ALSO forced dirty if the dimensions have changed
 		// We're just grabbing a fixed atlas size here for now
-		int32 AtlasWidth = Atlas->MaxWidth;
-		int32 AtlasHeight = Atlas->MaxHeight;
-		int32 BytesPerPixel = sizeof(FColor);
+		const int32 AtlasWidth = Atlas->MaxWidth;
+		const int32 AtlasHeight = Atlas->MaxHeight;
+		const int32 BytesPerPixel = sizeof(FColor);
 
-		bool bAtlasDirty = RemappedAtlasForceDirty[AtlasIndex] || AtlasTexture->GetSizeX() != AtlasWidth || AtlasTexture->GetSizeY() != AtlasHeight;
+		const bool bAtlasDirty = RemappedAtlasForceDirty[AtlasIndex] || (AtlasTexture->GetImportedSize() != FIntPoint(AtlasWidth, AtlasHeight));
 
 		// Propagate texture settings
 		AtlasTexture->CompressionSettings = Atlas->CompressionSettings;
diff --git a/Engine/Plugins/2D/Paper2D/Source/Paper2DEditor/Private/ExtractSprites/SPaperExtractSpritesDialog.cpp b/Engine/Plugins/2D/Paper2D/Source/Paper2DEditor/Private/ExtractSprites/SPaperExtractSpritesDialog.cpp
--- a/Engine/Plugins/2D/Paper2D/Source/Paper2DEditor/Private/ExtractSprites/SPaperExtractSpritesDialog.cpp
+++ b/Engine/Plugins/2D/Paper2D/Source/Paper2DEditor/Private/ExtractSprites/SPaperExtractSpritesDialog.cpp
@@ -146,8 +146,8 @@ void SPaperExtractSpritesDialog::Construct(const FArguments& InArgs, UTexture2D*
 	ExtractSpriteGridSettings = NewObject<UPaperExtractSpriteGridSettings>();
 	ExtractSpriteGridSettings->AddToRoot();
 	ExtractSpriteGridSettings->NamingTemplate = "Sprite_{0}";
-	ExtractSpriteGridSettings->CellWidth = InSourceTexture->GetSizeX();
-	ExtractSpriteGridSettings->CellHeight = InSourceTexture->GetSizeY();
+	ExtractSpriteGridSettings->CellWidth = InSourceTexture->GetImportedSize().X;
+	ExtractSpriteGridSettings->CellHeight = InSourceTexture->GetImportedSize().Y;
 
 	PreviewExtractedSprites();
 
@@ -340,8 +340,11 @@ void SPaperExtractSpritesDialog::PreviewExtractedSprites()
 		if (SourceTexture != nullptr)
 		{
 			int32 ExtractedRectIndex = ExtractSpriteGridSettings->NamingStartIndex;
-			int32 TextureWidth = SourceTexture->GetSizeX();
-			int32 TextureHeight = SourceTexture->GetSizeY();
+
+			const FIntPoint TextureSize(SourceTexture->GetImportedSize());
+			const int32 TextureWidth = TextureSize.X;
+			const int32 TextureHeight = TextureSize.Y;
+
 			int NumExtractedCellsY = 0;
 			for (int32 Y = ExtractSpriteGridSettings->MarginY; Y + ExtractSpriteGridSettings->CellHeight <= TextureHeight; Y += ExtractSpriteGridSettings->CellHeight + ExtractSpriteGridSettings->SpacingY)
 			{
diff --git a/Engine/Plugins/2D/Paper2D/Source/Paper2DEditor/Private/SpriteEditor/SpriteEditorViewportClient.cpp b/Engine/Plugins/2D/Paper2D/Source/Paper2DEditor/Private/SpriteEditor/SpriteEditorViewportClient.cpp
--- a/Engine/Plugins/2D/Paper2D/Source/Paper2DEditor/Private/SpriteEditor/SpriteEditorViewportClient.cpp
+++ b/Engine/Plugins/2D/Paper2D/Source/Paper2DEditor/Private/SpriteEditor/SpriteEditorViewportClient.cpp
@@ -1125,8 +1125,9 @@ bool FSpriteEditorViewportClient::ConvertMarqueeToSourceTextureSpace(/*out*/FVec
 			Swap(TextureSpaceStartPos.Y, TextureSpaceEndPos.Y);
 		}
 
-		const int32 SourceTextureWidth = SpriteSourceTexture->GetSizeX();
-		const int32 SourceTextureHeight = SpriteSourceTexture->GetSizeY();
+		const FIntPoint SourceTextureSize(SpriteSourceTexture->GetImportedSize());
+		const int32 SourceTextureWidth = SourceTextureSize.X;
+		const int32 SourceTextureHeight = SourceTextureSize.Y;
 		TextureSpaceStartPos.X = FMath::Clamp((int)TextureSpaceStartPos.X, 0, SourceTextureWidth - 1);
 		TextureSpaceStartPos.Y = FMath::Clamp((int)TextureSpaceStartPos.Y, 0, SourceTextureHeight - 1);
 		TextureSpaceEndPos.X = FMath::Clamp((int)TextureSpaceEndPos.X, 0, SourceTextureWidth - 1);
