Paper2D: Started to relax the assumption that sprite textures are always UTexture2D

diff --git a/Engine/Plugins/2D/Paper2D/Source/Paper2D/Classes/PaperSprite.h b/Engine/Plugins/2D/Paper2D/Source/Paper2D/Classes/PaperSprite.h
--- a/Engine/Plugins/2D/Paper2D/Source/Paper2D/Classes/PaperSprite.h
+++ b/Engine/Plugins/2D/Paper2D/Source/Paper2D/Classes/PaperSprite.h
@@ -22,7 +22,7 @@ struct FPaperSpriteSocket
 	FName SocketName;
 };
 
-typedef TArray<class UTexture2D*, TInlineAllocator<4>> FAdditionalSpriteTextureArray;
+typedef TArray<class UTexture*, TInlineAllocator<4>> FAdditionalSpriteTextureArray;
 
 /**
  * Sprite Asset
@@ -81,7 +81,7 @@ protected:
 
 	// Additional source textures for other slots
 	UPROPERTY(Category=Sprite, EditAnywhere, AssetRegistrySearchable, meta=(DisplayName="Additional Textures"))
-	TArray<UTexture2D*> AdditionalSourceTextures;
+	TArray<UTexture*> AdditionalSourceTextures;
 
 	UPROPERTY()
 	UTexture2D* BakedSourceTexture;
diff --git a/Engine/Plugins/2D/Paper2D/Source/Paper2D/Classes/SpriteDrawCall.h b/Engine/Plugins/2D/Paper2D/Source/Paper2D/Classes/SpriteDrawCall.h
--- a/Engine/Plugins/2D/Paper2D/Source/Paper2D/Classes/SpriteDrawCall.h
+++ b/Engine/Plugins/2D/Paper2D/Source/Paper2D/Classes/SpriteDrawCall.h
@@ -16,7 +16,7 @@ public:
 	FVector Destination;
 
 	UPROPERTY(Category=Sprite, EditAnywhere)
-	UTexture2D* BaseTexture;
+	UTexture* BaseTexture;
 
 	FAdditionalSpriteTextureArray AdditionalTextures;
 
diff --git a/Engine/Plugins/2D/Paper2D/Source/Paper2D/Classes/SpriteEditorOnlyTypes.h b/Engine/Plugins/2D/Paper2D/Source/Paper2D/Classes/SpriteEditorOnlyTypes.h
--- a/Engine/Plugins/2D/Paper2D/Source/Paper2D/Classes/SpriteEditorOnlyTypes.h
+++ b/Engine/Plugins/2D/Paper2D/Source/Paper2D/Classes/SpriteEditorOnlyTypes.h
@@ -283,7 +283,7 @@ public:
 	UTexture2D* Texture;
 
 	// Additional textures to use
-	TArray<UTexture2D*> AdditionalTextures;
+	TArray<UTexture*> AdditionalTextures;
 
 	// The offset within the texture (in pixels)
 	FVector2D Offset;
diff --git a/Engine/Plugins/2D/Paper2D/Source/Paper2D/Private/PaperRenderSceneProxy.h b/Engine/Plugins/2D/Paper2D/Source/Paper2D/Private/PaperRenderSceneProxy.h
--- a/Engine/Plugins/2D/Paper2D/Source/Paper2D/Private/PaperRenderSceneProxy.h
+++ b/Engine/Plugins/2D/Paper2D/Source/Paper2D/Private/PaperRenderSceneProxy.h
@@ -41,7 +41,7 @@ struct PAPER2D_API FPaperSpriteVertex
 struct PAPER2D_API FSpriteRenderSection
 {
 	UMaterialInterface* Material;
-	UTexture2D* BaseTexture;
+	UTexture* BaseTexture;
 	FAdditionalSpriteTextureArray AdditionalTextures;
 
 	int32 VertexOffset;
