Tile Maps: Fixed a crash when clearing the texture reference or choosing a texture that is smaller than the tile size [UE-13992]

diff --git a/Engine/Plugins/2D/Paper2D/Source/Paper2D/Private/PaperTileSet.cpp b/Engine/Plugins/2D/Paper2D/Source/Paper2D/Private/PaperTileSet.cpp
--- a/Engine/Plugins/2D/Paper2D/Source/Paper2D/Private/PaperTileSet.cpp
+++ b/Engine/Plugins/2D/Paper2D/Source/Paper2D/Private/PaperTileSet.cpp
@@ -190,8 +190,6 @@ void UPaperTileSet::PostEditChangeProperty(FPropertyChangedEvent& PropertyChange
 
 void UPaperTileSet::DestructiveAllocateTileData(int32 NewWidth, int32 NewHeight)
 {
-	check((NewWidth > 0) && (NewHeight > 0));
-
 	const int32 NumCells = NewWidth * NewHeight;
 	ExperimentalPerTileData.Empty(NumCells);
 	for (int32 Index = 0; Index < NumCells; ++Index)
