Tile Maps: Added a details customization for tile set assets, displays the # of tiles and warnings if the TileSheet is too small or missing [UE-14273]

diff --git a/Engine/Plugins/2D/Paper2D/Source/Paper2DEditor/Private/Paper2DEditorModule.cpp b/Engine/Plugins/2D/Paper2D/Source/Paper2DEditor/Private/Paper2DEditorModule.cpp
--- a/Engine/Plugins/2D/Paper2D/Source/Paper2DEditor/Private/Paper2DEditorModule.cpp
+++ b/Engine/Plugins/2D/Paper2D/Source/Paper2DEditor/Private/Paper2DEditorModule.cpp
@@ -34,6 +34,7 @@
 #include "TileSetAssetTypeActions.h"
 #include "PaperTileSetThumbnailRenderer.h"
 #include "TileSetEditor/TileSetEditorSettings.h"
+#include "TileSetEditor/TileSetDetailsCustomization.h"
 
 // Tile map support
 #include "TileMapEditing/TileMapAssetTypeActions.h"
@@ -129,6 +130,7 @@ public:
 			FPropertyEditorModule& PropertyModule = FModuleManager::LoadModuleChecked<FPropertyEditorModule>("PropertyEditor");
 			PropertyModule.RegisterCustomClassLayout(UPaperTileMapComponent::StaticClass()->GetFName(), FOnGetDetailCustomizationInstance::CreateStatic(&FPaperTileMapDetailsCustomization::MakeInstance));
 			PropertyModule.RegisterCustomClassLayout(UPaperTileMap::StaticClass()->GetFName(), FOnGetDetailCustomizationInstance::CreateStatic(&FPaperTileMapDetailsCustomization::MakeInstance));
+			PropertyModule.RegisterCustomClassLayout(UPaperTileSet::StaticClass()->GetFName(), FOnGetDetailCustomizationInstance::CreateStatic(&FTileSetDetailsCustomization::MakeInstance));
 			PropertyModule.RegisterCustomClassLayout(UPaperSprite::StaticClass()->GetFName(), FOnGetDetailCustomizationInstance::CreateStatic(&FSpriteDetailsCustomization::MakeInstance));
 			PropertyModule.RegisterCustomClassLayout(UPaperSpriteComponent::StaticClass()->GetFName(), FOnGetDetailCustomizationInstance::CreateStatic(&FSpriteComponentDetailsCustomization::MakeInstance));
 			PropertyModule.RegisterCustomClassLayout(UPaperFlipbookComponent::StaticClass()->GetFName(), FOnGetDetailCustomizationInstance::CreateStatic(&FFlipbookComponentDetailsCustomization::MakeInstance));
diff --git a/Engine/Plugins/2D/Paper2D/Source/Paper2DEditor/Private/TileSetEditor/TileSetDetailsCustomization.cpp b/Engine/Plugins/2D/Paper2D/Source/Paper2DEditor/Private/TileSetEditor/TileSetDetailsCustomization.cpp
new file mode 100644
--- /dev/null
+++ b/Engine/Plugins/2D/Paper2D/Source/Paper2DEditor/Private/TileSetEditor/TileSetDetailsCustomization.cpp
@@ -0,0 +1,128 @@
+// Copyright 1998-2015 Epic Games, Inc. All Rights Reserved.
+
+#include "Paper2DEditorPrivatePCH.h"
+#include "TileSetDetailsCustomization.h"
+#include "DetailLayoutBuilder.h"
+#include "DetailCategoryBuilder.h"
+#include "PropertyEditing.h"
+#include "ScopedTransaction.h"
+#include "IPropertyUtilities.h"
+
+#define LOCTEXT_NAMESPACE "Paper2D"
+
+//////////////////////////////////////////////////////////////////////////
+// FTileSetDetailsCustomization
+
+TSharedRef<IDetailCustomization> FTileSetDetailsCustomization::MakeInstance()
+{
+	return MakeShareable(new FTileSetDetailsCustomization);
+}
+
+void FTileSetDetailsCustomization::CustomizeDetails(IDetailLayoutBuilder& DetailLayout)
+{
+	MyDetailLayout = &DetailLayout;
+	
+	for (const TWeakObjectPtr<UObject> SelectedObject : DetailLayout.GetDetailsView().GetSelectedObjects())
+	{
+		if (UPaperTileSet* TileSet = Cast<UPaperTileSet>(SelectedObject.Get()))
+		{
+			TileSetPtr = TileSet;
+			break;
+		}
+	}
+
+ 	IDetailCategoryBuilder& TileSetCategory = DetailLayout.EditCategory("TileSet", FText::GetEmpty());
+
+	// Add the width and height in cells of this tile set to the header
+	TileSetCategory.HeaderContent
+	(
+		SNew(SBox)
+		.HAlign(HAlign_Right)
+		[
+			SNew(SHorizontalBox)
+			+SHorizontalBox::Slot()
+			.Padding(FMargin(5.0f, 0.0f))
+			.AutoWidth()
+			[
+				SNew(STextBlock)
+				.Font(FEditorStyle::GetFontStyle("TinyText"))
+				.Text(this, &FTileSetDetailsCustomization::GetCellDimensionHeaderText)
+				.ColorAndOpacity(this, &FTileSetDetailsCustomization::GetCellDimensionHeaderColor)
+				.ToolTipText(LOCTEXT("NumCellsTooltip", "Number of tile cells in this tile set"))
+			]
+		]
+	);
+
+
+	// Add all of the properties from the inline tilemap
+// 	if ((TileComponent != nullptr) && (TileComponent->OwnsTileMap()))
+// 	{
+// 		TArray<UObject*> ListOfTileMaps;
+// 		ListOfTileMaps.Add(TileMap);
+// 
+// 		for (TFieldIterator<UProperty> PropIt(UPaperTileMap::StaticClass()); PropIt; ++PropIt)
+// 		{
+// 			UProperty* TestProperty = *PropIt;
+// 
+// 			if (TestProperty->HasAnyPropertyFlags(CPF_Edit))
+// 			{
+// 				FName CategoryName(*TestProperty->GetMetaData(TEXT("Category")));
+// 				IDetailCategoryBuilder& Category = DetailLayout.EditCategory(CategoryName);
+// 
+// 				if (IDetailPropertyRow* ExternalRow = Category.AddExternalProperty(ListOfTileMaps, TestProperty->GetFName()))
+// 				{
+// 					ExternalRow->Visibility(InternalInstanceVis);
+// 				}
+// 			}
+// 		}
+// 	}
+}
+
+FText FTileSetDetailsCustomization::GetCellDimensionHeaderText() const
+{
+	FText Result;
+
+	if (UPaperTileSet* TileSet = TileSetPtr.Get())
+	{
+		const int32 NumTilesX = TileSet->GetTileCountX();
+		const int32 NumTilesY = TileSet->GetTileCountY();
+
+		if (TileSet->TileSheet == nullptr)
+		{
+			Result = LOCTEXT("NoTexture", "No TileSheet");
+		}
+		else if (NumTilesX == 0)
+		{
+			Result = LOCTEXT("TextureTooNarrow", "TileSheet too narrow");
+		}
+		else if (NumTilesY == 0)
+		{
+			Result = LOCTEXT("TextureTooShort", "TileSheet too short");
+		}
+		else
+		{
+			const FText TextNumTilesX = FText::AsNumber(NumTilesX, &FNumberFormattingOptions::DefaultNoGrouping());
+			const FText TextNumTilesY = FText::AsNumber(NumTilesY, &FNumberFormattingOptions::DefaultNoGrouping());
+			Result = FText::Format(LOCTEXT("CellDimensions", "{0} x {0} tiles"), TextNumTilesX, TextNumTilesY);
+		}
+	}
+
+	return Result;
+}
+
+FSlateColor FTileSetDetailsCustomization::GetCellDimensionHeaderColor() const
+{
+	if (UPaperTileSet* TileSet = TileSetPtr.Get())
+	{
+		if (TileSet->GetTileCount() == 0)
+		{
+			return FSlateColor(FLinearColor::Red);
+		}
+	}
+
+	return FSlateColor::UseForeground();
+}
+
+//////////////////////////////////////////////////////////////////////////
+
+#undef LOCTEXT_NAMESPACE
\ No newline at end of file
diff --git a/Engine/Plugins/2D/Paper2D/Source/Paper2DEditor/Private/TileSetEditor/TileSetDetailsCustomization.h b/Engine/Plugins/2D/Paper2D/Source/Paper2DEditor/Private/TileSetEditor/TileSetDetailsCustomization.h
new file mode 100644
--- /dev/null
+++ b/Engine/Plugins/2D/Paper2D/Source/Paper2DEditor/Private/TileSetEditor/TileSetDetailsCustomization.h
@@ -0,0 +1,28 @@
+// Copyright 1998-2015 Epic Games, Inc. All Rights Reserved.
+
+#pragma once
+
+#include "PropertyEditorModule.h"
+
+//////////////////////////////////////////////////////////////////////////
+// FTileSetDetailsCustomization
+
+class FTileSetDetailsCustomization : public IDetailCustomization
+{
+public:
+	// Makes a new instance of this detail layout class for a specific detail view requesting it
+	static TSharedRef<IDetailCustomization> MakeInstance();
+
+	// IDetailCustomization interface
+	virtual void CustomizeDetails(IDetailLayoutBuilder& DetailLayout) override;
+	// End of IDetailCustomization interface
+
+private:
+	TWeakObjectPtr<class UPaperTileSet> TileSetPtr;
+
+	IDetailLayoutBuilder* MyDetailLayout;
+
+private:
+	FText GetCellDimensionHeaderText() const;
+	FSlateColor GetCellDimensionHeaderColor() const;
+};
