Tile Maps: Added the ability to make a tile map component editable even if it is using a tile map asset by calling MakeTileMapEditable [UE-14087]

diff --git a/Engine/Plugins/2D/Paper2D/Source/Paper2D/Classes/PaperTileMap.h b/Engine/Plugins/2D/Paper2D/Source/Paper2D/Classes/PaperTileMap.h
--- a/Engine/Plugins/2D/Paper2D/Source/Paper2D/Classes/PaperTileMap.h
+++ b/Engine/Plugins/2D/Paper2D/Source/Paper2D/Classes/PaperTileMap.h
@@ -172,6 +172,8 @@ public:
 	// Adds a default layer and pulls the PixelsPerUnrealUnit from the project settings
 	void InitializeNewEmptyTileMap();
 
+	// Creates a clone of this tile map in the specified outer
+	UPaperTileMap* CloneTileMap(UObject* OuterForClone);
 protected:
 	virtual void UpdateBodySetup();
 };
diff --git a/Engine/Plugins/2D/Paper2D/Source/Paper2D/Classes/PaperTileMapComponent.h b/Engine/Plugins/2D/Paper2D/Source/Paper2D/Classes/PaperTileMapComponent.h
--- a/Engine/Plugins/2D/Paper2D/Source/Paper2D/Classes/PaperTileMapComponent.h
+++ b/Engine/Plugins/2D/Paper2D/Source/Paper2D/Classes/PaperTileMapComponent.h
@@ -128,11 +128,16 @@ public:
 
 	// Sets the per-layer color multiplier for a specific layer (multiplied with the tile map color and passed to the material as a vertex color)
 	// Note: This will only work on components that own their own tile map (OwnsTileMap returns true), you cannot modify standalone tile map assets
-	UFUNCTION(BlueprintCallable, Category = "Sprite")
+	UFUNCTION(BlueprintCallable, Category="Sprite")
 	void SetLayerColor(FLinearColor NewColor, int32 Layer = 0);
 
 	// Returns the wireframe color to use for this component.
 	FLinearColor GetWireframeColor() const;
+
+	// Makes the tile map asset pointed to by this component editable.  Nothing happens if it was already instanced, but
+	// if the tile map is an asset reference, it is cloned to make a unique instance.
+	UFUNCTION(BlueprintCallable, Category="Sprite")
+	void MakeTileMapEditable();
 };
 
 // Allow the old name to continue to work for one release
diff --git a/Engine/Plugins/2D/Paper2D/Source/Paper2D/Private/PaperTileMap.cpp b/Engine/Plugins/2D/Paper2D/Source/Paper2D/Private/PaperTileMap.cpp
--- a/Engine/Plugins/2D/Paper2D/Source/Paper2D/Private/PaperTileMap.cpp
+++ b/Engine/Plugins/2D/Paper2D/Source/Paper2D/Private/PaperTileMap.cpp
@@ -402,6 +402,11 @@ void UPaperTileMap::InitializeNewEmptyTileMap()
 	PixelsPerUnrealUnit = GetDefault<UPaperRuntimeSettings>()->DefaultPixelsPerUnrealUnit;
 }
 
+UPaperTileMap* UPaperTileMap::CloneTileMap(UObject* OuterForClone)
+{
+	return CastChecked<UPaperTileMap>(StaticDuplicateObject(this, OuterForClone, nullptr));
+}
+
 //////////////////////////////////////////////////////////////////////////
 
 #undef LOCTEXT_NAMESPACE
\ No newline at end of file
diff --git a/Engine/Plugins/2D/Paper2D/Source/Paper2D/Private/PaperTileMapComponent.cpp b/Engine/Plugins/2D/Paper2D/Source/Paper2D/Private/PaperTileMapComponent.cpp
--- a/Engine/Plugins/2D/Paper2D/Source/Paper2D/Private/PaperTileMapComponent.cpp
+++ b/Engine/Plugins/2D/Paper2D/Source/Paper2D/Private/PaperTileMapComponent.cpp
@@ -505,4 +505,12 @@ FLinearColor UPaperTileMapComponent::GetWireframeColor() const
 	return TileMapColor;
 }
 
+void UPaperTileMapComponent::MakeTileMapEditable()
+{
+	if ((TileMap != nullptr) && !OwnsTileMap())
+	{
+		SetTileMap(TileMap->CloneTileMap(this));
+	}
+}
+
 #undef LOCTEXT_NAMESPACE
\ No newline at end of file
