Tile Maps: Details customization UX improvements
- Added a 'convert to instance' option
- Added undo/redo support for 'promote to asset' and 'convert to instance'
- Added a content browser sync when promoting to an asset
- Improved the 'in level editor' versus 'in standalone editor' detection so it doesn't misbehave when something is open in each
- Made the buttons a uniform size and added a wrap box around them
[UE-14270]

diff --git a/Engine/Plugins/2D/Paper2D/Source/Paper2DEditor/Private/PaperTileMapPromotionFactory.cpp b/Engine/Plugins/2D/Paper2D/Source/Paper2DEditor/Private/PaperTileMapPromotionFactory.cpp
--- a/Engine/Plugins/2D/Paper2D/Source/Paper2DEditor/Private/PaperTileMapPromotionFactory.cpp
+++ b/Engine/Plugins/2D/Paper2D/Source/Paper2DEditor/Private/PaperTileMapPromotionFactory.cpp
@@ -18,8 +18,9 @@ UPaperTileMapPromotionFactory::UPaperTileMapPromotionFactory(const FObjectInitia
 
 UObject* UPaperTileMapPromotionFactory::FactoryCreateNew(UClass* Class, UObject* InParent, FName Name, EObjectFlags Flags, UObject* Context, FFeedbackContext* Warn)
 {
-	AssetToRename->Rename(*Name.ToString(), InParent);
 	AssetToRename->SetFlags(Flags | RF_Transactional);
+	AssetToRename->Modify();
+	AssetToRename->Rename(*Name.ToString(), InParent);
 
 	return AssetToRename;
 }
diff --git a/Engine/Plugins/2D/Paper2D/Source/Paper2DEditor/Private/TileMapEditing/PaperTileMapDetailsCustomization.cpp b/Engine/Plugins/2D/Paper2D/Source/Paper2DEditor/Private/TileMapEditing/PaperTileMapDetailsCustomization.cpp
--- a/Engine/Plugins/2D/Paper2D/Source/Paper2DEditor/Private/TileMapEditing/PaperTileMapDetailsCustomization.cpp
+++ b/Engine/Plugins/2D/Paper2D/Source/Paper2DEditor/Private/TileMapEditing/PaperTileMapDetailsCustomization.cpp
@@ -66,68 +66,113 @@ void FPaperTileMapDetailsCustomization::CustomizeDetails(IDetailLayoutBuilder& D
 	IDetailCategoryBuilder& TileMapCategory = DetailLayout.EditCategory("Tile Map", FText::GetEmpty(), ECategoryPriority::Important);
 
 	// Make sure the setup category is near the top
-	DetailLayout.EditCategory("Setup", FText::GetEmpty(), ECategoryPriority::Important);
+	IDetailCategoryBuilder& SetupCategory = DetailLayout.EditCategory("Setup", FText::GetEmpty(), ECategoryPriority::Important);
+
+	// Add the 'instanced' versus 'asset' indicator to the tile map header
+	TileMapCategory.HeaderContent
+	(
+		SNew(SBox)
+		.HAlign(HAlign_Right)
+		[
+			SNew(SHorizontalBox)
+			+SHorizontalBox::Slot()
+			.Padding(FMargin(5.0f, 0.0f))
+			.AutoWidth()
+			[
+				SNew(STextBlock)
+				.Font(FEditorStyle::GetFontStyle("TinyText"))
+				.Text_Lambda([this] { return IsInstanced() ? LOCTEXT("Instanced", "Instanced") : LOCTEXT("Asset", "Asset"); })
+				.ToolTipText(LOCTEXT("InstancedVersusAssetTooltip", "Tile map components can either own a unique tile map instance, or reference a shareable tile map asset"))
+			]
+		]
+	);
 
 
 	TAttribute<EVisibility> InternalInstanceVis = TAttribute<EVisibility>::Create(TAttribute<EVisibility>::FGetter::CreateSP(this, &FPaperTileMapDetailsCustomization::GetVisibilityForInstancedOnlyProperties));
 
+	TSharedRef<SWrapBox> ButtonBox = SNew(SWrapBox).UseAllottedWidth(true);
+
+	const float MinButtonSize = 120.0f;
+	const FMargin ButtonPadding(0.0f, 2.0f, 2.0f, 0.0f);
+
+	// Edit tile map button
+	ButtonBox->AddSlot()
+	.Padding(ButtonPadding)
+	[
+		SNew(SBox)
+		.MinDesiredWidth(MinButtonSize)
+		[
+			SNew(SButton)
+			.VAlign(VAlign_Center)
+			.HAlign(HAlign_Center)
+			.OnClicked(this, &FPaperTileMapDetailsCustomization::EnterTileMapEditingMode)
+			.Visibility(this, &FPaperTileMapDetailsCustomization::GetNonEditModeVisibility)
+			.Text(LOCTEXT("EditAsset", "Edit Map"))
+			.ToolTipText(LOCTEXT("EditAssetToolTip", "Edit this tile map"))
+		]
+	];
+
+	// Create new tile map button
+	ButtonBox->AddSlot()
+	.Padding(ButtonPadding)
+	[
+		SNew(SBox)
+		.MinDesiredWidth(MinButtonSize)
+		[
+			SNew(SButton)
+			.VAlign(VAlign_Center)
+			.HAlign(HAlign_Center)
+			.OnClicked(this, &FPaperTileMapDetailsCustomization::OnNewButtonClicked)
+			.Visibility(this, &FPaperTileMapDetailsCustomization::GetNewButtonVisiblity)
+			.Text(LOCTEXT("CreateNewInstancedMap", "New Empty Map"))
+			.ToolTipText(LOCTEXT("CreateNewInstancedMapToolTip", "Create a new (instanced) tile map"))
+		]
+	];
+
+	// Promote to asset button
+	ButtonBox->AddSlot()
+	.Padding(ButtonPadding)
+	[
+		SNew(SBox)
+		.MinDesiredWidth(MinButtonSize)
+		[
+			SNew(SButton)
+			.VAlign(VAlign_Center)
+			.HAlign(HAlign_Center)
+			.OnClicked(this, &FPaperTileMapDetailsCustomization::OnPromoteToAssetButtonClicked)
+			.Visibility(InternalInstanceVis)
+			.Text(LOCTEXT("PromoteToAsset", "Promote To Asset"))
+			.ToolTipText(LOCTEXT("PromoteToAssetToolTip", "Save this tile map as a reusable asset"))
+		]
+	];
+
+	// Convert to instance button
+	ButtonBox->AddSlot()
+	.Padding(ButtonPadding)
+	[
+		SNew(SBox)
+		.MinDesiredWidth(MinButtonSize)
+		[
+			SNew(SButton)
+			.VAlign(VAlign_Center)
+			.HAlign(HAlign_Center)
+			.OnClicked(this, &FPaperTileMapDetailsCustomization::OnMakeInstanceFromAssetButtonClicked)
+ 			.Visibility(this, &FPaperTileMapDetailsCustomization::GetVisibilityForMakeIntoInstance)
+			.Text(LOCTEXT("ConvertToInstance", "Convert To Instance"))
+			.ToolTipText(LOCTEXT("ConvertToInstanceToolTip", "Copy the asset referenced by this tile map component into a unique instance that can be locally edited"))
+		]
+	];
+
 	if (TileComponent != nullptr)
 	{
 		TileMapCategory
 		.AddCustomRow(LOCTEXT( "TileMapInstancingControlsSearchText", "Edit Map New Empty Map Promote Asset"))
 		[
-			SNew(SVerticalBox)
-			+ SVerticalBox::Slot()
-			.Padding(0.0f, 2.0f, 0.0f, 0.0f)
-			.FillHeight(1.0f)
-			.VAlign(VAlign_Center)
+			SNew(SHorizontalBox)
+			+SHorizontalBox::Slot()
+			.FillWidth(1.0f)
 			[
-				SNew(SHorizontalBox)
-
-				// Edit button
-				+SHorizontalBox::Slot()
-				.AutoWidth()
-				.Padding( 2.0f, 0.0f )
-				.VAlign(VAlign_Center)
-				.HAlign(HAlign_Left)
-				[
-					SNew(SButton)
-					.VAlign(VAlign_Center)
-					.OnClicked(this, &FPaperTileMapDetailsCustomization::EnterTileMapEditingMode)
-					.Visibility(this, &FPaperTileMapDetailsCustomization::GetNonEditModeVisibility)
-					.Text( LOCTEXT("EditAsset", "Edit Map") )
-					.ToolTipText( LOCTEXT("EditAssetToolTip", "Edit this tile map") )
-				]
-
-				// Create new tile map button
-				+SHorizontalBox::Slot()
-				.AutoWidth()
-				.Padding( 2.0f, 0.0f )
-				.VAlign(VAlign_Center)
-				.HAlign(HAlign_Left)
-				[
-					SNew(SButton)
-					.VAlign(VAlign_Center)
-					.OnClicked(this, &FPaperTileMapDetailsCustomization::OnNewButtonClicked)
-					.Visibility(this, &FPaperTileMapDetailsCustomization::GetNewButtonVisiblity)
-					.Text(LOCTEXT("CreateNewInstancedMap", "New Empty Map"))
-					.ToolTipText( LOCTEXT("CreateNewInstancedMapToolTip", "Create a new (instanced) tile map") )
-				]
-
-				// Promote to asset button
-				+SHorizontalBox::Slot()
-				.AutoWidth()
-				.Padding( 2.0f, 0.0f )
-				.VAlign(VAlign_Center)
-				.HAlign(HAlign_Left)
-				[
-					SNew(SButton)
-					.VAlign(VAlign_Center)
-					.OnClicked(this, &FPaperTileMapDetailsCustomization::OnPromoteButtonClicked)
-					.Visibility(this, &FPaperTileMapDetailsCustomization::GetVisibilityForInstancedOnlyProperties)
-					.Text(LOCTEXT("PromoteToAsset", "Promote to asset"))
-					.ToolTipText(LOCTEXT("PromoteToAssetToolTip", "Save this tile map as a reusable asset"))
-				]
+				ButtonBox
 			]
 		];
 
@@ -224,7 +269,7 @@ FReply FPaperTileMapDetailsCustomization::OnNewButtonClicked()
 	return FReply::Handled();
 }
 
-FReply FPaperTileMapDetailsCustomization::OnPromoteButtonClicked()
+FReply FPaperTileMapDetailsCustomization::OnPromoteToAssetButtonClicked()
 {
 	if (UPaperTileMapComponent* TileMapComponent = TileMapComponentPtr.Get())
 	{
@@ -232,12 +277,19 @@ FReply FPaperTileMapDetailsCustomization::OnPromoteButtonClicked()
 		{
 			if (TileMapComponent->TileMap != nullptr)
 			{
+				const FScopedTransaction Transaction(LOCTEXT("PromoteToAsset", "Convert Tile Map instance to an asset"));
+
 				// Try promoting the tile map to be an asset (prompts for a name&path, creates a package and then calls the factory, which renames the existing asset and sets RF_Public)
 				UPaperTileMapPromotionFactory* PromotionFactory = NewObject<UPaperTileMapPromotionFactory>();
 				PromotionFactory->AssetToRename = TileMapComponent->TileMap;
 
 				FAssetToolsModule& AssetToolsModule = FAssetToolsModule::GetModule();
 				UObject* NewAsset = AssetToolsModule.Get().CreateAsset(PromotionFactory->GetSupportedClass(), PromotionFactory);
+			
+				// Show it in the content browser
+				TArray<UObject*> ObjectsToSync;
+				ObjectsToSync.Add(NewAsset);
+				GEditor->SyncBrowserToObjects(ObjectsToSync);
 			}
 		}
 	}
@@ -245,9 +297,25 @@ FReply FPaperTileMapDetailsCustomization::OnPromoteButtonClicked()
 	return FReply::Handled();
 }
 
+FReply FPaperTileMapDetailsCustomization::OnMakeInstanceFromAssetButtonClicked()
+{
+	if (UPaperTileMapComponent* TileMapComponent = TileMapComponentPtr.Get())
+	{
+		if (!TileMapComponent->OwnsTileMap())
+		{
+			const FScopedTransaction Transaction(LOCTEXT("ConvertToInstance", "Convert Tile Map asset to unique instance"));
+
+			TileMapComponent->Modify();
+			TileMapComponent->MakeTileMapEditable();
+		}
+	}
+
+	return FReply::Handled();
+}
+
 EVisibility FPaperTileMapDetailsCustomization::GetNonEditModeVisibility() const
 {
-	return GLevelEditorModeTools().IsModeActive(FEdModeTileMap::EM_TileMap) ? EVisibility::Collapsed : EVisibility::Visible;
+	return InLevelEditorContext() ? EVisibility::Visible : EVisibility::Collapsed;
 }
 
 EVisibility FPaperTileMapDetailsCustomization::GetNewButtonVisiblity() const
@@ -256,16 +324,36 @@ EVisibility FPaperTileMapDetailsCustomization::GetNewButtonVisiblity() const
 }
 
 EVisibility FPaperTileMapDetailsCustomization::GetVisibilityForInstancedOnlyProperties() const
+{
+	return IsInstanced() ? EVisibility::Visible : EVisibility::Collapsed;
+}
+
+EVisibility FPaperTileMapDetailsCustomization::GetVisibilityForMakeIntoInstance() const
+{
+	return (!IsInstanced() && InLevelEditorContext()) ? EVisibility::Visible : EVisibility::Collapsed;
+}
+
+bool FPaperTileMapDetailsCustomization::InLevelEditorContext() const
 {
 	if (UPaperTileMapComponent* TileMapComponent = TileMapComponentPtr.Get())
 	{
-		if (TileMapComponent->OwnsTileMap())
-		{
-			return EVisibility::Visible;
-		}
+		return TileMapComponent->GetOwner() != nullptr;
+	}
+
+	return false;
+
+	//@TODO: This isn't the right question, we should instead look and see if we're a customization for an actor versus a component
+	//return GLevelEditorModeTools().IsModeActive(FEdModeTileMap::EM_TileMap);
+}
+
+bool FPaperTileMapDetailsCustomization::IsInstanced() const
+{
+	if (UPaperTileMapComponent* TileMapComponent = TileMapComponentPtr.Get())
+	{
+		return TileMapComponent->OwnsTileMap();
 	}
 
-	return EVisibility::Hidden;
+	return false;
 }
 
 //////////////////////////////////////////////////////////////////////////
diff --git a/Engine/Plugins/2D/Paper2D/Source/Paper2DEditor/Private/TileMapEditing/PaperTileMapDetailsCustomization.h b/Engine/Plugins/2D/Paper2D/Source/Paper2DEditor/Private/TileMapEditing/PaperTileMapDetailsCustomization.h
--- a/Engine/Plugins/2D/Paper2D/Source/Paper2DEditor/Private/TileMapEditing/PaperTileMapDetailsCustomization.h
+++ b/Engine/Plugins/2D/Paper2D/Source/Paper2DEditor/Private/TileMapEditing/PaperTileMapDetailsCustomization.h
@@ -26,11 +26,17 @@ private:
 private:
 	FReply EnterTileMapEditingMode();
 	FReply OnNewButtonClicked();
-	FReply OnPromoteButtonClicked();
+	FReply OnPromoteToAssetButtonClicked();
+	FReply OnMakeInstanceFromAssetButtonClicked();
 
 	EVisibility GetNonEditModeVisibility() const;
 
 	EVisibility GetVisibilityForInstancedOnlyProperties() const;
+	
+	EVisibility GetVisibilityForMakeIntoInstance() const;
 
 	EVisibility GetNewButtonVisiblity() const;
+
+	bool InLevelEditorContext() const;
+	bool IsInstanced() const;
 };
