Sprite Editor: Added display of 'lit/unlit' to the rendering section header of the sprite details customization [UE-14125]

diff --git a/Engine/Plugins/2D/Paper2D/Source/Paper2DEditor/Private/SpriteEditor/SpriteDetailsCustomization.cpp b/Engine/Plugins/2D/Paper2D/Source/Paper2DEditor/Private/SpriteEditor/SpriteDetailsCustomization.cpp
--- a/Engine/Plugins/2D/Paper2D/Source/Paper2DEditor/Private/SpriteEditor/SpriteDetailsCustomization.cpp
+++ b/Engine/Plugins/2D/Paper2D/Source/Paper2DEditor/Private/SpriteEditor/SpriteDetailsCustomization.cpp
@@ -504,11 +504,31 @@ FText FSpriteDetailsCustomization::GetRenderingHeaderContentText(TWeakObjectPtr<
 		int32 NumTranslucentTriangles = 0;
 		FSpriteEditorViewportClient::AnalyzeSpriteMaterialType(SpriteBeingEdited, /*out*/ NumOpaqueTriangles, /*out*/ NumMaskedTriangles, /*out*/ NumTranslucentTriangles);
 
+		FText MaterialType;
+		if (UMaterialInterface* Material = SpriteBeingEdited->GetDefaultMaterial())
+		{
+			switch (Material->GetShadingModel())
+			{
+			case MSM_Unlit:
+				MaterialType = LOCTEXT("Unlit", "Unlit");
+				break;
+			case MSM_DefaultLit:
+				MaterialType = LOCTEXT("Lit", "Lit");
+				break;
+			default:
+				MaterialType = LOCTEXT("Exotic", "Exotic");
+				break;
+			}
+			//@TODO: This isn't exported, and it would include the MSM_ prefix as well...
+			// MaterialType = FText::AsCultureInvariant(UMaterial::GetMaterialShadingModelString(Material->GetShadingModel()));
+		}
+
 		static const FText Opaque = LOCTEXT("OpaqueMaterial", "Opaque");
 		static const FText Translucent = LOCTEXT("TranslucentMaterial", "Translucent");
 		static const FText Masked = LOCTEXT("MaskedMaterial", "Masked");
-		static const FText TwoItems = LOCTEXT("SpriteWithTwoMaterialsRenderHeaderText", "{0} and {1}");
-		static const FText ThreeItems = LOCTEXT("SpriteWithThreeMaterialsRenderHeaderText", "{0}, {1}, {2}"); // Should never happen right now!
+		static const FText OneMaterial = LOCTEXT("SpriteWithOneMaterialRenderHeaderText", "{0} - {1}");
+		static const FText TwoMaterials = LOCTEXT("SpriteWithTwoMaterialsRenderHeaderText", "{0} - {1} and {2}");
+		static const FText ThreeMaterials = LOCTEXT("SpriteWithThreeMaterialsRenderHeaderText", "{0} - {1}, {2}, {3}"); // Should never happen right now!
 
 		TArray<const FText*, TInlineAllocator<3>> MaterialTypes;
 		if (NumOpaqueTriangles > 0)
@@ -531,13 +551,13 @@ FText FSpriteDetailsCustomization::GetRenderingHeaderContentText(TWeakObjectPtr<
 		case 0:
 			break;
 		case 1:
-			HeaderDisplayText = *MaterialTypes[0];
+			HeaderDisplayText = FText::Format(OneMaterial, MaterialType, *MaterialTypes[0]);
 			break;
 		case 2:
-			HeaderDisplayText = FText::Format(TwoItems, *MaterialTypes[0], *MaterialTypes[1]);
+			HeaderDisplayText = FText::Format(TwoMaterials, MaterialType, *MaterialTypes[0], *MaterialTypes[1]);
 			break;
 		case 3:
-			HeaderDisplayText = FText::Format(ThreeItems, *MaterialTypes[0], *MaterialTypes[1], *MaterialTypes[2]);
+			HeaderDisplayText = FText::Format(ThreeMaterials, MaterialType, *MaterialTypes[0], *MaterialTypes[1], *MaterialTypes[2]);
 			break;
 		};
 	}
