Tile Maps: Fixed a bug where layers would lose data and be stuck at a phantom smaller size after undoing a change in the size [UE-8032]

diff --git a/Engine/Plugins/2D/Paper2D/Source/Paper2D/Private/PaperTileMap.cpp b/Engine/Plugins/2D/Paper2D/Source/Paper2D/Private/PaperTileMap.cpp
--- a/Engine/Plugins/2D/Paper2D/Source/Paper2D/Private/PaperTileMap.cpp
+++ b/Engine/Plugins/2D/Paper2D/Source/Paper2D/Private/PaperTileMap.cpp
@@ -390,6 +390,7 @@ void UPaperTileMap::ResizeMap(int32 NewWidth, int32 NewHeight, bool bForceResize
 		for (int32 LayerIndex = 0; LayerIndex < TileLayers.Num(); ++LayerIndex)
 		{
 			UPaperTileLayer* TileLayer = TileLayers[LayerIndex];
+			TileLayer->Modify();
 			TileLayer->ResizeMap(MapWidth, MapHeight);
 		}
 	}
