Sprite Editor: Collapsed the details panel tab stack to save a bit more vertical space

diff --git a/Engine/Plugins/2D/Paper2D/Source/Paper2DEditor/Private/SpriteEditor/SpriteEditor.cpp b/Engine/Plugins/2D/Paper2D/Source/Paper2DEditor/Private/SpriteEditor/SpriteEditor.cpp
--- a/Engine/Plugins/2D/Paper2D/Source/Paper2DEditor/Private/SpriteEditor/SpriteEditor.cpp
+++ b/Engine/Plugins/2D/Paper2D/Source/Paper2DEditor/Private/SpriteEditor/SpriteEditor.cpp
@@ -358,7 +358,7 @@ void FSpriteEditor::InitSpriteEditor(const EToolkitMode::Type Mode, const TShare
 	ViewportPtr = SNew(SSpriteEditorViewport, SharedThis(this));
 	
 	// Default layout
-	const TSharedRef<FTabManager::FLayout> StandaloneDefaultLayout = FTabManager::NewLayout("Standalone_SpriteEditor_Layout_v5")
+	const TSharedRef<FTabManager::FLayout> StandaloneDefaultLayout = FTabManager::NewLayout("Standalone_SpriteEditor_Layout_v6")
 		->AddArea
 		(
 			FTabManager::NewPrimaryArea()
@@ -391,6 +391,7 @@ void FSpriteEditor::InitSpriteEditor(const EToolkitMode::Type Mode, const TShare
 					(
 						FTabManager::NewStack()
 						->SetSizeCoefficient(0.75f)
+						->SetHideTabWell(true)
 						->AddTab(FSpriteEditorTabs::DetailsID, ETabState::OpenedTab)
 					)
 					->Split
