Tile Set Editor: Added a show grid toggle option for the single tile editor [UE-13627]

diff --git a/Engine/Plugins/2D/Paper2D/Source/Paper2DEditor/Private/PaperStyle.cpp b/Engine/Plugins/2D/Paper2D/Source/Paper2DEditor/Private/PaperStyle.cpp
--- a/Engine/Plugins/2D/Paper2D/Source/Paper2DEditor/Private/PaperStyle.cpp
+++ b/Engine/Plugins/2D/Paper2D/Source/Paper2DEditor/Private/PaperStyle.cpp
@@ -136,6 +136,12 @@ void FPaperStyle::Initialize()
 		StyleSet->Set("TileMapEditor.SetShowCollision.Small", new IMAGE_BRUSH(TEXT("Icons/icon_StaticMeshEd_Collision_40x"), Icon20x20));
 	}
 
+	// Tile set editor
+	{
+		StyleSet->Set("TileSetEditor.SetShowGrid", new IMAGE_BRUSH(TEXT("Icons/icon_MatEd_Grid_40x"), Icon40x40));
+		StyleSet->Set("TileSetEditor.SetShowGrid.Small", new IMAGE_BRUSH(TEXT("Icons/icon_MatEd_Grid_40x"), Icon20x20));
+	}
+
 	// Sprite editor
 	{
 		StyleSet->Set("SpriteEditor.SetShowGrid", new IMAGE_BRUSH(TEXT("Icons/icon_MatEd_Grid_40x"), Icon40x40));
diff --git a/Engine/Plugins/2D/Paper2D/Source/Paper2DEditor/Private/TileSetEditor.cpp b/Engine/Plugins/2D/Paper2D/Source/Paper2DEditor/Private/TileSetEditor.cpp
--- a/Engine/Plugins/2D/Paper2D/Source/Paper2DEditor/Private/TileSetEditor.cpp
+++ b/Engine/Plugins/2D/Paper2D/Source/Paper2DEditor/Private/TileSetEditor.cpp
@@ -16,6 +16,7 @@
 #include "TileSetEditor/TileSetSelectorViewport.h"
 #include "TileSetEditor/SingleTileEditorViewport.h"
 #include "TileSetEditor/SingleTileEditorViewportClient.h"
+#include "TileSetEditor/TileSetEditorCommands.h"
 #include "SpriteEditor/SpriteEditorCommands.h"
 
 #define LOCTEXT_NAMESPACE "TileSetEditor"
@@ -139,7 +140,7 @@ void FTileSetEditor::InitTileSetEditor(const EToolkitMode::Type Mode, const TSha
 
 	TileEditorViewport = SNew(SSingleTileEditorViewport, TileEditorViewportClient);
 
-	//@TODO: FTileSetEditorCommands::Register();
+	FTileSetEditorCommands::Register();
 	FSpriteGeometryEditCommands::Register();
 
 	BindCommands();
@@ -235,6 +236,7 @@ TSharedRef<SDockTab> FTileSetEditor::SpawnTab_SingleTileEditor(const FSpawnTabAr
 
 void FTileSetEditor::BindCommands()
 {
+	// Commands would go here
 }
 
 void FTileSetEditor::ExtendMenu()
diff --git a/Engine/Plugins/2D/Paper2D/Source/Paper2DEditor/Private/TileSetEditor/SingleTileEditorViewport.cpp b/Engine/Plugins/2D/Paper2D/Source/Paper2DEditor/Private/TileSetEditor/SingleTileEditorViewport.cpp
--- a/Engine/Plugins/2D/Paper2D/Source/Paper2DEditor/Private/TileSetEditor/SingleTileEditorViewport.cpp
+++ b/Engine/Plugins/2D/Paper2D/Source/Paper2DEditor/Private/TileSetEditor/SingleTileEditorViewport.cpp
@@ -5,6 +5,7 @@
 #include "SingleTileEditorViewport.h"
 #include "SpriteEditor/SpriteEditorCommands.h"
 #include "SingleTileEditorViewportClient.h"
+#include "TileSetEditorCommands.h"
 
 #define LOCTEXT_NAMESPACE "TileSetEditor"
 
@@ -40,6 +41,9 @@ TSharedRef<SWidget> STileSetEditorViewportToolbar::GenerateShowMenu() const
 	FMenuBuilder ShowMenuBuilder(bInShouldCloseWindowAfterMenuSelection, ViewportRef->GetCommandList());
 	{
 		ShowMenuBuilder.AddMenuEntry(FSpriteGeometryEditCommands::Get().SetShowNormals);
+		// Disabled as the single tile editor doesn't actually create collision, that's done in the tile map component
+		// ShowMenuBuilder.AddMenuEntry(FTileSetEditorCommands::Get().SetShowCollision);
+		ShowMenuBuilder.AddMenuEntry(FTileSetEditorCommands::Get().SetShowGrid);
 	}
 
 	return ShowMenuBuilder.MakeWidget();
@@ -109,6 +113,29 @@ TSharedRef<FEditorViewportClient> SSingleTileEditorViewport::MakeEditorViewportC
 	return TypedViewportClient.ToSharedRef();
 }
 
+void SSingleTileEditorViewport::BindCommands()
+{
+	SEditorViewport::BindCommands();
+
+	FTileSetEditorCommands::Register();
+	const FTileSetEditorCommands& Commands = FTileSetEditorCommands::Get();
+
+	TSharedRef<FSingleTileEditorViewportClient> EditorViewportClientRef = TypedViewportClient.ToSharedRef();
+
+	// Show toggles
+	CommandList->MapAction(
+		Commands.SetShowCollision,
+		FExecuteAction::CreateSP(EditorViewportClientRef, &FEditorViewportClient::SetShowCollision),
+		FCanExecuteAction(),
+		FIsActionChecked::CreateSP(EditorViewportClientRef, &FEditorViewportClient::IsSetShowCollisionChecked));
+
+	CommandList->MapAction(
+		Commands.SetShowGrid,
+		FExecuteAction::CreateSP(EditorViewportClientRef, &FEditorViewportClient::SetShowGrid),
+		FCanExecuteAction(),
+		FIsActionChecked::CreateSP(EditorViewportClientRef, &FEditorViewportClient::IsSetShowGridChecked));
+}
+
 TSharedRef<class SEditorViewport> SSingleTileEditorViewport::GetViewportWidget()
 {
 	return SharedThis(this);
diff --git a/Engine/Plugins/2D/Paper2D/Source/Paper2DEditor/Private/TileSetEditor/SingleTileEditorViewport.h b/Engine/Plugins/2D/Paper2D/Source/Paper2DEditor/Private/TileSetEditor/SingleTileEditorViewport.h
--- a/Engine/Plugins/2D/Paper2D/Source/Paper2DEditor/Private/TileSetEditor/SingleTileEditorViewport.h
+++ b/Engine/Plugins/2D/Paper2D/Source/Paper2DEditor/Private/TileSetEditor/SingleTileEditorViewport.h
@@ -25,6 +25,7 @@ public:
 	// SEditorViewport interface
 	virtual TSharedPtr<SWidget> MakeViewportToolbar() override;
 	virtual TSharedRef<FEditorViewportClient> MakeEditorViewportClient() override;
+	virtual void BindCommands() override;
 	// End of SEditorViewport interface
 
 	// ICommonEditorViewportToolbarInfoProvider interface
diff --git a/Engine/Plugins/2D/Paper2D/Source/Paper2DEditor/Private/TileSetEditor/TileSetEditorCommands.cpp b/Engine/Plugins/2D/Paper2D/Source/Paper2DEditor/Private/TileSetEditor/TileSetEditorCommands.cpp
new file mode 100644
--- /dev/null
+++ b/Engine/Plugins/2D/Paper2D/Source/Paper2DEditor/Private/TileSetEditor/TileSetEditorCommands.cpp
@@ -0,0 +1,18 @@
+// Copyright 1998-2015 Epic Games, Inc. All Rights Reserved.
+
+#include "Paper2DEditorPrivatePCH.h"
+#include "TileSetEditorCommands.h"
+
+//////////////////////////////////////////////////////////////////////////
+// FTileSetEditorCommands
+
+#define LOCTEXT_NAMESPACE "TileSetEditor"
+
+void FTileSetEditorCommands::RegisterCommands()
+{
+	// Show toggles
+	UI_COMMAND(SetShowCollision, "Collision", "Toggles display of the collision geometry, if any exists.", EUserInterfaceActionType::ToggleButton, FInputChord(EKeys::C, EModifierKey::Alt));
+	UI_COMMAND(SetShowGrid, "Grid", "Display the grid.", EUserInterfaceActionType::ToggleButton, FInputChord());
+}
+
+#undef LOCTEXT_NAMESPACE
\ No newline at end of file
diff --git a/Engine/Plugins/2D/Paper2D/Source/Paper2DEditor/Private/TileSetEditor/TileSetEditorCommands.h b/Engine/Plugins/2D/Paper2D/Source/Paper2DEditor/Private/TileSetEditor/TileSetEditorCommands.h
new file mode 100644
--- /dev/null
+++ b/Engine/Plugins/2D/Paper2D/Source/Paper2DEditor/Private/TileSetEditor/TileSetEditorCommands.h
@@ -0,0 +1,28 @@
+// Copyright 1998-2015 Epic Games, Inc. All Rights Reserved.
+
+#pragma once
+
+#include "PaperStyle.h"
+
+class FTileSetEditorCommands : public TCommands<FTileSetEditorCommands>
+{
+public:
+	FTileSetEditorCommands()
+		: TCommands<FTileSetEditorCommands>(
+			TEXT("TileSetEditor"), // Context name for fast lookup
+			NSLOCTEXT("Contexts", "TileSetEditor", "Tile Set Editor"), // Localized context name for displaying
+			NAME_None, // Parent
+			FPaperStyle::Get()->GetStyleSetName() // Icon Style Set
+			)
+	{
+	}
+
+	// TCommand<> interface
+	virtual void RegisterCommands() override;
+	// End of TCommand<> interface
+
+public:
+	// Show toggles
+	TSharedPtr<FUICommandInfo> SetShowCollision;
+	TSharedPtr<FUICommandInfo> SetShowGrid;
+};
\ No newline at end of file
