Fix monolithic builds.

diff --git a/Engine/Plugins/2D/Paper2D/Source/Paper2D/Private/PaperTileMapRenderSceneProxy.cpp b/Engine/Plugins/2D/Paper2D/Source/Paper2D/Private/PaperTileMapRenderSceneProxy.cpp
--- a/Engine/Plugins/2D/Paper2D/Source/Paper2D/Private/PaperTileMapRenderSceneProxy.cpp
+++ b/Engine/Plugins/2D/Paper2D/Source/Paper2D/Private/PaperTileMapRenderSceneProxy.cpp
@@ -25,9 +25,11 @@ FPaperTileMapRenderSceneProxy::FPaperTileMapRenderSceneProxy(const UPaperTileMap
 		Material = InTileComponent->GetMaterial(0);
 		MaterialRelevance = InTileComponent->GetMaterialRelevance(GetScene().GetFeatureLevel());
 
+#if WITH_EDITORONLY_DATA
 		bShowPerTileGrid = InTileComponent->bShowPerTileGridWhenSelected;
 		bShowPerLayerGrid = InTileComponent->bShowPerLayerGridWhenSelected;
 		bShowOutlineWhenUnselected = InTileComponent->bShowOutlineWhenUnselected;
+#endif
 	}
 }
 
