Tile Maps: Further cleaned up tile map grid flags in non-editor builds

diff --git a/Engine/Plugins/2D/Paper2D/Source/Paper2D/Private/PaperTileMapRenderSceneProxy.cpp b/Engine/Plugins/2D/Paper2D/Source/Paper2D/Private/PaperTileMapRenderSceneProxy.cpp
--- a/Engine/Plugins/2D/Paper2D/Source/Paper2D/Private/PaperTileMapRenderSceneProxy.cpp
+++ b/Engine/Plugins/2D/Paper2D/Source/Paper2D/Private/PaperTileMapRenderSceneProxy.cpp
@@ -17,20 +17,24 @@
 FPaperTileMapRenderSceneProxy::FPaperTileMapRenderSceneProxy(const UPaperTileMapComponent* InComponent)
 	: FPaperRenderSceneProxy(InComponent)
 	, TileMap(nullptr)
+#if WITH_EDITOR
+	, bShowPerTileGrid(false)
+	, bShowPerLayerGrid(false)
+	, bShowOutlineWhenUnselected(false)
+#endif
 {
-	if (const UPaperTileMapComponent* InTileComponent = Cast<const UPaperTileMapComponent>(InComponent))
-	{
-		WireframeColor = InTileComponent->GetWireframeColor();
-		TileMap = InTileComponent->TileMap;
-		Material = InTileComponent->GetMaterial(0);
-		MaterialRelevance = InTileComponent->GetMaterialRelevance(GetScene().GetFeatureLevel());
+	check(InComponent);
+
+	WireframeColor = InComponent->GetWireframeColor();
+	TileMap = InComponent->TileMap;
+	Material = InComponent->GetMaterial(0);
+	MaterialRelevance = InComponent->GetMaterialRelevance(GetScene().GetFeatureLevel());
 
 #if WITH_EDITORONLY_DATA
-		bShowPerTileGrid = InTileComponent->bShowPerTileGridWhenSelected;
-		bShowPerLayerGrid = InTileComponent->bShowPerLayerGridWhenSelected;
-		bShowOutlineWhenUnselected = InTileComponent->bShowOutlineWhenUnselected;
+	bShowPerTileGrid = InComponent->bShowPerTileGridWhenSelected;
+	bShowPerLayerGrid = InComponent->bShowPerLayerGridWhenSelected;
+	bShowOutlineWhenUnselected = InComponent->bShowOutlineWhenUnselected;
 #endif
-	}
 }
 
 void FPaperTileMapRenderSceneProxy::DrawBoundsForLayer(FPrimitiveDrawInterface* PDI, const FLinearColor& Color, int32 LayerIndex) const
diff --git a/Engine/Plugins/2D/Paper2D/Source/Paper2D/Private/PaperTileMapRenderSceneProxy.h b/Engine/Plugins/2D/Paper2D/Source/Paper2D/Private/PaperTileMapRenderSceneProxy.h
--- a/Engine/Plugins/2D/Paper2D/Source/Paper2D/Private/PaperTileMapRenderSceneProxy.h
+++ b/Engine/Plugins/2D/Paper2D/Source/Paper2D/Private/PaperTileMapRenderSceneProxy.h
@@ -25,9 +25,12 @@ protected:
 	void DrawBoundsForLayer(FPrimitiveDrawInterface* PDI, const FLinearColor& Color, int32 LayerIndex) const;
 
 protected:
+
+#if WITH_EDITOR
 	bool bShowPerTileGrid;
 	bool bShowPerLayerGrid;
 	bool bShowOutlineWhenUnselected;
+#endif
 
 	//@TODO: Not thread safe
 	const class UPaperTileMap* TileMap;
