Paper2D: Added stats to the render scene proxies [UE-14412]

diff --git a/Engine/Plugins/2D/Paper2D/Source/Paper2D/Private/Paper2DModule.cpp b/Engine/Plugins/2D/Paper2D/Source/Paper2D/Private/Paper2DModule.cpp
--- a/Engine/Plugins/2D/Paper2D/Source/Paper2D/Private/Paper2DModule.cpp
+++ b/Engine/Plugins/2D/Paper2D/Source/Paper2D/Private/Paper2DModule.cpp
@@ -3,9 +3,10 @@
 #include "Paper2DPrivatePCH.h"
 #include "Paper2DModule.h"
 
-//#include "Physics/Box2DIntegration.h"
 #include "Rendering/PaperBatchManager.h"
 
+DEFINE_STAT(STAT_PaperRender_SetSpriteRT);
+
 //////////////////////////////////////////////////////////////////////////
 // FPaper2DModule
 
diff --git a/Engine/Plugins/2D/Paper2D/Source/Paper2D/Private/PaperFlipbookComponent.cpp b/Engine/Plugins/2D/Paper2D/Source/Paper2D/Private/PaperFlipbookComponent.cpp
--- a/Engine/Plugins/2D/Paper2D/Source/Paper2D/Private/PaperFlipbookComponent.cpp
+++ b/Engine/Plugins/2D/Paper2D/Source/Paper2D/Private/PaperFlipbookComponent.cpp
@@ -12,6 +12,8 @@
 #include "PhysicsEngine/BodySetup.h"
 #include "SpriteDrawCall.h"
 
+DECLARE_CYCLE_STAT(TEXT("Tick Flipbook"), STAT_TickFlipbook, STATGROUP_Paper2D);
+
 #define LOCTEXT_NAMESPACE "Paper2D"
 
 //////////////////////////////////////////////////////////////////////////
@@ -325,6 +327,8 @@ void UPaperFlipbookComponent::OnRep_SourceFlipbook(class UPaperFlipbook* OldFlip
 
 void UPaperFlipbookComponent::TickComponent(float DeltaTime, enum ELevelTick TickType, FActorComponentTickFunction* ThisTickFunction)
 {
+	SCOPE_CYCLE_COUNTER(STAT_TickFlipbook);
+
 	// Advance time
 	TickFlipbook(DeltaTime);
 
@@ -522,14 +526,14 @@ void UPaperFlipbookComponent::SetPlaybackPosition(float NewPosition, bool bFireE
 			bool bFireThisEvent = false;
 			if (!bReversePlayback)
 			{
-				if (EventTime >= MinTime && EventTime < MaxTime)
+				if ((EventTime >= MinTime) && (EventTime < MaxTime))
 				{
 					bFireThisEvent = true;
 				}
 			}
 			else
 			{
-				if (EventTime > MinTime && EventTime <= MaxTime)
+				if ((EventTime > MinTime) && (EventTime <= MaxTime))
 				{
 					bFireThisEvent = true;
 				}
diff --git a/Engine/Plugins/2D/Paper2D/Source/Paper2D/Private/PaperRenderSceneProxy.cpp b/Engine/Plugins/2D/Paper2D/Source/Paper2D/Private/PaperRenderSceneProxy.cpp
--- a/Engine/Plugins/2D/Paper2D/Source/Paper2D/Private/PaperRenderSceneProxy.cpp
+++ b/Engine/Plugins/2D/Paper2D/Source/Paper2D/Private/PaperRenderSceneProxy.cpp
@@ -13,6 +13,9 @@
 #include "Engine/Engine.h"
 #include "SpriteDrawCall.h"
 
+DECLARE_CYCLE_STAT(TEXT("Get Batch Mesh"), STAT_PaperRender_GetBatchMesh, STATGROUP_Paper2D);
+DECLARE_CYCLE_STAT(TEXT("Get New Batch Meshes"), STAT_PaperRender_GetNewBatchMeshes, STATGROUP_Paper2D);
+
 //////////////////////////////////////////////////////////////////////////
 // FPaperSpriteVertex
 
@@ -341,6 +344,8 @@ void FPaperRenderSceneProxy::GetNewBatchMeshes(const FSceneView* View, bool bUse
 		return;
 	}
 
+	SCOPE_CYCLE_COUNTER(STAT_PaperRender_GetNewBatchMeshes);
+
 	const uint8 DPG = GetDepthPriorityGroup(View);
 	FVertexFactory* VertexFactory = GetPaperSpriteVertexFactory();
 	const bool bIsWireframeView = View->Family->EngineShowFlags.Wireframe;
@@ -416,6 +421,8 @@ public:
 
 void FPaperRenderSceneProxy::GetBatchMesh(const FSceneView* View, bool bUseOverrideColor, const FLinearColor& OverrideColor, class UMaterialInterface* BatchMaterial, const TArray<FSpriteDrawCallRecord>& Batch, int32 ViewIndex, FMeshElementCollector& Collector) const
 {
+	SCOPE_CYCLE_COUNTER(STAT_PaperRender_GetBatchMesh);
+
 	const uint8 DPG = GetDepthPriorityGroup(View);
 	FVertexFactory* VertexFactory = GetPaperSpriteVertexFactory();
 
@@ -554,6 +561,8 @@ bool FPaperRenderSceneProxy::CanBeOccluded() const
 
 void FPaperRenderSceneProxy::SetDrawCall_RenderThread(const FSpriteDrawCallRecord& NewDynamicData)
 {
+	SCOPE_CYCLE_COUNTER(STAT_PaperRender_SetSpriteRT);
+
 	BatchedSprites.Empty();
 
 	FSpriteDrawCallRecord& Record = *new (BatchedSprites) FSpriteDrawCallRecord;
diff --git a/Engine/Plugins/2D/Paper2D/Source/Paper2D/Private/PaperSpriteSceneProxy.cpp b/Engine/Plugins/2D/Paper2D/Source/Paper2D/Private/PaperSpriteSceneProxy.cpp
--- a/Engine/Plugins/2D/Paper2D/Source/Paper2D/Private/PaperSpriteSceneProxy.cpp
+++ b/Engine/Plugins/2D/Paper2D/Source/Paper2D/Private/PaperSpriteSceneProxy.cpp
@@ -91,6 +91,8 @@ void FPaperSpriteSceneProxy::GetDynamicMeshElements(const TArray<const FSceneVie
 
 void FPaperSpriteSceneProxy::SetSprite_RenderThread(const FSpriteDrawCallRecord& NewDynamicData, int32 SplitIndex)
 {
+	SCOPE_CYCLE_COUNTER(STAT_PaperRender_SetSpriteRT);
+
 	BatchedSprites.Empty();
 	AlternateBatchedSprites.Empty();
 
diff --git a/Engine/Plugins/2D/Paper2D/Source/Paper2D/Private/PaperTileMapRenderSceneProxy.cpp b/Engine/Plugins/2D/Paper2D/Source/Paper2D/Private/PaperTileMapRenderSceneProxy.cpp
--- a/Engine/Plugins/2D/Paper2D/Source/Paper2D/Private/PaperTileMapRenderSceneProxy.cpp
+++ b/Engine/Plugins/2D/Paper2D/Source/Paper2D/Private/PaperTileMapRenderSceneProxy.cpp
@@ -11,6 +11,8 @@
 #include "PaperTileLayer.h"
 #include "PaperTileMapComponent.h"
 
+DECLARE_CYCLE_STAT(TEXT("Tile Map Proxy"), STAT_TileMap_GetDynamicMeshElements, STATGROUP_Paper2D);
+
 //////////////////////////////////////////////////////////////////////////
 // FPaperTileMapRenderSceneProxy
 
@@ -56,7 +58,7 @@ void FPaperTileMapRenderSceneProxy::DrawBoundsForLayer(FPrimitiveDrawInterface*
 
 void FPaperTileMapRenderSceneProxy::GetDynamicMeshElements(const TArray<const FSceneView*>& Views, const FSceneViewFamily& ViewFamily, uint32 VisibilityMap, FMeshElementCollector& Collector) const
 {
-	QUICK_SCOPE_CYCLE_COUNTER(STAT_FPaperTileMapRenderSceneProxy_GetDynamicMeshElements);
+	SCOPE_CYCLE_COUNTER(STAT_TileMap_GetDynamicMeshElements);
 	checkSlow(IsInRenderingThread());
 
 	// Slight depth bias so that the wireframe grid overlay doesn't z-fight with the tiles themselves
diff --git a/Engine/Plugins/2D/Paper2D/Source/Paper2D/Private/Terrain/PaperTerrainComponent.cpp b/Engine/Plugins/2D/Paper2D/Source/Paper2D/Private/Terrain/PaperTerrainComponent.cpp
--- a/Engine/Plugins/2D/Paper2D/Source/Paper2D/Private/Terrain/PaperTerrainComponent.cpp
+++ b/Engine/Plugins/2D/Paper2D/Source/Paper2D/Private/Terrain/PaperTerrainComponent.cpp
@@ -13,6 +13,7 @@
 
 #define PAPER_USE_MATERIAL_SLOPES 1
 
+DECLARE_CYCLE_STAT(TEXT("Terrain Spline Proxy"), STAT_TerrainSpline_GetDynamicMeshElements, STATGROUP_Paper2D);
 
 //////////////////////////////////////////////////////////////////////////
 
@@ -82,6 +83,8 @@ FPaperTerrainSceneProxy::FPaperTerrainSceneProxy(const UPaperTerrainComponent* I
 
 void FPaperTerrainSceneProxy::GetDynamicMeshElementsForView(const FSceneView* View, int32 ViewIndex, bool bUseOverrideColor, const FLinearColor& OverrideColor, FMeshElementCollector& Collector) const
 {
+	SCOPE_CYCLE_COUNTER(STAT_TerrainSpline_GetDynamicMeshElements);
+
 	for (const FPaperTerrainSpriteGeometry& Batch : DrawingData)
 	{
 		if (Batch.Material != nullptr)
diff --git a/Engine/Plugins/2D/Paper2D/Source/Paper2D/Public/Paper2DModule.h b/Engine/Plugins/2D/Paper2D/Source/Paper2D/Public/Paper2DModule.h
--- a/Engine/Plugins/2D/Paper2D/Source/Paper2D/Public/Paper2DModule.h
+++ b/Engine/Plugins/2D/Paper2D/Source/Paper2D/Public/Paper2DModule.h
@@ -11,6 +11,8 @@ class IPaper2DModuleInterface : public IModuleInterface
 {
 };
 
+//////////////////////////////////////////////////////////////////////////
+
 // NOTE: Experimental: Do not adjust these values, many aspects of the editor or runtime assume PaperX=Right(+X) PaperY=Up(+Z) PaperZ=PaperX cross PaperY=Out(-Y)
 // The ability to adjust these values may be removed entirely in the future, or made to work properly
 //@TODO: PAPER2D: Synchronize with the Box2D axes in the engine (or remove the ability to customize them)
@@ -23,3 +25,9 @@ extern PAPER2D_API FVector PaperAxisY;
 
 // The 3D vector that corresponds to the sprite Z axis
 extern PAPER2D_API FVector PaperAxisZ;
+
+//////////////////////////////////////////////////////////////////////////
+// Paper2D stats
+
+DECLARE_STATS_GROUP(TEXT("Paper2D"), STATGROUP_Paper2D, STATCAT_Advanced);
+DECLARE_CYCLE_STAT_EXTERN(TEXT("SetSprite (RT)"), STAT_PaperRender_SetSpriteRT, STATGROUP_Paper2D, PAPER2D_API);
