Tile Set Editor: Applying collision changes now conditions single tile collision geometry as well [UE-13644]

diff --git a/Engine/Plugins/2D/Paper2D/Source/Paper2D/Classes/SpriteEditorOnlyTypes.h b/Engine/Plugins/2D/Paper2D/Source/Paper2D/Classes/SpriteEditorOnlyTypes.h
--- a/Engine/Plugins/2D/Paper2D/Source/Paper2D/Classes/SpriteEditorOnlyTypes.h
+++ b/Engine/Plugins/2D/Paper2D/Source/Paper2D/Classes/SpriteEditorOnlyTypes.h
@@ -214,7 +214,8 @@ public:
 	void Reset();
 
 	// Conditions this geometry collection (turning Polygons back to Boxes if they meet the definition of a box, etc...)
-	void ConditionGeometry();
+	// Returns true if something was modified
+	bool ConditionGeometry();
 
 	// Takes all polygon shapes and generates a list of triangles from them.
 	// Output will contain a multiple of 3 points, each set is one triangle.
diff --git a/Engine/Plugins/2D/Paper2D/Source/Paper2D/Private/PaperSprite.cpp b/Engine/Plugins/2D/Paper2D/Source/Paper2D/Private/PaperSprite.cpp
--- a/Engine/Plugins/2D/Paper2D/Source/Paper2D/Private/PaperSprite.cpp
+++ b/Engine/Plugins/2D/Paper2D/Source/Paper2D/Private/PaperSprite.cpp
@@ -1985,8 +1985,10 @@ bool AreVectorsPerpendicular(const FVector2D& Vector1, const FVector2D& Vector2,
 	return FMath::IsNearlyEqual(DotProduct, 0.0f, Threshold);
 }
 
-void FSpriteGeometryCollection::ConditionGeometry()
+bool FSpriteGeometryCollection::ConditionGeometry()
 {
+	bool bModifiedGeometry = false;
+
 	for (FSpriteGeometryShape& Shape : Shapes)
 	{
 		if ((Shape.ShapeType == ESpriteShapeType::Polygon) && (Shape.Vertices.Num() == 4))
@@ -2014,10 +2016,13 @@ void FSpriteGeometryCollection::ConditionGeometry()
 					Shape.SetNewPivot(NewPivotTextureSpace);
 					Shape.BoxSize = FVector2D(AB.Size(), DA.Size());
 					Shape.ShapeType = ESpriteShapeType::Box;
+					bModifiedGeometry = true;
 				}
 			}
 		}
 	}
+
+	return bModifiedGeometry;
 }
 
 //////////////////////////////////////////////////////////////////////////
diff --git a/Engine/Plugins/2D/Paper2D/Source/Paper2DEditor/Private/TileSetEditor/SingleTileEditorViewportClient.cpp b/Engine/Plugins/2D/Paper2D/Source/Paper2DEditor/Private/TileSetEditor/SingleTileEditorViewportClient.cpp
--- a/Engine/Plugins/2D/Paper2D/Source/Paper2DEditor/Private/TileSetEditor/SingleTileEditorViewportClient.cpp
+++ b/Engine/Plugins/2D/Paper2D/Source/Paper2DEditor/Private/TileSetEditor/SingleTileEditorViewportClient.cpp
@@ -7,6 +7,9 @@
 #include "ScopedTransaction.h"
 #include "TileSetEditorSettings.h"
 
+#include "SNotificationList.h"
+#include "NotificationManager.h"
+
 #define LOCTEXT_NAMESPACE "TileSetEditor"
 
 //////////////////////////////////////////////////////////////////////////
@@ -244,6 +247,38 @@ void FSingleTileEditorViewportClient::ActivateEditMode(TSharedPtr<FUICommandList
 
 void FSingleTileEditorViewportClient::ApplyCollisionGeometryEdits()
 {
+	bool bConditionedSomething = false;
+
+	// See if anything needs to be conditioned
+	const int32 NumTiles = TileSet->GetTileCount();
+	for (int32 TileIndex = 0; TileIndex < NumTiles; ++TileIndex)
+	{
+		if (TileSet->GetTileMetadata(TileIndex) != nullptr)
+		{
+			FPaperTileMetadata* TileData = TileSet->GetMutableTileMetadata(TileIndex);
+			check(TileData);
+
+			if (TileData->HasCollision())
+			{
+				if (TileData->CollisionData.ConditionGeometry())
+				{
+					bConditionedSomething = true;
+				}
+			}
+		}
+	}
+
+	if (bConditionedSomething)
+	{
+		TileSet->Modify();
+
+		// Create and display a notification about the tile set being modified
+		const FText NotificationText = FText::Format(LOCTEXT("NeedToSaveTileSet", "Optimized collision on one or more tiles.\n'{0}' needs to be saved."), FText::AsCultureInvariant(TileSet->GetName()));
+		FNotificationInfo Info(NotificationText);
+		Info.ExpireDuration = 2.0f;
+		FSlateNotificationManager::Get().AddNotification(Info);
+	}
+
 	// Apply changes to all tile maps that use this tile set
 	for (TObjectIterator<UPaperTileMap> TileMapIt; TileMapIt; ++TileMapIt)
 	{
