Tile Set Editor: Fixed a confusing selection issue so that the single tile editor only changes targets if there is exactly one tile selected [UE-13622]

diff --git a/Engine/Plugins/2D/Paper2D/Source/Paper2DEditor/Private/TileSetEditor/SingleTileEditorViewportClient.cpp b/Engine/Plugins/2D/Paper2D/Source/Paper2DEditor/Private/TileSetEditor/SingleTileEditorViewportClient.cpp
--- a/Engine/Plugins/2D/Paper2D/Source/Paper2DEditor/Private/TileSetEditor/SingleTileEditorViewportClient.cpp
+++ b/Engine/Plugins/2D/Paper2D/Source/Paper2DEditor/Private/TileSetEditor/SingleTileEditorViewportClient.cpp
@@ -219,8 +219,14 @@ void FSingleTileEditorViewportClient::SetTileIndex(int32 InTileIndex)
 
 void FSingleTileEditorViewportClient::OnActiveTileIndexChanged(const FIntPoint& TopLeft, const FIntPoint& Dimensions)
 {
-	const int32 NewIndex = (TileSet->GetTileCountX() * TopLeft.Y) + TopLeft.X;
-	SetTileIndex(NewIndex);
+	const int32 TopLeftIndex = (TileSet->GetTileCountX() * TopLeft.Y) + TopLeft.X;
+	const int32 SelectionArea = Dimensions.X * Dimensions.Y;
+
+	if (SelectionArea > 0)
+	{
+		const int32 NewIndex = (SelectionArea == 1) ? TopLeftIndex : INDEX_NONE;
+		SetTileIndex(NewIndex);
+	}
 }
 
 int32 FSingleTileEditorViewportClient::GetTileIndex() const
