Tile Set Editor: Added support for F to focus in the single tile editor, and it focuses automatically the first time you edit a tile as well [UE-13618]

diff --git a/Engine/Plugins/2D/Paper2D/Source/Paper2DEditor/Private/TileSetEditor/SingleTileEditorViewport.cpp b/Engine/Plugins/2D/Paper2D/Source/Paper2DEditor/Private/TileSetEditor/SingleTileEditorViewport.cpp
--- a/Engine/Plugins/2D/Paper2D/Source/Paper2DEditor/Private/TileSetEditor/SingleTileEditorViewport.cpp
+++ b/Engine/Plugins/2D/Paper2D/Source/Paper2DEditor/Private/TileSetEditor/SingleTileEditorViewport.cpp
@@ -141,6 +141,11 @@ void SSingleTileEditorViewport::BindCommands()
 		FExecuteAction::CreateSP(EditorViewportClientRef, &FSingleTileEditorViewportClient::ApplyCollisionGeometryEdits));
 }
 
+void SSingleTileEditorViewport::OnFocusViewportToSelection()
+{
+	TypedViewportClient->RequestFocusOnSelection(/*bInstant=*/ false);
+}
+
 TSharedRef<class SEditorViewport> SSingleTileEditorViewport::GetViewportWidget()
 {
 	return SharedThis(this);
diff --git a/Engine/Plugins/2D/Paper2D/Source/Paper2DEditor/Private/TileSetEditor/SingleTileEditorViewport.h b/Engine/Plugins/2D/Paper2D/Source/Paper2DEditor/Private/TileSetEditor/SingleTileEditorViewport.h
--- a/Engine/Plugins/2D/Paper2D/Source/Paper2DEditor/Private/TileSetEditor/SingleTileEditorViewport.h
+++ b/Engine/Plugins/2D/Paper2D/Source/Paper2DEditor/Private/TileSetEditor/SingleTileEditorViewport.h
@@ -26,6 +26,7 @@ public:
 	virtual TSharedPtr<SWidget> MakeViewportToolbar() override;
 	virtual TSharedRef<FEditorViewportClient> MakeEditorViewportClient() override;
 	virtual void BindCommands() override;
+	virtual void OnFocusViewportToSelection() override;
 	// End of SEditorViewport interface
 
 	// ICommonEditorViewportToolbarInfoProvider interface
diff --git a/Engine/Plugins/2D/Paper2D/Source/Paper2DEditor/Private/TileSetEditor/SingleTileEditorViewportClient.cpp b/Engine/Plugins/2D/Paper2D/Source/Paper2DEditor/Private/TileSetEditor/SingleTileEditorViewportClient.cpp
--- a/Engine/Plugins/2D/Paper2D/Source/Paper2DEditor/Private/TileSetEditor/SingleTileEditorViewportClient.cpp
+++ b/Engine/Plugins/2D/Paper2D/Source/Paper2DEditor/Private/TileSetEditor/SingleTileEditorViewportClient.cpp
@@ -53,7 +53,6 @@ FBox FSingleTileEditorViewportClient::GetDesiredFocusBounds() const
 	return ComponentToFocusOn->Bounds.GetBox();
 }
 
-
 void FSingleTileEditorViewportClient::Tick(float DeltaSeconds)
 {
 	FPaperEditorViewportClient::Tick(DeltaSeconds);
@@ -166,6 +165,7 @@ void FSingleTileEditorViewportClient::InvalidateViewportAndHitProxies()
 
 void FSingleTileEditorViewportClient::SetTileIndex(int32 InTileIndex)
 {
+	const int32 OldTileIndex = TileBeingEditedIndex;
 	const bool bNewIndexValid = (InTileIndex >= 0) && (InTileIndex < TileSet->GetTileCount());
 	TileBeingEditedIndex = bNewIndexValid ? InTileIndex : INDEX_NONE;
 
@@ -213,6 +213,12 @@ void FSingleTileEditorViewportClient::SetTileIndex(int32 InTileIndex)
 	DesiredBounds.Max = HalfTileSize;
 	GeometryEditMode->SetNewGeometryPreferredBounds(DesiredBounds);
 
+	// Zoom to fit when we start editing a tile and weren't before, but leave it alone if we just changed tiles, in case they zoomed in or out further
+	if ((TileBeingEditedIndex != INDEX_NONE) && (OldTileIndex == INDEX_NONE))
+	{
+		RequestFocusOnSelection(/*bInstant=*/ true);
+	}
+
 	// Redraw the viewport
 	Invalidate();
 }
