Paper2D: Exposed default texture group and texture compression as part of the importer settings [UE-14127]
- The defaults remain the same as before (uncompressed, nearest filtering), but are user-editable now and are expressed slightly differently (using the '2D Pixels' group to get nearest filtering, instead of using 'UI' group and directly editing the filtering properties on the texture)

diff --git a/Engine/Plugins/2D/Paper2D/Source/Paper2DEditor/Classes/PaperImporterSettings.h b/Engine/Plugins/2D/Paper2D/Source/Paper2DEditor/Classes/PaperImporterSettings.h
--- a/Engine/Plugins/2D/Paper2D/Source/Paper2DEditor/Classes/PaperImporterSettings.h
+++ b/Engine/Plugins/2D/Paper2D/Source/Paper2DEditor/Classes/PaperImporterSettings.h
@@ -21,16 +21,29 @@ public:
 	UPROPERTY(config, EditAnywhere, Category=ImportSettings)
 	TArray<FString> BaseMapTextureSuffixes;
 
+	// The default texture group for imported sprite textures, tile sheets, etc... (typically set to UI for 'modern 2D' or 2D pixels for 'retro 2D')
+	UPROPERTY(config, EditAnywhere, Category=ImportSettings)
+	TEnumAsByte<TextureGroup> DefaultSpriteTextureGroup;
+
+	// Should texture compression settings be overridden on imported sprite textures, tile sheets, etc...?
+	UPROPERTY(config, EditAnywhere, Category=ImportSettings)
+	bool bOverrideTextureCompression;
+	
+	// Compression settings to use when building the texture.
+	// The default texture group for imported sprite textures, tile sheets, etc... (typically set to UI for 'modern 2D' or 2D pixels for 'retro 2D')
+	UPROPERTY(config, EditAnywhere, Category = ImportSettings, meta = (EditCondition = bOverrideTextureCompression))
+	TEnumAsByte<TextureCompressionSettings> DefaultSpriteTextureCompression;
+
 	// The default masked material for newly created sprites (masked means binary opacity: things are either opaque or see-thru, with nothing in between)
-	UPROPERTY(config, EditAnywhere, Category=Settings, meta=(AllowedClasses="MaterialInterface", DisplayName="Lit Default Masked Material"))
+	UPROPERTY(config, EditAnywhere, Category=MaterialSettings, meta=(AllowedClasses="MaterialInterface", DisplayName="Lit Default Masked Material"))
 	FStringAssetReference LitDefaultMaskedMaterialName;
 
 	// The default translucent material for newly created sprites (translucent means smooth opacity which can vary continuously from 0..1, but translucent rendering is more expensive that opaque or masked rendering and has different sorting rules)
-	UPROPERTY(config, EditAnywhere, Category=Settings, meta=(AllowedClasses="MaterialInterface", DisplayName="Lit Default Translucent Material"))
+	UPROPERTY(config, EditAnywhere, Category=MaterialSettings, meta=(AllowedClasses="MaterialInterface", DisplayName="Lit Default Translucent Material"))
 	FStringAssetReference LitDefaultTranslucentMaterialName;
 
 	// The default opaque material for newly created sprites
-	UPROPERTY(config, EditAnywhere, Category=Settings, meta=(AllowedClasses="MaterialInterface", DisplayName="Lit Default Opaque Material"))
+	UPROPERTY(config, EditAnywhere, Category=MaterialSettings, meta=(AllowedClasses="MaterialInterface", DisplayName="Lit Default Opaque Material"))
 	FStringAssetReference LitDefaultOpaqueMaterialName;
 
 public:
@@ -44,4 +57,7 @@ public:
 
 	// Pushes the lit materials
 	void PopulateMaterialsIntoInitParams(FSpriteAssetInitParameters& InitParams) const;
+
+	// Applies the compression settings to the specified texture
+	void ApplyTextureSettings(UTexture2D* Texture) const;
 };
diff --git a/Engine/Plugins/2D/Paper2D/Source/Paper2DEditor/Private/ContentBrowserExtensions/ContentBrowserExtensions.cpp b/Engine/Plugins/2D/Paper2D/Source/Paper2DEditor/Private/ContentBrowserExtensions/ContentBrowserExtensions.cpp
--- a/Engine/Plugins/2D/Paper2D/Source/Paper2DEditor/Private/ContentBrowserExtensions/ContentBrowserExtensions.cpp
+++ b/Engine/Plugins/2D/Paper2D/Source/Paper2DEditor/Private/ContentBrowserExtensions/ContentBrowserExtensions.cpp
@@ -5,6 +5,7 @@
 #include "ContentBrowserExtensions.h"
 #include "PaperSpriteFactory.h"
 #include "PaperTileSetFactory.h"
+#include "PaperImporterSettings.h"
 
 #include "ExtractSprites/SPaperExtractSpritesDialog.h"
 
@@ -117,11 +118,7 @@ struct FConfigureTexturesForSpriteUsageExtension : public FContentBrowserSelecte
 			const FAssetData& AssetData = *AssetIt;
 			if (UTexture2D* Texture = Cast<UTexture2D>(AssetData.GetAsset()))
 			{
-				Texture->Modify();
-				Texture->LODGroup = TEXTUREGROUP_UI;
-				Texture->CompressionSettings = TC_EditorIcon;
-				Texture->Filter = TF_Nearest;
-				Texture->PostEditChange();
+				GetDefault<UPaperImporterSettings>()->ApplyTextureSettings(Texture);
 			}
 		}
 	}
diff --git a/Engine/Plugins/2D/Paper2D/Source/Paper2DEditor/Private/PaperImporterSettings.cpp b/Engine/Plugins/2D/Paper2D/Source/Paper2DEditor/Private/PaperImporterSettings.cpp
--- a/Engine/Plugins/2D/Paper2D/Source/Paper2DEditor/Private/PaperImporterSettings.cpp
+++ b/Engine/Plugins/2D/Paper2D/Source/Paper2DEditor/Private/PaperImporterSettings.cpp
@@ -7,6 +7,9 @@
 // UPaperImporterSettings
 
 UPaperImporterSettings::UPaperImporterSettings()
+: DefaultSpriteTextureGroup(TEXTUREGROUP_Pixels2D)
+	, bOverrideTextureCompression(true)
+	, DefaultSpriteTextureCompression(TC_EditorIcon)
 {
 	NormalMapTextureSuffixes.Add(TEXT("_N"));
 	NormalMapTextureSuffixes.Add(TEXT("_Normal"));
@@ -51,3 +54,24 @@ void UPaperImporterSettings::PopulateMaterialsIntoInitParams(FSpriteAssetInitPar
 	InitParams.OpaqueMaterialOverride = Cast<UMaterialInterface>(LitDefaultOpaqueMaterialName.TryLoad());
 	InitParams.TranslucentMaterialOverride = Cast<UMaterialInterface>(LitDefaultTranslucentMaterialName.TryLoad());
 }
+
+void UPaperImporterSettings::ApplyTextureSettings(UTexture2D* Texture) const
+{
+	if (Texture->IsNormalMap())
+	{
+		// Leave normal maps alone
+	}
+	else
+	{
+		Texture->Modify();
+
+		Texture->LODGroup = DefaultSpriteTextureGroup;
+
+		if (bOverrideTextureCompression)
+		{
+			Texture->CompressionSettings = DefaultSpriteTextureCompression;
+		}
+
+		Texture->PostEditChange();
+	}
+}
diff --git a/Engine/Plugins/2D/Paper2D/Source/PaperSpriteSheetImporter/Private/PaperJsonSpriteSheetImporter.cpp b/Engine/Plugins/2D/Paper2D/Source/PaperSpriteSheetImporter/Private/PaperJsonSpriteSheetImporter.cpp
--- a/Engine/Plugins/2D/Paper2D/Source/PaperSpriteSheetImporter/Private/PaperJsonSpriteSheetImporter.cpp
+++ b/Engine/Plugins/2D/Paper2D/Source/PaperSpriteSheetImporter/Private/PaperJsonSpriteSheetImporter.cpp
@@ -496,12 +496,7 @@ UTexture2D* FPaperJsonSpriteSheetImporter::ImportTexture(const FString& TextureS
 	if (ImportedTexture != nullptr)
 	{
 		// Change the compression settings
-		//@TODO: Should we always be doing this (particularly the TF_Nearest seems potentially undesirable; maybe make it a property of the sprite sheet asset?)
-		ImportedTexture->Modify();
-		ImportedTexture->LODGroup = TEXTUREGROUP_UI;
-		ImportedTexture->CompressionSettings = TC_EditorIcon;
-		ImportedTexture->Filter = TF_Nearest;
-		ImportedTexture->PostEditChange();
+		GetDefault<UPaperImporterSettings>()->ApplyTextureSettings(ImportedTexture);
 	}
 
 	return ImportedTexture;
diff --git a/Engine/Plugins/2D/Paper2D/Source/PaperTiledImporter/Private/PaperTiledImporterFactory.cpp b/Engine/Plugins/2D/Paper2D/Source/PaperTiledImporter/Private/PaperTiledImporterFactory.cpp
--- a/Engine/Plugins/2D/Paper2D/Source/PaperTiledImporter/Private/PaperTiledImporterFactory.cpp
+++ b/Engine/Plugins/2D/Paper2D/Source/PaperTiledImporter/Private/PaperTiledImporterFactory.cpp
@@ -7,6 +7,7 @@
 #include "PaperJSONHelpers.h"
 #include "AssetToolsModule.h"
 #include "AssetRegistryModule.h"
+#include "PaperImporterSettings.h"
 #include "PackageTools.h"
 
 #define LOCTEXT_NAMESPACE "Paper2D"
@@ -477,11 +478,7 @@ UTexture2D* UPaperTiledImporterFactory::ImportTexture(const FString& SourceFilen
 	if (ImportedTexture != nullptr)
 	{
 		// Change the compression settings
-		ImportedTexture->Modify();
-		ImportedTexture->LODGroup = TEXTUREGROUP_UI;
-		ImportedTexture->CompressionSettings = TC_EditorIcon;
-		ImportedTexture->Filter = TF_Nearest;
-		ImportedTexture->PostEditChange();
+		GetDefault<UPaperImporterSettings>()->ApplyTextureSettings(ImportedTexture);
 	}
 
 	return ImportedTexture;
diff --git a/Engine/Plugins/2D/Paper2D/Source/SpriterImporter/Private/SpriterImporterFactory.cpp b/Engine/Plugins/2D/Paper2D/Source/SpriterImporter/Private/SpriterImporterFactory.cpp
--- a/Engine/Plugins/2D/Paper2D/Source/SpriterImporter/Private/SpriterImporterFactory.cpp
+++ b/Engine/Plugins/2D/Paper2D/Source/SpriterImporter/Private/SpriterImporterFactory.cpp
@@ -8,6 +8,7 @@
 #include "PackageTools.h"
 #include "SpriterImporterFactory.h"
 #include "PaperSpriterImportData.h"
+#include "PaperImporterSettings.h"
 
 #include "PaperSprite.h"
 
@@ -645,11 +646,7 @@ UTexture2D* USpriterImporterFactory::ImportTexture(const FString& SourceFilename
 	if (ImportedTexture != nullptr)
 	{
 		// Change the compression settings
-		ImportedTexture->Modify();
-		ImportedTexture->LODGroup = TEXTUREGROUP_UI;
-		ImportedTexture->CompressionSettings = TC_EditorIcon;
-		//@TODO: Determine if/when we should apply this ImportedTexture->Filter = TF_Nearest;
-		ImportedTexture->PostEditChange();
+		GetDefault<UPaperImporterSettings>()->ApplyTextureSettings(ImportedTexture);
 	}
 
 	return ImportedTexture;
diff --git a/Engine/Source/Runtime/Engine/Classes/Engine/TextureDefines.h b/Engine/Source/Runtime/Engine/Classes/Engine/TextureDefines.h
--- a/Engine/Source/Runtime/Engine/Classes/Engine/TextureDefines.h
+++ b/Engine/Source/Runtime/Engine/Classes/Engine/TextureDefines.h
@@ -50,7 +50,7 @@ enum TextureGroup
 	// no compression, created on import of a .IES file
 	TEXTUREGROUP_IESLightProfile UMETA(DisplayName="IESLightProfile"),
 	// non-filtered, useful for 2D rendering
-	TEXTUREGROUP_Pixels2D UMETA(DisplayName="2D Pixels"),
+	TEXTUREGROUP_Pixels2D UMETA(DisplayName="2D Pixels (unfiltered)"),
 	TEXTUREGROUP_MAX,
 };
 
