Tile Maps: Added support for tile rotation/flipping when generating tile map collision [UE-13615]

diff --git a/Engine/Plugins/2D/Paper2D/Source/Paper2D/Private/PaperSprite.cpp b/Engine/Plugins/2D/Paper2D/Source/Paper2D/Private/PaperSprite.cpp
--- a/Engine/Plugins/2D/Paper2D/Source/Paper2D/Private/PaperSprite.cpp
+++ b/Engine/Plugins/2D/Paper2D/Source/Paper2D/Private/PaperSprite.cpp
@@ -2077,11 +2077,12 @@ void FSpriteGeometryCollisionBuilderBase::AddBoxCollisionShapesToBodySetup(const
 			{
 				case ESpriteCollisionMode::Use3DPhysics:
 					{
-					   const FVector BoxSize3D = (PaperAxisX * BoxSize2D.X) + (PaperAxisY * BoxSize2D.Y) + (PaperAxisZ * CollisionThickness);
+						const FVector BoxPos3D = (PaperAxisX * CenterInScaledSpace.X) + (PaperAxisY * CenterInScaledSpace.Y) + (PaperAxisZ * ZOffsetAmount);
+						const FVector BoxSize3D = (PaperAxisX * BoxSize2D.X) + (PaperAxisY * BoxSize2D.Y) + (PaperAxisZ * CollisionThickness);
 
 						// Create a new box primitive
 						FKBoxElem& Box = *new (MyBodySetup->AggGeom.BoxElems) FKBoxElem(FMath::Abs(BoxSize3D.X), FMath::Abs(BoxSize3D.Y), FMath::Abs(BoxSize3D.Z));
-						Box.Center = (PaperAxisX * CenterInScaledSpace.X) + (PaperAxisY * CenterInScaledSpace.Y) + (PaperAxisZ * ZOffsetAmount);
+						Box.Center = BoxPos3D;
 						Box.Orientation = FQuat(FRotator(Shape.Rotation, 0.0f, 0.0f));
 					}
 					break;
diff --git a/Engine/Plugins/2D/Paper2D/Source/Paper2D/Private/PaperTileLayer.cpp b/Engine/Plugins/2D/Paper2D/Source/Paper2D/Private/PaperTileLayer.cpp
--- a/Engine/Plugins/2D/Paper2D/Source/Paper2D/Private/PaperTileLayer.cpp
+++ b/Engine/Plugins/2D/Paper2D/Source/Paper2D/Private/PaperTileLayer.cpp
@@ -9,6 +9,36 @@
 #include "PaperTileLayer.h"
 #include "PaperTileMapComponent.h"
 
+// Handles the rotation and flipping of collision geometry from a tile
+// 0,5,6,3 are clockwise rotations of a regular tile
+// 4,7,2,1 are clockwise rotations of a horizontally flipped tile
+const static FTransform TilePermutationTransforms[8] =
+{
+	// 000 - normal
+	FTransform::Identity,
+
+	// 001 - diagonal
+	FTransform(FRotator(  90.0f, 0.0f, 0.0f), FVector::ZeroVector, -PaperAxisX.GetAbs() + PaperAxisY.GetAbs() + PaperAxisZ.GetAbs()),
+
+	// 010 - flip Y
+	FTransform(FRotator(-180.0f, 0.0f, 0.0f), FVector::ZeroVector, -PaperAxisX.GetAbs() + PaperAxisY.GetAbs() + PaperAxisZ.GetAbs()),
+
+	// 011 - diagonal then flip Y (rotate 270 clockwise)
+	FTransform(FRotator(  90.0f, 0.0f, 0.0f)),
+
+	// 100 - flip X
+	FTransform(FRotator::ZeroRotator, FVector::ZeroVector, -PaperAxisX.GetAbs() + PaperAxisY.GetAbs() + PaperAxisZ.GetAbs()),
+
+	// 101 - diagonal then flip X (clockwise 90)
+	FTransform(FRotator( -90.0f, 0.0f, 0.0f)),
+
+	// 110 - flip X and flip Y (rotate 180 either way)
+	FTransform(FRotator(-180.0f, 0.0f, 0.0f)),
+
+	// 111 - diagonal then flip X and Y
+	FTransform(FRotator(-90.0f, 0.0f, 0.0f), FVector::ZeroVector, -PaperAxisX.GetAbs() + PaperAxisY.GetAbs() + PaperAxisZ.GetAbs()),
+};
+
 //////////////////////////////////////////////////////////////////////////
 // FTileMapLayerReregisterContext
 
@@ -65,25 +95,39 @@ public:
 		ZOffsetAmount = InZOffset;
 	}
 
-	void SetCellOffset(const FVector2D& NewOffset)
+	void SetCellOffset(const FVector2D& NewOffset, const FTransform& NewTransform)
 	{
 		CurrentCellOffset = NewOffset;
+		MyTransform = NewTransform;
 	}
 
 protected:
 	// FSpriteGeometryCollisionBuilderBase interface
 	virtual FVector2D ConvertTextureSpaceToPivotSpace(const FVector2D& Input) const override
 	{
-		return FVector2D(CurrentCellOffset.X + Input.X, CurrentCellOffset.Y - Input.Y);
+		const FVector LocalPos3D = (Input.X * PaperAxisX) - (Input.Y * PaperAxisY);
+		const FVector RotatedLocalPos3D = MyTransform.TransformPosition(LocalPos3D);
+
+		const float OutputX = CurrentCellOffset.X + FVector::DotProduct(RotatedLocalPos3D, PaperAxisX);
+		const float OutputY = CurrentCellOffset.Y + FVector::DotProduct(RotatedLocalPos3D, PaperAxisY);
+
+		return FVector2D(OutputX, OutputY);
 	}
 
 	virtual FVector2D ConvertTextureSpaceToPivotSpaceNoTranslation(const FVector2D& Input) const override
 	{
-		return Input;
+		const FVector LocalPos3D = (Input.X * PaperAxisX) + (Input.Y * PaperAxisY);
+		const FVector RotatedLocalPos3D = MyTransform.TransformVector(LocalPos3D);
+
+		const float OutputX = FVector::DotProduct(RotatedLocalPos3D, PaperAxisX);
+		const float OutputY = FVector::DotProduct(RotatedLocalPos3D, PaperAxisY);
+
+		return FVector2D(OutputX, OutputY);
 	}
 	// End of FSpriteGeometryCollisionBuilderBase
 
 protected:
+	FTransform MyTransform;
 	UPaperTileLayer* MySprite;
 	FVector2D CurrentCellOffset;
 };
@@ -233,9 +277,11 @@ void UPaperTileLayer::AugmentBodySetup(UBodySetup* ShapeBodySetup)
 			{
 				if (const FPaperTileMetadata* CellMetadata = CellInfo.TileSet->GetTileMetadata(CellInfo.GetTileIndex()))
 				{
-					//@TODO: Add support for flipped/mirrored/rotated tiles here (and inside FPaperTileLayerToBodySetupBuilder::ConvertTextureSpaceToPivotSpace(NoTranslation))
+					const int32 Flags = CellInfo.GetFlagsAsIndex();
+
+					const FTransform& LocalTransform = TilePermutationTransforms[Flags];
 					const FVector2D CellOffset(TileWidth * CellX, TileHeight * -CellY);
-					CollisionBuilder.SetCellOffset(CellOffset);
+					CollisionBuilder.SetCellOffset(CellOffset, LocalTransform);
 
 					CollisionBuilder.ProcessGeometry(CellMetadata->CollisionData);
 				}
