Tile Maps: Changed the default collision domain for Tile Maps from None to 3D [UE-10933]
- Note: Collision must still be authored in the tile set editor per-tile in order for any actual collision geometry to be generated

diff --git a/Engine/Plugins/2D/Paper2D/Source/Paper2D/Private/PaperTileMap.cpp b/Engine/Plugins/2D/Paper2D/Source/Paper2D/Private/PaperTileMap.cpp
--- a/Engine/Plugins/2D/Paper2D/Source/Paper2D/Private/PaperTileMap.cpp
+++ b/Engine/Plugins/2D/Paper2D/Source/Paper2D/Private/PaperTileMap.cpp
@@ -24,7 +24,7 @@ UPaperTileMap::UPaperTileMap(const FObjectInitializer& ObjectInitializer)
 	SeparationPerTileY = 0.0f;
 	SeparationPerLayer = 4.0f;
 	CollisionThickness = 50.0f;
-	SpriteCollisionDomain = ESpriteCollisionMode::None;
+	SpriteCollisionDomain = ESpriteCollisionMode::Use3DPhysics;
 
 #if WITH_EDITORONLY_DATA
 	SelectedLayerIndex = INDEX_NONE;
