Paper2D: Added naming template to "Auto" sprite extract mode. [UE-13343]

diff --git a/Engine/Plugins/2D/Paper2D/Source/Paper2DEditor/Private/ExtractSprites/SPaperExtractSpritesDialog.cpp b/Engine/Plugins/2D/Paper2D/Source/Paper2DEditor/Private/ExtractSprites/SPaperExtractSpritesDialog.cpp
--- a/Engine/Plugins/2D/Paper2D/Source/Paper2DEditor/Private/ExtractSprites/SPaperExtractSpritesDialog.cpp
+++ b/Engine/Plugins/2D/Paper2D/Source/Paper2DEditor/Private/ExtractSprites/SPaperExtractSpritesDialog.cpp
@@ -142,10 +142,10 @@ void SPaperExtractSpritesDialog::Construct(const FArguments& InArgs, UTexture2D*
 
 	ExtractSpriteSettings = NewObject<UPaperExtractSpritesSettings>();
 	ExtractSpriteSettings->AddToRoot();
+	ExtractSpriteSettings->NamingTemplate = "Sprite_{0}";
 
 	ExtractSpriteGridSettings = NewObject<UPaperExtractSpriteGridSettings>();
 	ExtractSpriteGridSettings->AddToRoot();
-	ExtractSpriteGridSettings->NamingTemplate = "Sprite_{0}";
 	ExtractSpriteGridSettings->CellWidth = InSourceTexture->GetImportedSize().X;
 	ExtractSpriteGridSettings->CellHeight = InSourceTexture->GetImportedSize().Y;
 
@@ -267,6 +267,13 @@ void SPaperExtractSpritesDialog::PreviewExtractedSprites()
 	SlowTask.MakeDialog(false);
 	SlowTask.EnterProgressFrame();
 
+	FString NamingTemplate = ExtractSpriteSettings->NamingTemplate;
+	if (NamingTemplate.Find(TEXT("{0}")) == INDEX_NONE)
+	{
+		NamingTemplate.Append(TEXT("_{0}"));
+	}
+	int32 ExtractedRectIndex = ExtractSpriteSettings->NamingStartIndex;
+
 	const FString DefaultSuffix = TEXT("Sprite");
 	if (ExtractSpriteSettings->SpriteExtractMode == ESpriteExtractMode::Auto)
 	{
@@ -318,29 +325,22 @@ void SPaperExtractSpritesDialog::PreviewExtractedSprites()
 		};
 		FRectangleSortHelper RectSorter(ExtractedRects);
 
-		int32 ExtractedRectIndex = 0;
 		for (FIntRect Rect : ExtractedRects)
 		{
 			FPaperExtractedSprite* Sprite = new(ExtractedSprites)FPaperExtractedSprite();
 			Sprite->Rect = Rect;
-			Sprite->Name = FString::Printf(TEXT("%s_%d"), *DefaultSuffix, ExtractedRectIndex);
+			Sprite->Name = NamingTemplate;
+			Sprite->Name.ReplaceInline(TEXT("{0}"), *FString::Printf(TEXT("%d"), ExtractedRectIndex));
+			
 			ExtractedRectIndex++;
 		}
 	}
 	else
 	{
-		FString NamingTemplate = ExtractSpriteGridSettings->NamingTemplate;
-		if (NamingTemplate.Find(TEXT("{0}")) == INDEX_NONE)
-		{
-			NamingTemplate.Append(TEXT("_{0}"));
-		}
-
 		// Calculate rects
 		ExtractedSprites.Empty();
 		if (SourceTexture != nullptr)
 		{
-			int32 ExtractedRectIndex = ExtractSpriteGridSettings->NamingStartIndex;
-
 			const FIntPoint TextureSize(SourceTexture->GetImportedSize());
 			const int32 TextureWidth = TextureSize.X;
 			const int32 TextureHeight = TextureSize.Y;
