MeshPaint: Added support for texture overrides to the sprite adapter, allowing painting on the variable texture in addition to the material referenced textures [UE-14210]

diff --git a/Engine/Plugins/2D/Paper2D/Source/Paper2D/Classes/PaperSpriteComponent.h b/Engine/Plugins/2D/Paper2D/Source/Paper2D/Classes/PaperSpriteComponent.h
--- a/Engine/Plugins/2D/Paper2D/Source/Paper2D/Classes/PaperSpriteComponent.h
+++ b/Engine/Plugins/2D/Paper2D/Source/Paper2D/Classes/PaperSpriteComponent.h
@@ -89,6 +89,10 @@ public:
 	virtual int32 GetNumMaterials() const override;
 	// End of UPrimitiveComponent interface
 
+#if WITH_EDITOR
+	void SetTransientTextureOverride(const UTexture* TextureToModifyOverrideFor, UTexture* OverrideTexture);
+#endif
+
 protected:
 	friend class FPaperSpriteSceneProxy;
 };
diff --git a/Engine/Plugins/2D/Paper2D/Source/Paper2D/Private/PaperRenderSceneProxy.cpp b/Engine/Plugins/2D/Paper2D/Source/Paper2D/Private/PaperRenderSceneProxy.cpp
--- a/Engine/Plugins/2D/Paper2D/Source/Paper2D/Private/PaperRenderSceneProxy.cpp
+++ b/Engine/Plugins/2D/Paper2D/Source/Paper2D/Private/PaperRenderSceneProxy.cpp
@@ -109,12 +109,15 @@ public:
 	const FMaterialRenderProxy* const Parent;
 	const UTexture* BaseTexture;
 	FAdditionalSpriteTextureArray AdditionalTextures;
+	UE_EXPAND_IF_WITH_EDITOR(const FPaperRenderSceneProxyTextureOverrideMap& TextureOverrideList);
 
+public:
 	/** Initialization constructor. */
-	FSpriteTextureOverrideRenderProxy(const FMaterialRenderProxy* InParent, const UTexture* InBaseTexture, FAdditionalSpriteTextureArray InAdditionalTextures)
+	FSpriteTextureOverrideRenderProxy(const FMaterialRenderProxy* InParent, const UTexture* InBaseTexture, FAdditionalSpriteTextureArray InAdditionalTextures UE_EXPAND_IF_WITH_EDITOR(, const FPaperRenderSceneProxyTextureOverrideMap& InTextureOverrideList))
 		: Parent(InParent)
 		, BaseTexture(InBaseTexture)
 		, AdditionalTextures(InAdditionalTextures)
+		UE_EXPAND_IF_WITH_EDITOR(, TextureOverrideList(InTextureOverrideList))
 	{}
 
 	virtual ~FSpriteTextureOverrideRenderProxy()
@@ -146,11 +149,11 @@ public:
 		return Parent->GetScalarValue(ParameterName, OutValue, Context);
 	}
 
-	virtual bool GetTextureValue(const FName ParameterName,const UTexture** OutValue, const FMaterialRenderContext& Context) const override
+	virtual bool GetTextureValue(const FName ParameterName, const UTexture** OutValue, const FMaterialRenderContext& Context) const override
 	{
 		if (ParameterName == TextureParameterName)
 		{
-			*OutValue = BaseTexture;
+			*OutValue = ApplyEditorOverrides(BaseTexture);
 			return true;
 		}
 		else if (ParameterName.GetComparisonIndex() == AdditionalTextureParameterRootName.GetComparisonIndex())
@@ -158,7 +161,7 @@ public:
 			const int32 AdditionalSlotIndex = ParameterName.GetNumber() - 1;
 			if (AdditionalTextures.IsValidIndex(AdditionalSlotIndex))
 			{
-				*OutValue = AdditionalTextures[AdditionalSlotIndex];
+				*OutValue = ApplyEditorOverrides(AdditionalTextures[AdditionalSlotIndex]);
 				return true;
 			}
 		}
@@ -166,7 +169,26 @@ public:
 		return Parent->GetTextureValue(ParameterName, OutValue, Context);
 	}
 
-protected:
+#if WITH_EDITOR
+	inline const UTexture* ApplyEditorOverrides(const UTexture* InTexture) const
+	{
+		if (TextureOverrideList.Num() > 0)
+		{
+			if (const UTexture* const* OverridePtr = TextureOverrideList.Find(InTexture))
+			{
+				return *OverridePtr;
+			}
+		}
+
+		return InTexture;
+	}
+#else
+	FORCEINLINE const UTexture* ApplyEditorOverrides(const UTexture* InTexture) const
+	{
+		return InTexture;
+	}
+#endif
+
 	static const FName TextureParameterName;
 	static const FName AdditionalTextureParameterRootName;
 };
@@ -380,7 +402,7 @@ void FPaperRenderSceneProxy::GetNewBatchMeshes(const FSceneView* View, bool bUse
 			}
 
 			// Create a texture override material proxy and register it as a dynamic resource so that it won't be deleted until the rendering thread has finished with it
-			FSpriteTextureOverrideRenderProxy* TextureOverrideMaterialProxy = new FSpriteTextureOverrideRenderProxy(ParentMaterialProxy, Batch.BaseTexture, Batch.AdditionalTextures);
+			FSpriteTextureOverrideRenderProxy* TextureOverrideMaterialProxy = new FSpriteTextureOverrideRenderProxy(ParentMaterialProxy, Batch.BaseTexture, Batch.AdditionalTextures UE_EXPAND_IF_WITH_EDITOR(, TextureOverrideList));
 			Collector.RegisterOneFrameMaterialProxy(TextureOverrideMaterialProxy);
 
 			Mesh.UseDynamicData = true;
@@ -478,7 +500,7 @@ void FPaperRenderSceneProxy::GetBatchMesh(const FSceneView* View, bool bUseOverr
 			}
 
 			// Create a texture override material proxy and register it as a dynamic resource so that it won't be deleted until the rendering thread has finished with it
-			FSpriteTextureOverrideRenderProxy* TextureOverrideMaterialProxy = new FSpriteTextureOverrideRenderProxy(ParentMaterialProxy, Record.BaseTexture, Record.AdditionalTextures);
+			FSpriteTextureOverrideRenderProxy* TextureOverrideMaterialProxy = new FSpriteTextureOverrideRenderProxy(ParentMaterialProxy, Record.BaseTexture, Record.AdditionalTextures UE_EXPAND_IF_WITH_EDITOR(, TextureOverrideList));
 			Collector.RegisterOneFrameMaterialProxy(TextureOverrideMaterialProxy);
 
 			Mesh.MaterialRenderProxy = TextureOverrideMaterialProxy;
@@ -596,3 +618,17 @@ bool FPaperRenderSceneProxy::IsCollisionView(const FEngineShowFlags& EngineShowF
 
 	return bInCollisionView;
 }
+
+#if WITH_EDITOR
+void FPaperRenderSceneProxy::SetTransientTextureOverride_RenderThread(const UTexture* InTextureToModifyOverrideFor, UTexture* InOverrideTexture)
+{
+	if (InOverrideTexture != nullptr)
+	{
+		TextureOverrideList.FindOrAdd(InTextureToModifyOverrideFor) = InOverrideTexture;
+	}
+	else
+	{
+		TextureOverrideList.Remove(InTextureToModifyOverrideFor);
+	}
+}
+#endif
\ No newline at end of file
diff --git a/Engine/Plugins/2D/Paper2D/Source/Paper2D/Private/PaperRenderSceneProxy.h b/Engine/Plugins/2D/Paper2D/Source/Paper2D/Private/PaperRenderSceneProxy.h
--- a/Engine/Plugins/2D/Paper2D/Source/Paper2D/Private/PaperRenderSceneProxy.h
+++ b/Engine/Plugins/2D/Paper2D/Source/Paper2D/Private/PaperRenderSceneProxy.h
@@ -7,6 +7,14 @@
 
 #define TEST_BATCHING 0
 
+#if WITH_EDITOR
+typedef TMap<const UTexture*, const UTexture*> FPaperRenderSceneProxyTextureOverrideMap;
+
+#define UE_EXPAND_IF_WITH_EDITOR(...) __VA_ARGS__
+#else
+#define UE_EXPAND_IF_WITH_EDITOR(...)
+#endif
+
 //////////////////////////////////////////////////////////////////////////
 // FPaperSpriteVertex
 
@@ -114,6 +122,10 @@ public:
 	void SetDrawCall_RenderThread(const FSpriteDrawCallRecord& NewDynamicData);
 	void SetBodySetup_RenderThread(UBodySetup* NewSetup);
 
+#if WITH_EDITOR
+	void SetTransientTextureOverride_RenderThread(const UTexture* InTextureToModifyOverrideFor, UTexture* InOverrideTexture);
+#endif
+
 	class UMaterialInterface* GetMaterial() const
 	{
 		return Material;
@@ -148,4 +160,7 @@ protected:
 
 	// The Collision Response of the component being proxied
 	FCollisionResponseContainer CollisionResponse;
+
+	// The texture override list
+	UE_EXPAND_IF_WITH_EDITOR(FPaperRenderSceneProxyTextureOverrideMap TextureOverrideList);
 };
diff --git a/Engine/Plugins/2D/Paper2D/Source/Paper2D/Private/PaperSpriteComponent.cpp b/Engine/Plugins/2D/Paper2D/Source/Paper2D/Private/PaperSpriteComponent.cpp
--- a/Engine/Plugins/2D/Paper2D/Source/Paper2D/Private/PaperSpriteComponent.cpp
+++ b/Engine/Plugins/2D/Paper2D/Source/Paper2D/Private/PaperSpriteComponent.cpp
@@ -344,6 +344,20 @@ void UPaperSpriteComponent::CheckForErrors()
 }
 #endif
 
+#if WITH_EDITOR
+void UPaperSpriteComponent::SetTransientTextureOverride(const UTexture* TextureToModifyOverrideFor, UTexture* OverrideTexture)
+{
+	ENQUEUE_UNIQUE_RENDER_COMMAND_THREEPARAMETER(
+		FSendPaperSpriteComponentDynamicData,
+		FPaperSpriteSceneProxy*, InSceneProxy, (FPaperSpriteSceneProxy*)SceneProxy,
+		const UTexture*, InTextureToModifyOverrideFor, TextureToModifyOverrideFor,
+		UTexture*, InOverrideTexture, OverrideTexture,
+		{
+			InSceneProxy->SetTransientTextureOverride_RenderThread(InTextureToModifyOverrideFor, InOverrideTexture);
+		});
+}
+#endif
+
 //////////////////////////////////////////////////////////////////////////
 
 #undef LOCTEXT_NAMESPACE
\ No newline at end of file
diff --git a/Engine/Plugins/2D/Paper2D/Source/Paper2DEditor/Private/MeshPainting/MeshPaintSpriteAdapter.cpp b/Engine/Plugins/2D/Paper2D/Source/Paper2DEditor/Private/MeshPainting/MeshPaintSpriteAdapter.cpp
--- a/Engine/Plugins/2D/Paper2D/Source/Paper2DEditor/Private/MeshPainting/MeshPaintSpriteAdapter.cpp
+++ b/Engine/Plugins/2D/Paper2D/Source/Paper2DEditor/Private/MeshPainting/MeshPaintSpriteAdapter.cpp
@@ -136,16 +136,8 @@ void FMeshPaintSpriteAdapter::QueryPaintableTextures(int32 MaterialIndex, int32&
 
 void FMeshPaintSpriteAdapter::ApplyOrRemoveTextureOverride(UTexture* SourceTexture, UTexture* OverrideTexture) const
 {
-	UPaperSprite* Sprite = SpriteComponent->GetSprite();
-	checkSlow(Sprite != nullptr);
-
-	if (Sprite->GetSourceTexture() == SourceTexture)
-	{
-		//@TOOD: Need to set SourceTexture override on the component or something like that!
-	}
-	//else if (Sprite->AdditionalSourceTextures.Contains(SourceTexture)) //@TOOD: Need to swap it here too
-	//{
-	//}
+	// Apply it to the sprite component
+	SpriteComponent->SetTransientTextureOverride(SourceTexture, OverrideTexture);
 
 	// Make sure we swap it out on any textures that aren't part of the sprite as well
 	DefaultApplyOrRemoveTextureOverride(SpriteComponent, SourceTexture, OverrideTexture);
