Mesh Paint: Added support for painting on additional sprite textures [UE-14426]

diff --git a/Engine/Plugins/2D/Paper2D/Source/Paper2D/Private/PaperSpriteComponent.cpp b/Engine/Plugins/2D/Paper2D/Source/Paper2D/Private/PaperSpriteComponent.cpp
--- a/Engine/Plugins/2D/Paper2D/Source/Paper2D/Private/PaperSpriteComponent.cpp
+++ b/Engine/Plugins/2D/Paper2D/Source/Paper2D/Private/PaperSpriteComponent.cpp
@@ -233,10 +233,17 @@ void UPaperSpriteComponent::GetUsedTextures(TArray<UTexture*>& OutTextures, EMat
 		{
 			OutTextures.AddUnique(BakedTexture);
 		}
-	}
 
-	// Add any extra textures
-	//@TODO: 
+		FAdditionalSpriteTextureArray AdditionalTextureList;
+		SourceSprite->GetBakedAdditionalSourceTextures(/*out*/ AdditionalTextureList);
+		for (UTexture* AdditionalTexture : AdditionalTextureList)
+		{
+			if (AdditionalTexture != nullptr)
+			{
+				OutTextures.AddUnique(AdditionalTexture);
+			}
+		}
+	}
 
 	// Get any textures referenced by our materials
 	Super::GetUsedTextures(OutTextures, QualityLevel);
diff --git a/Engine/Plugins/2D/Paper2D/Source/Paper2DEditor/Private/MeshPainting/MeshPaintSpriteAdapter.cpp b/Engine/Plugins/2D/Paper2D/Source/Paper2DEditor/Private/MeshPainting/MeshPaintSpriteAdapter.cpp
--- a/Engine/Plugins/2D/Paper2D/Source/Paper2DEditor/Private/MeshPainting/MeshPaintSpriteAdapter.cpp
+++ b/Engine/Plugins/2D/Paper2D/Source/Paper2DEditor/Private/MeshPainting/MeshPaintSpriteAdapter.cpp
@@ -121,10 +121,22 @@ void FMeshPaintSpriteAdapter::QueryPaintableTextures(int32 MaterialIndex, int32&
 	int32 ForceIndex = INDEX_NONE;
 	if (UTexture2D* SourceTexture = Sprite->GetBakedTexture())
 	{
-		new (InOutTextureList) FPaintableTexture(SourceTexture, 0);
-		ForceIndex = InOutTextureList.Num() - 1;
+		FPaintableTexture PaintableTexture(SourceTexture, 0);
+		ForceIndex = InOutTextureList.AddUnique(PaintableTexture);
 	}
-	
+
+	// Grab the additional textures next
+	FAdditionalSpriteTextureArray AdditionalTextureList;
+	Sprite->GetBakedAdditionalSourceTextures(/*out*/ AdditionalTextureList);
+	for (UTexture* AdditionalTexture : AdditionalTextureList)
+	{
+		if (AdditionalTexture != nullptr)
+		{
+			FPaintableTexture PaintableTexture(AdditionalTexture, 0);
+			InOutTextureList.AddUnique(AdditionalTexture);
+		}
+	}
+
 	// Now ask the material
 	DefaultQueryPaintableTextures(MaterialIndex, SpriteComponent, OutDefaultIndex, InOutTextureList);
 
