Paper2D: Defer rebuilding geometry and collision data during property spinner drags [UE-13180]

diff --git a/Engine/Plugins/2D/Paper2D/Source/Paper2D/Private/PaperSprite.cpp b/Engine/Plugins/2D/Paper2D/Source/Paper2D/Private/PaperSprite.cpp
--- a/Engine/Plugins/2D/Paper2D/Source/Paper2D/Private/PaperSprite.cpp
+++ b/Engine/Plugins/2D/Paper2D/Source/Paper2D/Private/PaperSprite.cpp
@@ -576,6 +576,14 @@ void UPaperSprite::PostEditChangeProperty(FPropertyChangedEvent& PropertyChanged
 		// bBothModified = true;
 	//}
 
+	// Don't do rebuilds during an interactive event to make things more responsive.
+	// They'll always be followed by a ValueSet event at the end to force the change.
+	if (PropertyChangedEvent.ChangeType == EPropertyChangeType::Interactive)
+	{
+		bCollisionDataModified = false;
+		bRenderDataModified = false;
+		bBothModified = false;
+	}
 
 	if (bCollisionDataModified || bBothModified)
 	{
