Fix for shadow variable.

diff --git a/Engine/Plugins/2D/Paper2D/Source/Paper2DEditor/Private/TileSetEditor/TileSetEditorViewportClient.cpp b/Engine/Plugins/2D/Paper2D/Source/Paper2DEditor/Private/TileSetEditor/TileSetEditorViewportClient.cpp
--- a/Engine/Plugins/2D/Paper2D/Source/Paper2DEditor/Private/TileSetEditor/TileSetEditorViewportClient.cpp
+++ b/Engine/Plugins/2D/Paper2D/Source/Paper2DEditor/Private/TileSetEditor/TileSetEditorViewportClient.cpp
@@ -63,9 +63,9 @@ void FTileSetEditorViewportClient::Draw(FViewport* Viewport, FCanvas* Canvas)
 			const float Width = (TileSet->TileWidth - 2) * ZoomAmount;
 			const float Height = (TileSet->TileHeight - 2) * ZoomAmount;
 
-			for (int32 TileIndex = 0; TileIndex < NumTiles; ++TileIndex)
+			for (int32 OtherTileIndex = 0; OtherTileIndex < NumTiles; ++OtherTileIndex)
 			{
-				if (const FPaperTileMetadata* TileMetadata = TileSet->GetTileMetadata(TileIndex))
+				if (const FPaperTileMetadata* TileMetadata = TileSet->GetTileMetadata(OtherTileIndex))
 				{
 					const bool bShowDueToCollision = TileMetadata->HasCollision() && bShowTilesWithCollision;
 					const bool bShowDueToMetaData = TileMetadata->HasMetaData() && bShowTilesWithMetaData;
@@ -73,7 +73,7 @@ void FTileSetEditorViewportClient::Draw(FViewport* Viewport, FCanvas* Canvas)
 					if (bShowDueToCollision || bShowDueToMetaData)
 					{
 						FVector2D TileUV;
-						TileSet->GetTileUV(TileIndex, /*out*/ TileUV);
+						TileSet->GetTileUV(OtherTileIndex, /*out*/ TileUV);
 
 						const float XPos = (TileUV.X + 1 - ZoomPos.X) * ZoomAmount;
 						const float YPos = (TileUV.Y + 1 - ZoomPos.Y) * ZoomAmount;
