Tile Maps: Changed the per-tile grid to display even when 'show grid' is turned off if the component is selected; now only the unselected layer 0 outline is gated by 'show grid' [UE-14437]

diff --git a/Engine/Plugins/2D/Paper2D/Source/Paper2D/Private/PaperTileMapRenderSceneProxy.cpp b/Engine/Plugins/2D/Paper2D/Source/Paper2D/Private/PaperTileMapRenderSceneProxy.cpp
--- a/Engine/Plugins/2D/Paper2D/Source/Paper2D/Private/PaperTileMapRenderSceneProxy.cpp
+++ b/Engine/Plugins/2D/Paper2D/Source/Paper2D/Private/PaperTileMapRenderSceneProxy.cpp
@@ -126,73 +126,64 @@ void FPaperTileMapRenderSceneProxy::GetDynamicMeshElements(const TArray<const FS
 				RenderBounds(PDI, View->Family->EngineShowFlags, GetBounds(), IsSelected());
 
 #if WITH_EDITOR
+				const bool bShowAsSelected = IsSelected();
+				const bool bEffectivelySelected = bShowAsSelected || IsHovered();
+
+				const uint8 DPG = SDPG_Foreground;//GetDepthPriorityGroup(View);
+
+				// Draw separation wires if selected
+				const FLinearColor OverrideColor = GetSelectionColor(FLinearColor::White, bShowAsSelected, IsHovered());
+
+				FTransform LocalToWorld(GetLocalToWorld());
+
 				// Draw the debug outline
-				if (View->Family->EngineShowFlags.Grid)
+				if (bEffectivelySelected)
 				{
-					const uint8 DPG = SDPG_Foreground;//GetDepthPriorityGroup(View);
-
-					// Draw separation wires if selected
-					FLinearColor OverrideColor;
-					bool bUseOverrideColor = false;
-					bool bEffectivelySelected = false;
-					const bool bShowAsSelected = !(GIsEditor && View->Family->EngineShowFlags.Selection) || IsSelected();
-					if (bShowAsSelected || IsHovered())
-					{
-						bEffectivelySelected = true;
-						bUseOverrideColor = true;
-						OverrideColor = GetSelectionColor(FLinearColor::White, bShowAsSelected, IsHovered());
-					}
-
-					FTransform LocalToWorld(GetLocalToWorld());
+					const int32 SelectedLayerIndex = TileMap->SelectedLayerIndex;
 
-					if (bEffectivelySelected)
+					if (bShowPerLayerGrid)
 					{
-						const int32 SelectedLayerIndex = TileMap->SelectedLayerIndex;
-
-						if (bShowPerLayerGrid)
+						// Draw a bound for every layer but the selected one (and even that one if the per-tile grid is off)
+						for (int32 LayerIndex = 0; LayerIndex < TileMap->TileLayers.Num(); ++LayerIndex)
 						{
-							// Draw a bound for every layer but the selected one (and even that one if the per-tile grid is off)
-							for (int32 LayerIndex = 0; LayerIndex < TileMap->TileLayers.Num(); ++LayerIndex)
+							if ((LayerIndex != SelectedLayerIndex) || !bShowPerTileGrid)
 							{
-								if ((LayerIndex != SelectedLayerIndex) || !bShowPerTileGrid)
-								{
-									DrawBoundsForLayer(PDI, OverrideColor, LayerIndex);
-								}
+								DrawBoundsForLayer(PDI, OverrideColor, LayerIndex);
 							}
 						}
+					}
 
-						if (bShowPerTileGrid && (SelectedLayerIndex != INDEX_NONE))
+					if (bShowPerTileGrid && (SelectedLayerIndex != INDEX_NONE))
+					{
+						// Draw horizontal lines on the selection
+						for (int32 Y = 0; Y <= TileMap->MapHeight; ++Y)
 						{
-							// Draw horizontal lines on the selection
-							for (int32 Y = 0; Y <= TileMap->MapHeight; ++Y)
-							{
-								int32 X = 0;
-								const FVector Start(TileMap->GetTilePositionInLocalSpace(X, Y, SelectedLayerIndex));
+							int32 X = 0;
+							const FVector Start(TileMap->GetTilePositionInLocalSpace(X, Y, SelectedLayerIndex));
 
-								X = TileMap->MapWidth;
-								const FVector End(TileMap->GetTilePositionInLocalSpace(X, Y, SelectedLayerIndex));
+							X = TileMap->MapWidth;
+							const FVector End(TileMap->GetTilePositionInLocalSpace(X, Y, SelectedLayerIndex));
 
-								PDI->DrawLine(LocalToWorld.TransformPosition(Start), LocalToWorld.TransformPosition(End), OverrideColor, DPG, 0.0f, DepthBias);
-							}
+							PDI->DrawLine(LocalToWorld.TransformPosition(Start), LocalToWorld.TransformPosition(End), OverrideColor, DPG, 0.0f, DepthBias);
+						}
 
-							// Draw vertical lines
-							for (int32 X = 0; X <= TileMap->MapWidth; ++X)
-							{
-								int32 Y = 0;
-								const FVector Start(TileMap->GetTilePositionInLocalSpace(X, Y, SelectedLayerIndex));
+						// Draw vertical lines
+						for (int32 X = 0; X <= TileMap->MapWidth; ++X)
+						{
+							int32 Y = 0;
+							const FVector Start(TileMap->GetTilePositionInLocalSpace(X, Y, SelectedLayerIndex));
 
-								Y = TileMap->MapHeight;
-								const FVector End(TileMap->GetTilePositionInLocalSpace(X, Y, SelectedLayerIndex));
+							Y = TileMap->MapHeight;
+							const FVector End(TileMap->GetTilePositionInLocalSpace(X, Y, SelectedLayerIndex));
 
-								PDI->DrawLine(LocalToWorld.TransformPosition(Start), LocalToWorld.TransformPosition(End), OverrideColor, DPG, 0.0f, DepthBias);
-							}
+							PDI->DrawLine(LocalToWorld.TransformPosition(Start), LocalToWorld.TransformPosition(End), OverrideColor, DPG, 0.0f, DepthBias);
 						}
 					}
-					else if (bShowOutlineWhenUnselected)
-					{
-						// Draw layer 0 even when not selected, so you can see where the tile map is in the editor
-						DrawBoundsForLayer(PDI, WireframeColor, 0);
-					}
+				}
+				else if (View->Family->EngineShowFlags.Grid && bShowOutlineWhenUnselected)
+				{
+					// Draw layer 0 even when not selected, so you can see where the tile map is in the editor
+					DrawBoundsForLayer(PDI, WireframeColor, /*LayerIndex=*/ 0);
 				}
 #endif
 			}
