Tile Maps: Redefined behavior when a tile set has a different tile height than the map.  It now puts the bottom left corner of the tile in the bottom left corner of the map cell, rather than the top left in the top left corner. This matches the behavior of Tiled, and makes working with variable-height isometric tiles feel better [UE-14441]

Fixed an issue with DrawingOffset when Pixels/uu is not 1.0 [UE-14443]

diff --git a/Engine/Plugins/2D/Paper2D/Source/Paper2D/Private/PaperTileMapComponent.cpp b/Engine/Plugins/2D/Paper2D/Source/Paper2D/Private/PaperTileMapComponent.cpp
--- a/Engine/Plugins/2D/Paper2D/Source/Paper2D/Private/PaperTileMapComponent.cpp
+++ b/Engine/Plugins/2D/Paper2D/Source/Paper2D/Private/PaperTileMapComponent.cpp
@@ -323,7 +323,10 @@ void UPaperTileMapComponent::RebuildRenderData(FPaperTileMapRenderSceneProxy* Pr
 						{
 							SourceDimensionsUV = FVector2D(TileInfo.TileSet->TileWidth * InverseTextureSize.X, TileInfo.TileSet->TileHeight * InverseTextureSize.Y);
 							TileSizeXY = FVector2D(UnrealUnitsPerPixel * TileInfo.TileSet->TileWidth, UnrealUnitsPerPixel * TileInfo.TileSet->TileHeight);
-							TileSetOffset = (TileInfo.TileSet->DrawingOffset.X * PaperAxisX) + (TileInfo.TileSet->DrawingOffset.Y * PaperAxisY);
+
+							const float HorizontalCellOffset = TileInfo.TileSet->DrawingOffset.X * UnrealUnitsPerPixel;
+							const float VerticalCellOffset = (TileInfo.TileSet->DrawingOffset.Y - TileHeight + TileInfo.TileSet->TileHeight) * UnrealUnitsPerPixel;
+							TileSetOffset = (HorizontalCellOffset * PaperAxisX) + (VerticalCellOffset * PaperAxisY);
 						}
 						else
 						{
