Tile Sets: Allow negative values for DrawingOffset [UE-14442]

diff --git a/Engine/Plugins/2D/Paper2D/Source/Paper2D/Classes/PaperTileSet.h b/Engine/Plugins/2D/Paper2D/Source/Paper2D/Classes/PaperTileSet.h
--- a/Engine/Plugins/2D/Paper2D/Source/Paper2D/Classes/PaperTileSet.h
+++ b/Engine/Plugins/2D/Paper2D/Source/Paper2D/Classes/PaperTileSet.h
@@ -93,7 +93,7 @@ class PAPER2D_API UPaperTileSet : public UObject
 	int32 Spacing;
 
 	// The drawing offset for tiles from this set (in pixels)
-	UPROPERTY(Category=TileSet, EditAnywhere, meta=(UIMin=0, ClampMin=0))
+	UPROPERTY(Category=TileSet, EditAnywhere)
 	FIntPoint DrawingOffset;
 
 #if WITH_EDITORONLY_DATA
