Tile Set Editor: Clarified the purpose of TileIndex by renaming it

diff --git a/Engine/Plugins/2D/Paper2D/Source/Paper2DEditor/Private/TileSetEditor/TileSetEditorViewportClient.cpp b/Engine/Plugins/2D/Paper2D/Source/Paper2DEditor/Private/TileSetEditor/TileSetEditorViewportClient.cpp
--- a/Engine/Plugins/2D/Paper2D/Source/Paper2DEditor/Private/TileSetEditor/TileSetEditorViewportClient.cpp
+++ b/Engine/Plugins/2D/Paper2D/Source/Paper2DEditor/Private/TileSetEditor/TileSetEditorViewportClient.cpp
@@ -13,7 +13,7 @@ FTileSetEditorViewportClient::FTileSetEditorViewportClient(UPaperTileSet* InTile
 	, bShowTilesWithCollision(false)
 	, bShowTilesWithMetaData(false)
 	, bHasValidPaintRectangle(false)
-	, TileIndex(INDEX_NONE)
+	, CurrentSelectedTileIndex(INDEX_NONE)
 {
 }
 
@@ -63,9 +63,9 @@ void FTileSetEditorViewportClient::Draw(FViewport* Viewport, FCanvas* Canvas)
 			const float Width = (TileSet->TileWidth - 2) * ZoomAmount;
 			const float Height = (TileSet->TileHeight - 2) * ZoomAmount;
 
-			for (int32 OtherTileIndex = 0; OtherTileIndex < NumTiles; ++OtherTileIndex)
+			for (int32 TileIndex = 0; TileIndex < NumTiles; ++TileIndex)
 			{
-				if (const FPaperTileMetadata* TileMetadata = TileSet->GetTileMetadata(OtherTileIndex))
+				if (const FPaperTileMetadata* TileMetadata = TileSet->GetTileMetadata(TileIndex))
 				{
 					const bool bShowDueToCollision = TileMetadata->HasCollision() && bShowTilesWithCollision;
 					const bool bShowDueToMetaData = TileMetadata->HasMetaData() && bShowTilesWithMetaData;
@@ -73,7 +73,7 @@ void FTileSetEditorViewportClient::Draw(FViewport* Viewport, FCanvas* Canvas)
 					if (bShowDueToCollision || bShowDueToMetaData)
 					{
 						FVector2D TileUV;
-						TileSet->GetTileUV(OtherTileIndex, /*out*/ TileUV);
+						TileSet->GetTileUV(TileIndex, /*out*/ TileUV);
 
 						const float XPos = (TileUV.X + 1 - ZoomPos.X) * ZoomAmount;
 						const float YPos = (TileUV.Y + 1 - ZoomPos.Y) * ZoomAmount;
@@ -102,9 +102,9 @@ void FTileSetEditorViewportClient::Draw(FViewport* Viewport, FCanvas* Canvas)
 		Canvas->DrawItem(BoxItem);
 	}
 
-	if (TileIndex != INDEX_NONE)
+	if (CurrentSelectedTileIndex != INDEX_NONE)
 	{
-		const FString TileIndexString = FString::Printf(TEXT("Tile# %d"), TileIndex);
+		const FString TileIndexString = FString::Printf(TEXT("Tile# %d"), CurrentSelectedTileIndex);
 
 		int32 XL;
 		int32 YL;
diff --git a/Engine/Plugins/2D/Paper2D/Source/Paper2DEditor/Private/TileSetEditor/TileSetEditorViewportClient.h b/Engine/Plugins/2D/Paper2D/Source/Paper2DEditor/Private/TileSetEditor/TileSetEditorViewportClient.h
--- a/Engine/Plugins/2D/Paper2D/Source/Paper2DEditor/Private/TileSetEditor/TileSetEditorViewportClient.h
+++ b/Engine/Plugins/2D/Paper2D/Source/Paper2DEditor/Private/TileSetEditor/TileSetEditorViewportClient.h
@@ -33,5 +33,5 @@ public:
 	bool bShowTilesWithMetaData;
 	bool bHasValidPaintRectangle;
 	FViewportSelectionRectangle ValidPaintRectangle;
-	int32 TileIndex;
+	int32 CurrentSelectedTileIndex;
 };
diff --git a/Engine/Plugins/2D/Paper2D/Source/Paper2DEditor/Private/TileSetEditor/TileSetSelectorViewport.cpp b/Engine/Plugins/2D/Paper2D/Source/Paper2DEditor/Private/TileSetEditor/TileSetSelectorViewport.cpp
--- a/Engine/Plugins/2D/Paper2D/Source/Paper2DEditor/Private/TileSetEditor/TileSetSelectorViewport.cpp
+++ b/Engine/Plugins/2D/Paper2D/Source/Paper2DEditor/Private/TileSetEditor/TileSetSelectorViewport.cpp
@@ -144,7 +144,7 @@ void STileSetSelectorViewport::RefreshSelectionRectangle()
 		Client->bHasValidPaintRectangle = bHasSelection && (TileSetBeingEdited != nullptr);
 
 		const int32 TileIndex = (bHasSelection && (TileSetBeingEdited != nullptr)) ? (SelectionTopLeft.X + SelectionTopLeft.Y * TileSetBeingEdited->GetTileCountX()) : INDEX_NONE;
-		Client->TileIndex = TileIndex;
+		Client->CurrentSelectedTileIndex = TileIndex;
 
 		if (bHasSelection && (TileSetBeingEdited != nullptr))
 		{
