Tile Maps: Fixed the pivot show flag not working in the stand-alone tile map editor [UE-14440]

diff --git a/Engine/Plugins/2D/Paper2D/Source/Paper2DEditor/Private/TileMapEditing/TileMapEditorViewportClient.cpp b/Engine/Plugins/2D/Paper2D/Source/Paper2DEditor/Private/TileMapEditing/TileMapEditorViewportClient.cpp
--- a/Engine/Plugins/2D/Paper2D/Source/Paper2DEditor/Private/TileMapEditing/TileMapEditorViewportClient.cpp
+++ b/Engine/Plugins/2D/Paper2D/Source/Paper2DEditor/Private/TileMapEditing/TileMapEditorViewportClient.cpp
@@ -40,7 +40,7 @@ FTileMapEditorViewportClient::FTileMapEditorViewportClient(TWeakPtr<FTileMapEdit
 
 	DrawHelper.bDrawGrid = GetDefault<UTileMapEditorSettings>()->bShowGridByDefault;
 	DrawHelper.bDrawPivot = false;
-	bShowPivot = true;
+	bShowPivot = false;
 
 	EngineShowFlags.DisableAdvancedFeatures();
 	EngineShowFlags.CompositeEditorPrimitives = true;
@@ -91,6 +91,32 @@ void FTileMapEditorViewportClient::DrawBoundsAsText(FViewport& InViewport, FScen
 	YPos += 18;
 }
 
+void FTileMapEditorViewportClient::Draw(const FSceneView* View, FPrimitiveDrawInterface* PDI)
+{
+	FPaperEditorViewportClient::Draw(View, PDI);
+
+	if (bShowPivot)
+	{
+		FColor PivotColor = FColor::Red;
+		float PivotSize = 0.02f;
+
+
+		//@TODO: Merge this with FEditorCommonDrawHelper::DrawPivot, which needs to take the pivot position as an argument
+		const FMatrix CameraToWorld = View->ViewMatrices.ViewMatrix.InverseFast();
+
+		const FVector PivLoc = FVector::ZeroVector;
+
+		const float ZoomFactor = FMath::Min<float>(View->ViewMatrices.ProjMatrix.M[0][0], View->ViewMatrices.ProjMatrix.M[1][1]);
+		const float WidgetRadius = View->ViewMatrices.GetViewProjMatrix().TransformPosition(PivLoc).W * (PivotSize / ZoomFactor);
+
+		const FVector CamX = CameraToWorld.TransformVector(FVector(1, 0, 0));
+		const FVector CamY = CameraToWorld.TransformVector(FVector(0, 1, 0));
+
+		PDI->DrawLine(PivLoc - (WidgetRadius*CamX), PivLoc + (WidgetRadius*CamX), PivotColor, SDPG_Foreground);
+		PDI->DrawLine(PivLoc - (WidgetRadius*CamY), PivLoc + (WidgetRadius*CamY), PivotColor, SDPG_Foreground);
+	}
+}
+
 void FTileMapEditorViewportClient::DrawCanvas(FViewport& Viewport, FSceneView& View, FCanvas& Canvas)
 {
 	FEditorViewportClient::DrawCanvas(Viewport, View, Canvas);
diff --git a/Engine/Plugins/2D/Paper2D/Source/Paper2DEditor/Private/TileMapEditing/TileMapEditorViewportClient.h b/Engine/Plugins/2D/Paper2D/Source/Paper2DEditor/Private/TileMapEditing/TileMapEditorViewportClient.h
--- a/Engine/Plugins/2D/Paper2D/Source/Paper2DEditor/Private/TileMapEditing/TileMapEditorViewportClient.h
+++ b/Engine/Plugins/2D/Paper2D/Source/Paper2DEditor/Private/TileMapEditing/TileMapEditorViewportClient.h
@@ -20,6 +20,7 @@ public:
 	FTileMapEditorViewportClient(TWeakPtr<FTileMapEditor> InTileMapEditor, TWeakPtr<class STileMapEditorViewport> InTileMapEditorViewportPtr);
 
 	// FViewportClient interface
+	virtual void Draw(const FSceneView* View, FPrimitiveDrawInterface* PDI) override;
 	virtual void DrawCanvas(FViewport& InViewport, FSceneView& View, FCanvas& Canvas) override;
 	virtual void Tick(float DeltaSeconds) override;
 	// End of FViewportClient interface
