Tile Maps: Improved rendering and interaction behavior for IsometricStaggered tile maps

diff --git a/Engine/Plugins/2D/Paper2D/Source/Paper2D/Private/PaperTileMap.cpp b/Engine/Plugins/2D/Paper2D/Source/Paper2D/Private/PaperTileMap.cpp
--- a/Engine/Plugins/2D/Paper2D/Source/Paper2D/Private/PaperTileMap.cpp
+++ b/Engine/Plugins/2D/Paper2D/Source/Paper2D/Private/PaperTileMap.cpp
@@ -203,7 +203,7 @@ void UPaperTileMap::GetTileToLocalParameters(FVector& OutCornerPosition, FVector
 		break;
 	case ETileMapProjectionMode::HexagonalStaggered:
 	case ETileMapProjectionMode::IsometricStaggered:
-		OutCornerPosition = -(TileWidthInUU * PaperAxisX * 0.5f) + (TileHeightInUU * PaperAxisY * 1.0f);
+		OutCornerPosition = -(TileWidthInUU * PaperAxisX * 0.5f) + (TileHeightInUU * PaperAxisY * 0.5f);
 		OutOffsetYFactor = 0.5f * TileWidthInUU * PaperAxisX;
 		OutStepX = PaperAxisX * TileWidthInUU;
 		OutStepY = 0.5f * -PaperAxisY * TileHeightInUU;
@@ -234,7 +234,7 @@ void UPaperTileMap::GetLocalToTileParameters(FVector& OutCornerPosition, FVector
 		break;
 	case ETileMapProjectionMode::HexagonalStaggered:
 	case ETileMapProjectionMode::IsometricStaggered:
-		OutCornerPosition = -(TileWidthInUU * PaperAxisX * 0.5f) + (TileHeightInUU * PaperAxisY * 1.0f);
+		OutCornerPosition = -(TileWidthInUU * PaperAxisX * 0.5f) + (TileHeightInUU * PaperAxisY * 0.5f);
 		OutOffsetYFactor = 0.5f * TileWidthInUU * PaperAxisX;
 		OutStepX = PaperAxisX / TileWidthInUU;
 		OutStepY = -PaperAxisY / TileHeightInUU;
@@ -258,8 +258,37 @@ void UPaperTileMap::GetTileCoordinatesFromLocalSpacePosition(const FVector& Posi
 	const float ProjectionSpaceXInTiles = FVector::DotProduct(RelativePosition, ParameterAxisX);
 	const float ProjectionSpaceYInTiles = FVector::DotProduct(RelativePosition, ParameterAxisY);
 
-	OutTileX = FMath::FloorToInt(ProjectionSpaceXInTiles);
-	OutTileY = FMath::FloorToInt(ProjectionSpaceYInTiles);
+	float X2 = ProjectionSpaceXInTiles;
+	float Y2 = ProjectionSpaceYInTiles;
+
+	if ((ProjectionMode == ETileMapProjectionMode::IsometricStaggered) || (ProjectionMode == ETileMapProjectionMode::HexagonalStaggered))
+	{
+		const float px = FMath::Frac(ProjectionSpaceXInTiles);
+		const float py = FMath::Frac(ProjectionSpaceYInTiles);
+
+		// Determine if the point is inside of the diamond or outside
+		const float h = 0.5f;
+		const float Det1 = -((px - h)*h - py*h);
+		const float Det2 = -(px - 1.0f)*h - (py - h)*h;
+		const float Det3 = -(-(px - h)*h + (py - 1.0f)*h);
+		const float Det4 = px*h + (py - h)*h;
+
+		const bool bOutsideTile = (Det1 < 0.0f) || (Det2 < 0.0f) || (Det3 < 0.0f) || (Det4 < 0.0f);
+
+		if (bOutsideTile)
+		{
+			X2 = ProjectionSpaceXInTiles - ((px < 0.5f) ? 1.0f : 0.0f);
+			Y2 = FMath::FloorToFloat(ProjectionSpaceYInTiles)*2.0f + py + ((py < 0.5f) ? -1.0f : 1.0f);
+		}
+ 		else
+ 		{
+ 			X2 = ProjectionSpaceXInTiles;
+			Y2 = FMath::FloorToFloat(ProjectionSpaceYInTiles)*2.0f + py;
+ 		}
+	}
+
+	OutTileX = FMath::FloorToInt(X2);
+	OutTileY = FMath::FloorToInt(Y2);
 }
 
 FVector UPaperTileMap::GetTilePositionInLocalSpace(float TileX, float TileY, int32 LayerIndex) const
diff --git a/Engine/Plugins/2D/Paper2D/Source/Paper2DEditor/Private/TileMapEditing/EdModeTileMap.cpp b/Engine/Plugins/2D/Paper2D/Source/Paper2DEditor/Private/TileMapEditing/EdModeTileMap.cpp
--- a/Engine/Plugins/2D/Paper2D/Source/Paper2DEditor/Private/TileMapEditing/EdModeTileMap.cpp
+++ b/Engine/Plugins/2D/Paper2D/Source/Paper2DEditor/Private/TileMapEditing/EdModeTileMap.cpp
@@ -448,15 +448,66 @@ void FEdModeTileMap::Render(const FSceneView* View, FViewport* Viewport, FPrimit
 				CursorRange = LastEyeDropperBounds;
 			}
 
-			const FVector TL(ComponentToWorld.TransformPosition(TileMap->GetTilePositionInLocalSpace(CursorRange.Min.X, CursorRange.Min.Y, LastCursorTileZ)));
-			const FVector TR(ComponentToWorld.TransformPosition(TileMap->GetTilePositionInLocalSpace(CursorRange.Max.X, CursorRange.Min.Y, LastCursorTileZ)));
-			const FVector BL(ComponentToWorld.TransformPosition(TileMap->GetTilePositionInLocalSpace(CursorRange.Min.X, CursorRange.Max.Y, LastCursorTileZ)));
-			const FVector BR(ComponentToWorld.TransformPosition(TileMap->GetTilePositionInLocalSpace(CursorRange.Max.X, CursorRange.Max.Y, LastCursorTileZ)));
-
-			PDI->DrawLine(TL, TR, CursorWireColor, SDPG_Foreground, 0.0f, DepthBias);
-			PDI->DrawLine(TR, BR, CursorWireColor, SDPG_Foreground, 0.0f, DepthBias);
-			PDI->DrawLine(BR, BL, CursorWireColor, SDPG_Foreground, 0.0f, DepthBias);
-			PDI->DrawLine(BL, TL, CursorWireColor, SDPG_Foreground, 0.0f, DepthBias);
+			switch (TileMap->ProjectionMode)
+			{
+			case ETileMapProjectionMode::Orthogonal:
+			case ETileMapProjectionMode::IsometricDiamond:
+			default:
+				{
+					const FVector TL(ComponentToWorld.TransformPosition(TileMap->GetTilePositionInLocalSpace(CursorRange.Min.X, CursorRange.Min.Y, LastCursorTileZ)));
+					const FVector TR(ComponentToWorld.TransformPosition(TileMap->GetTilePositionInLocalSpace(CursorRange.Max.X, CursorRange.Min.Y, LastCursorTileZ)));
+					const FVector BL(ComponentToWorld.TransformPosition(TileMap->GetTilePositionInLocalSpace(CursorRange.Min.X, CursorRange.Max.Y, LastCursorTileZ)));
+					const FVector BR(ComponentToWorld.TransformPosition(TileMap->GetTilePositionInLocalSpace(CursorRange.Max.X, CursorRange.Max.Y, LastCursorTileZ)));
+
+					PDI->DrawLine(TL, TR, CursorWireColor, SDPG_Foreground, 0.0f, DepthBias);
+					PDI->DrawLine(TR, BR, CursorWireColor, SDPG_Foreground, 0.0f, DepthBias);
+					PDI->DrawLine(BR, BL, CursorWireColor, SDPG_Foreground, 0.0f, DepthBias);
+					PDI->DrawLine(BL, TL, CursorWireColor, SDPG_Foreground, 0.0f, DepthBias);
+				}
+				break;
+
+			case ETileMapProjectionMode::IsometricStaggered:
+				{
+					const float UnrealUnitsPerPixel = TileMap->GetUnrealUnitsPerPixel();
+					const float TileWidthInUU = TileMap->TileWidth * UnrealUnitsPerPixel;
+					const float TileHeightInUU = TileMap->TileHeight * UnrealUnitsPerPixel;
+
+					const FVector RecenterOffset = PaperAxisX*TileWidthInUU*0.5f;
+					const FVector LSTM = TileMap->GetTilePositionInLocalSpace(CursorRange.Min.X, CursorRange.Min.Y, LastCursorTileZ) + RecenterOffset;
+
+					const FVector TL(ComponentToWorld.TransformPosition(LSTM));
+					const FVector TR(ComponentToWorld.TransformPosition(LSTM + PaperAxisX*TileWidthInUU*0.5f - PaperAxisY*TileHeightInUU*0.5f));
+					const FVector BL(ComponentToWorld.TransformPosition(LSTM - PaperAxisX*TileWidthInUU*0.5f - PaperAxisY*TileHeightInUU*0.5f));
+					const FVector BR(ComponentToWorld.TransformPosition(LSTM - PaperAxisY*TileHeightInUU*1.0f));
+
+					PDI->DrawLine(TL, TR, CursorWireColor, SDPG_Foreground, 0.0f, DepthBias);
+					PDI->DrawLine(TR, BR, CursorWireColor, SDPG_Foreground, 0.0f, DepthBias);
+					PDI->DrawLine(BR, BL, CursorWireColor, SDPG_Foreground, 0.0f, DepthBias);
+					PDI->DrawLine(BL, TL, CursorWireColor, SDPG_Foreground, 0.0f, DepthBias);
+				}
+				break;
+
+			case ETileMapProjectionMode::HexagonalStaggered:
+				{
+					const float UnrealUnitsPerPixel = TileMap->GetUnrealUnitsPerPixel();
+					const float TileWidthInUU = TileMap->TileWidth * UnrealUnitsPerPixel;
+					const float TileHeightInUU = TileMap->TileHeight * UnrealUnitsPerPixel;
+
+					const FVector RecenterOffset = PaperAxisX*TileWidthInUU*0.5f;
+					const FVector LSTM = TileMap->GetTilePositionInLocalSpace(CursorRange.Min.X, CursorRange.Min.Y, LastCursorTileZ) + RecenterOffset;
+
+					const FVector TL(ComponentToWorld.TransformPosition(LSTM));
+					const FVector TR(ComponentToWorld.TransformPosition(LSTM + PaperAxisX*TileWidthInUU*0.5f - PaperAxisY*TileHeightInUU*0.5f));
+					const FVector BL(ComponentToWorld.TransformPosition(LSTM - PaperAxisX*TileWidthInUU*0.5f - PaperAxisY*TileHeightInUU*0.5f));
+					const FVector BR(ComponentToWorld.TransformPosition(LSTM - PaperAxisY*TileHeightInUU*1.0f));
+
+					PDI->DrawLine(TL, TR, CursorWireColor, SDPG_Foreground, 0.0f, DepthBias);
+					PDI->DrawLine(TR, BR, CursorWireColor, SDPG_Foreground, 0.0f, DepthBias);
+					PDI->DrawLine(BR, BL, CursorWireColor, SDPG_Foreground, 0.0f, DepthBias);
+					PDI->DrawLine(BL, TL, CursorWireColor, SDPG_Foreground, 0.0f, DepthBias);
+				}
+				break;
+			};
 		}
 	}
 }
