Tile Maps: Fixed the flipped direction of DrawingOffset.Y [UE-14450]

diff --git a/Engine/Plugins/2D/Paper2D/Source/Paper2D/Private/PaperTileMapComponent.cpp b/Engine/Plugins/2D/Paper2D/Source/Paper2D/Private/PaperTileMapComponent.cpp
--- a/Engine/Plugins/2D/Paper2D/Source/Paper2D/Private/PaperTileMapComponent.cpp
+++ b/Engine/Plugins/2D/Paper2D/Source/Paper2D/Private/PaperTileMapComponent.cpp
@@ -325,7 +325,7 @@ void UPaperTileMapComponent::RebuildRenderData(FPaperTileMapRenderSceneProxy* Pr
 							TileSizeXY = FVector2D(UnrealUnitsPerPixel * TileInfo.TileSet->TileWidth, UnrealUnitsPerPixel * TileInfo.TileSet->TileHeight);
 
 							const float HorizontalCellOffset = TileInfo.TileSet->DrawingOffset.X * UnrealUnitsPerPixel;
-							const float VerticalCellOffset = (TileInfo.TileSet->DrawingOffset.Y - TileHeight + TileInfo.TileSet->TileHeight) * UnrealUnitsPerPixel;
+							const float VerticalCellOffset = (-TileInfo.TileSet->DrawingOffset.Y - TileHeight + TileInfo.TileSet->TileHeight) * UnrealUnitsPerPixel;
 							TileSetOffset = (HorizontalCellOffset * PaperAxisX) + (VerticalCellOffset * PaperAxisY);
 						}
 						else
