Tile Maps: Improved rendering and interaction behavior for HexagonalStaggered tile maps
- Added storage of the hex side length
- Added adjustment of the map tile height by the hex side length when importing from Tiled
- Added a hexagonal cursor (still draws a second incorrect cursor as the brush preview however)

diff --git a/Engine/Plugins/2D/Paper2D/Source/Paper2D/Classes/PaperTileMap.h b/Engine/Plugins/2D/Paper2D/Source/Paper2D/Classes/PaperTileMap.h
--- a/Engine/Plugins/2D/Paper2D/Source/Paper2D/Classes/PaperTileMap.h
+++ b/Engine/Plugins/2D/Paper2D/Source/Paper2D/Classes/PaperTileMap.h
@@ -49,6 +49,11 @@ class PAPER2D_API UPaperTileMap : public UObject
 	UPROPERTY(Category=Setup, EditAnywhere, BlueprintReadOnly, meta=(UIMin=1, ClampMin=1))
 	int32 TileHeight;
 
+	// The vertical height of the sides of the hex cell for a tile.
+	// Note: This value should already be included as part of the TileHeight, and is purely cosmetic; it only affects how the tile cursor preview is drawn.
+	UPROPERTY(Category=Setup, EditAnywhere, meta=(UIMin=0, ClampMin=0))
+	int32 HexSideLength;
+
 	// The scaling factor between pixels and Unreal units (cm) (e.g., 0.64 would make a 64 pixel wide tile take up 100 cm)
 	UPROPERTY(Category = Setup, EditAnywhere)
 	float PixelsPerUnrealUnit;
@@ -118,6 +123,7 @@ public:
 	// UObject interface
 #if WITH_EDITOR
 	virtual void PostEditChangeProperty(FPropertyChangedEvent& PropertyChangedEvent) override;
+	virtual bool CanEditChange(const UProperty* InProperty) const override;
 	virtual void PostLoad() override;
 	void ValidateSelectedLayerIndex();
 #endif
diff --git a/Engine/Plugins/2D/Paper2D/Source/Paper2D/Private/PaperTileMap.cpp b/Engine/Plugins/2D/Paper2D/Source/Paper2D/Private/PaperTileMap.cpp
--- a/Engine/Plugins/2D/Paper2D/Source/Paper2D/Private/PaperTileMap.cpp
+++ b/Engine/Plugins/2D/Paper2D/Source/Paper2D/Private/PaperTileMap.cpp
@@ -106,6 +106,22 @@ void UPaperTileMap::PostEditChangeProperty(FPropertyChangedEvent& PropertyChange
 	Super::PostEditChangeProperty(PropertyChangedEvent);
 }
 
+bool UPaperTileMap::CanEditChange(const UProperty* InProperty) const
+{
+	bool bIsEditable = Super::CanEditChange(InProperty);
+	if (bIsEditable && InProperty)
+	{
+		const FName PropertyName = InProperty->GetFName();
+
+		if (PropertyName == GET_MEMBER_NAME_CHECKED(UPaperTileMap, HexSideLength))
+		{
+			bIsEditable = ProjectionMode == ETileMapProjectionMode::HexagonalStaggered;
+		}
+	}
+
+	return bIsEditable;
+}
+
 void UPaperTileMap::PostLoad()
 {
 	Super::PostLoad();
diff --git a/Engine/Plugins/2D/Paper2D/Source/Paper2DEditor/Private/TileMapEditing/EdModeTileMap.cpp b/Engine/Plugins/2D/Paper2D/Source/Paper2DEditor/Private/TileMapEditing/EdModeTileMap.cpp
--- a/Engine/Plugins/2D/Paper2D/Source/Paper2DEditor/Private/TileMapEditing/EdModeTileMap.cpp
+++ b/Engine/Plugins/2D/Paper2D/Source/Paper2DEditor/Private/TileMapEditing/EdModeTileMap.cpp
@@ -496,15 +496,24 @@ void FEdModeTileMap::Render(const FSceneView* View, FViewport* Viewport, FPrimit
 					const FVector RecenterOffset = PaperAxisX*TileWidthInUU*0.5f;
 					const FVector LSTM = TileMap->GetTilePositionInLocalSpace(CursorRange.Min.X, CursorRange.Min.Y, LastCursorTileZ) + RecenterOffset;
 
-					const FVector TL(ComponentToWorld.TransformPosition(LSTM));
-					const FVector TR(ComponentToWorld.TransformPosition(LSTM + PaperAxisX*TileWidthInUU*0.5f - PaperAxisY*TileHeightInUU*0.5f));
-					const FVector BL(ComponentToWorld.TransformPosition(LSTM - PaperAxisX*TileWidthInUU*0.5f - PaperAxisY*TileHeightInUU*0.5f));
-					const FVector BR(ComponentToWorld.TransformPosition(LSTM - PaperAxisY*TileHeightInUU*1.0f));
+					const float HexSideLengthInUU = TileMap->HexSideLength * UnrealUnitsPerPixel;
 
-					PDI->DrawLine(TL, TR, CursorWireColor, SDPG_Foreground, 0.0f, DepthBias);
-					PDI->DrawLine(TR, BR, CursorWireColor, SDPG_Foreground, 0.0f, DepthBias);
-					PDI->DrawLine(BR, BL, CursorWireColor, SDPG_Foreground, 0.0f, DepthBias);
-					PDI->DrawLine(BL, TL, CursorWireColor, SDPG_Foreground, 0.0f, DepthBias);
+					const FVector Top(ComponentToWorld.TransformPosition(LSTM - PaperAxisY*HexSideLengthInUU));
+
+					const FVector RightTop(ComponentToWorld.TransformPosition(LSTM + PaperAxisX*TileWidthInUU*0.5f - PaperAxisY*TileHeightInUU*0.5f));
+					const FVector LeftTop(ComponentToWorld.TransformPosition(LSTM - PaperAxisX*TileWidthInUU*0.5f - PaperAxisY*TileHeightInUU*0.5f));
+
+					const FVector RightBottom(ComponentToWorld.TransformPosition(LSTM + PaperAxisX*TileWidthInUU*0.5f - PaperAxisY*(TileHeightInUU*0.5f + HexSideLengthInUU)));
+					const FVector LeftBottom(ComponentToWorld.TransformPosition(LSTM - PaperAxisX*TileWidthInUU*0.5f - PaperAxisY*(TileHeightInUU*0.5f + HexSideLengthInUU)));
+
+					const FVector Bottom(ComponentToWorld.TransformPosition(LSTM - PaperAxisY*TileHeightInUU*1.0f));
+
+					PDI->DrawLine(Top, RightTop, CursorWireColor, SDPG_Foreground, 0.0f, DepthBias);
+					PDI->DrawLine(RightTop, RightBottom, CursorWireColor, SDPG_Foreground, 0.0f, DepthBias);
+					PDI->DrawLine(RightBottom, Bottom, CursorWireColor, SDPG_Foreground, 0.0f, DepthBias);
+					PDI->DrawLine(Bottom, LeftBottom, CursorWireColor, SDPG_Foreground, 0.0f, DepthBias);
+					PDI->DrawLine(LeftBottom, LeftTop, CursorWireColor, SDPG_Foreground, 0.0f, DepthBias);
+					PDI->DrawLine(LeftTop, Top, CursorWireColor, SDPG_Foreground, 0.0f, DepthBias);
 				}
 				break;
 			};
diff --git a/Engine/Plugins/2D/Paper2D/Source/PaperTiledImporter/Private/PaperTiledImporterFactory.cpp b/Engine/Plugins/2D/Paper2D/Source/PaperTiledImporter/Private/PaperTiledImporterFactory.cpp
--- a/Engine/Plugins/2D/Paper2D/Source/PaperTiledImporter/Private/PaperTiledImporterFactory.cpp
+++ b/Engine/Plugins/2D/Paper2D/Source/PaperTiledImporter/Private/PaperTiledImporterFactory.cpp
@@ -211,6 +211,12 @@ UObject* UPaperTiledImporterFactory::FactoryCreateText(UClass* InClass, UObject*
 		Result->ProjectionMode = GlobalInfo.GetOrientationType();
 		Result->PixelsPerUnrealUnit = GetDefault<UPaperRuntimeSettings>()->DefaultPixelsPerUnrealUnit;
 		Result->BackgroundColor = GlobalInfo.BackgroundColor;
+		Result->HexSideLength = GlobalInfo.HexSideLength;
+
+		if (GlobalInfo.Orientation == ETiledOrientation::Hexagonal)
+		{
+			Result->TileHeight += GlobalInfo.HexSideLength;
+		}
 
 		// Create the tile sets
 		for (const FTileSetFromTiled& TileSetData : GlobalInfo.TileSets)
