Tile Set Editor: Added a customization to narrow down the metadata display to only it for the tile currently being edited
- Added a UserDataKey metadata property and exposed a filter for set tags (along the same lines as the collision filter)
- Exposed some properties of UTileSet as readable from Blueprints
- Removed the Experimental tag on the per-tile metadata
[UE-13623]

diff --git a/Engine/Config/BaseEngine.ini b/Engine/Config/BaseEngine.ini
--- a/Engine/Config/BaseEngine.ini
+++ b/Engine/Config/BaseEngine.ini
@@ -614,6 +614,7 @@ bUseStreamingPause=false
 +ActiveClassRedirects=(OldClassName="SpritePolygon",NewClassName="/Script/Paper2D.SpriteGeometryShape")
 +ActiveClassRedirects=(OldClassName="SpritePolygonCollection",NewClassName="/Script/Paper2D.SpriteGeometryCollection")
 +TaggedPropertyRedirects=(ClassName="SpriteGeometryCollection",OldPropertyName="Polygons",NewPropertyName="Shapes")
++TaggedPropertyRedirects=(ClassName="PaperTileSet",OldPropertyName="ExperimentalPerTileData",NewPropertyName="PerTileData")
 
 ; collision response container refactor
 +TaggedPropertyRedirects=(ClassName="CollisionResponseContainer",OldPropertyName="RigidBody",NewPropertyName="PhysicsBody")
diff --git a/Engine/Plugins/2D/Paper2D/Source/Paper2D/Classes/PaperTileSet.h b/Engine/Plugins/2D/Paper2D/Source/Paper2D/Classes/PaperTileSet.h
--- a/Engine/Plugins/2D/Paper2D/Source/Paper2D/Classes/PaperTileSet.h
+++ b/Engine/Plugins/2D/Paper2D/Source/Paper2D/Classes/PaperTileSet.h
@@ -17,6 +17,11 @@ struct PAPER2D_API FPaperTileMetadata
 public:
 	GENERATED_USTRUCT_BODY()
 
+	// A tag that can be used for grouping and categorizing (consider using it as the index into a UDataTable asset).
+	UPROPERTY(EditAnywhere, BlueprintReadOnly, Category=Sprite)
+	FName UserDataName;
+
+	// Collision data for the tile
 	UPROPERTY(EditAnywhere, Category=Sprite)
 	FSpriteGeometryCollection CollisionData;
 
@@ -44,7 +49,7 @@ public:
 	// Does this tile have user-specified metadata?
 	bool HasMetaData() const
 	{
-		return false;
+		return !UserDataName.IsNone();
 	}
 };
 
@@ -73,27 +78,27 @@ class PAPER2D_API UPaperTileSet : public UObject
 	GENERATED_UCLASS_BODY()
 
 	// The width of a single tile (in pixels)
-	UPROPERTY(Category=TileSet, EditAnywhere, meta=(UIMin=1, ClampMin=1))
+	UPROPERTY(Category=TileSet, BlueprintReadOnly, EditAnywhere, meta=(UIMin=1, ClampMin=1))
 	int32 TileWidth;
 
 	// The height of a single tile (in pixels)
-	UPROPERTY(Category=TileSet, EditAnywhere, meta=(UIMin=1, ClampMin=1))
+	UPROPERTY(Category=TileSet, BlueprintReadOnly, EditAnywhere, meta=(UIMin=1, ClampMin=1))
 	int32 TileHeight;
 
 	// The tile sheet texture associated with this tile set
-	UPROPERTY(Category=TileSet, EditAnywhere)
+	UPROPERTY(Category=TileSet, BlueprintReadOnly, EditAnywhere)
 	UTexture2D* TileSheet;
 
 	// The amount of padding around the perimeter of the tile sheet (in pixels)
-	UPROPERTY(Category=TileSet, EditAnywhere, meta=(UIMin=0, ClampMin=0))
+	UPROPERTY(Category=TileSet, BlueprintReadOnly, EditAnywhere, meta=(UIMin=0, ClampMin=0))
 	int32 Margin;
 
 	// The amount of padding between tiles in the tile sheet (in pixels)
-	UPROPERTY(Category=TileSet, EditAnywhere, meta=(UIMin=0, ClampMin=0))
+	UPROPERTY(Category=TileSet, BlueprintReadOnly, EditAnywhere, meta=(UIMin=0, ClampMin=0))
 	int32 Spacing;
 
 	// The drawing offset for tiles from this set (in pixels)
-	UPROPERTY(Category=TileSet, EditAnywhere)
+	UPROPERTY(Category=TileSet, BlueprintReadOnly, EditAnywhere)
 	FIntPoint DrawingOffset;
 
 #if WITH_EDITORONLY_DATA
@@ -120,9 +125,9 @@ protected:
 	UPROPERTY()
 	int32 AllocatedHeight;
 
-	// Per-tile information (Experimental)
+	// Per-tile information
 	UPROPERTY(EditAnywhere, EditFixedSize, Category=Sprite)
-	TArray<FPaperTileMetadata> ExperimentalPerTileData;
+	TArray<FPaperTileMetadata> PerTileData;
 
 	// Terrain information
 	UPROPERTY()//@TODO: TileMapTerrains: (EditAnywhere, Category=Sprite)
@@ -130,14 +135,14 @@ protected:
 
 protected:
 
-	// Reallocates the ExperimentalPerTileData array to the specified size.
+	// Reallocates the PerTileData array to the specified size.
 	// Note: This is a destructive operation!
 	void DestructiveAllocateTileData(int32 NewWidth, int32 NewHeight);
 
 	// Reallocates the per-tile data to match the current (WidthInTiles, HeightInTiles) size, copying over what it can
 	void ReallocateAndCopyTileData();
 
-public:	
+public:
 
 	// UObject interface
 #if WITH_EDITOR
@@ -184,4 +189,8 @@ public:
 
 	// Returns the terrain type this tile is a member of, or INDEX_NONE if it is not part of a terrain
 	int32 GetTerrainMembership(const FPaperTileInfo& TileInfo) const;
+
+#if WITH_EDITOR
+	static FName GetPerTilePropertyName() { return GET_MEMBER_NAME_CHECKED(UPaperTileSet, PerTileData); }
+#endif
 };
diff --git a/Engine/Plugins/2D/Paper2D/Source/Paper2D/Private/PaperTileSet.cpp b/Engine/Plugins/2D/Paper2D/Source/Paper2D/Private/PaperTileSet.cpp
--- a/Engine/Plugins/2D/Paper2D/Source/Paper2D/Private/PaperTileSet.cpp
+++ b/Engine/Plugins/2D/Paper2D/Source/Paper2D/Private/PaperTileSet.cpp
@@ -69,9 +69,9 @@ int32 UPaperTileSet::GetTileCountY() const
 
 FPaperTileMetadata* UPaperTileSet::GetMutableTileMetadata(int32 TileIndex)
 {
-	if (ExperimentalPerTileData.IsValidIndex(TileIndex))
+	if (PerTileData.IsValidIndex(TileIndex))
 	{
-		return &(ExperimentalPerTileData[TileIndex]);
+		return &(PerTileData[TileIndex]);
 	}
 	else
 	{
@@ -81,9 +81,9 @@ FPaperTileMetadata* UPaperTileSet::GetMutableTileMetadata(int32 TileIndex)
 
 const FPaperTileMetadata* UPaperTileSet::GetTileMetadata(int32 TileIndex) const
 {
-	if (ExperimentalPerTileData.IsValidIndex(TileIndex))
+	if (PerTileData.IsValidIndex(TileIndex))
 	{
-		return &(ExperimentalPerTileData[TileIndex]);
+		return &(PerTileData[TileIndex]);
 	}
 	else
 	{
@@ -191,10 +191,10 @@ void UPaperTileSet::PostEditChangeProperty(FPropertyChangedEvent& PropertyChange
 void UPaperTileSet::DestructiveAllocateTileData(int32 NewWidth, int32 NewHeight)
 {
 	const int32 NumCells = NewWidth * NewHeight;
-	ExperimentalPerTileData.Empty(NumCells);
+	PerTileData.Empty(NumCells);
 	for (int32 Index = 0; Index < NumCells; ++Index)
 	{
-		ExperimentalPerTileData.Add(FPaperTileMetadata());
+		PerTileData.Add(FPaperTileMetadata());
 	}
 
 	AllocatedWidth = NewWidth;
@@ -210,7 +210,7 @@ void UPaperTileSet::ReallocateAndCopyTileData()
 
 	const int32 SavedWidth = AllocatedWidth;
 	const int32 SavedHeight = AllocatedHeight;
-	TArray<FPaperTileMetadata> SavedDesignedMap(ExperimentalPerTileData);
+	TArray<FPaperTileMetadata> SavedDesignedMap(PerTileData);
 
 	DestructiveAllocateTileData(WidthInTiles, HeightInTiles);
 
@@ -224,8 +224,7 @@ void UPaperTileSet::ReallocateAndCopyTileData()
 			const int32 SrcIndex = Y*SavedWidth + X;
 			const int32 DstIndex = Y*WidthInTiles + X;
 
-			ExperimentalPerTileData[DstIndex] = SavedDesignedMap[SrcIndex];
+			PerTileData[DstIndex] = SavedDesignedMap[SrcIndex];
 		}
 	}
-
 }
diff --git a/Engine/Plugins/2D/Paper2D/Source/Paper2DEditor/Private/TileSetEditor.cpp b/Engine/Plugins/2D/Paper2D/Source/Paper2DEditor/Private/TileSetEditor.cpp
--- a/Engine/Plugins/2D/Paper2D/Source/Paper2DEditor/Private/TileSetEditor.cpp
+++ b/Engine/Plugins/2D/Paper2D/Source/Paper2DEditor/Private/TileSetEditor.cpp
@@ -18,6 +18,8 @@
 #include "TileSetEditor/SingleTileEditorViewportClient.h"
 #include "TileSetEditor/TileSetEditorCommands.h"
 #include "SpriteEditor/SpriteEditorCommands.h"
+#include "TileSetEditor/TileSetDetailsCustomization.h"
+#include "Editor/PropertyEditor/Public/IDetailsView.h"
 
 #define LOCTEXT_NAMESPACE "TileSetEditor"
 
@@ -50,6 +52,26 @@ public:
 	SLATE_BEGIN_ARGS(STileSetPropertiesTabBody) {}
 	SLATE_END_ARGS()
 
+private:
+	TSharedRef<IDetailCustomization> MakeEmbeddedInstance()
+	{
+		TSharedRef<FTileSetDetailsCustomization> Customization = FTileSetDetailsCustomization::MakeEmbeddedInstance();
+
+		// Make sure the customization starts off looking at the right tile index
+		TSharedPtr<FSingleTileEditorViewportClient> SingleTileEditor = TileSetEditorPtr.Pin()->GetSingleTileEditor();
+
+		const int32 TileIndex = SingleTileEditor->GetTileIndex();
+		if (TileIndex != INDEX_NONE)
+		{
+			Customization->OnTileIndexChanged(TileIndex, INDEX_NONE);
+		}
+
+		// And that it gets updated when the tile index changes
+		SingleTileEditor->GetOnSingleTileIndexChanged().AddSP(Customization, &FTileSetDetailsCustomization::OnTileIndexChanged);
+
+		return Customization;
+	}
+
 private:
 	// Pointer back to owning TileSet editor instance (the keeper of state)
 	TWeakPtr<class FTileSetEditor> TileSetEditorPtr;
@@ -59,6 +81,9 @@ public:
 		TileSetEditorPtr = InTileSetEditor;
 
 		SSingleObjectDetailsPanel::Construct(SSingleObjectDetailsPanel::FArguments().HostCommandList(InTileSetEditor->GetToolkitCommands()));
+
+		FOnGetDetailCustomizationInstance LayoutTileSetDetails = FOnGetDetailCustomizationInstance::CreateSP(this, &STileSetPropertiesTabBody::MakeEmbeddedInstance);
+		PropertyView->RegisterInstancedCustomPropertyLayout(UPaperTileSet::StaticClass(), LayoutTileSetDetails);
 	}
 
 	// SSingleObjectDetailsPanel interface
@@ -136,7 +161,7 @@ void FTileSetEditor::InitTileSetEditor(const EToolkitMode::Type Mode, const TSha
 
 	TileSetViewport = SNew(STileSetSelectorViewport, InitTileSet, /*EdMode=*/ nullptr);
 	TileEditorViewportClient = MakeShareable(new FSingleTileEditorViewportClient(InitTileSet));
-	TileSetViewport->GetTileSelectionChanged().AddRaw(TileEditorViewportClient.Get(), &FSingleTileEditorViewportClient::OnActiveTileIndexChanged);
+	TileSetViewport->GetTileSelectionChanged().AddRaw(TileEditorViewportClient.Get(), &FSingleTileEditorViewportClient::OnTileSelectionRegionChanged);
 
 	TileEditorViewport = SNew(SSingleTileEditorViewport, TileEditorViewportClient);
 
@@ -252,7 +277,7 @@ void FTileSetEditor::ExtendToolbar()
 			ToolbarBuilder.BeginSection("TileHighlights");
 			{
 				ToolbarBuilder.AddToolBarButton(FTileSetEditorCommands::Get().SetShowTilesWithCollision);
-				//ToolbarBuilder.AddToolBarButton(FTileSetEditorCommands::Get().SetShowTilesWithMetaData); //@TODO: TileMetadata: Enable once there is actually metadata...
+				ToolbarBuilder.AddToolBarButton(FTileSetEditorCommands::Get().SetShowTilesWithMetaData);
 			}
 			ToolbarBuilder.EndSection();
 
diff --git a/Engine/Plugins/2D/Paper2D/Source/Paper2DEditor/Private/TileSetEditor.h b/Engine/Plugins/2D/Paper2D/Source/Paper2DEditor/Private/TileSetEditor.h
--- a/Engine/Plugins/2D/Paper2D/Source/Paper2DEditor/Private/TileSetEditor.h
+++ b/Engine/Plugins/2D/Paper2D/Source/Paper2DEditor/Private/TileSetEditor.h
@@ -6,6 +6,10 @@
 #include "Toolkits/AssetEditorManager.h"
 #include "SPaperEditorViewport.h"
 
+class STileSetSelectorViewport;
+class SSingleTileEditorViewport;
+class FSingleTileEditorViewportClient;
+
 //////////////////////////////////////////////////////////////////////////
 // FTileSetEditor
 
@@ -40,6 +44,8 @@ public:
 
 	UPaperTileSet* GetTileSetBeingEdited() const { return TileSetBeingEdited; }
 
+	TSharedPtr<FSingleTileEditorViewportClient> GetSingleTileEditor() const { return TileEditorViewportClient; }
+
 protected:
 	TSharedRef<class SDockTab> SpawnTab_TextureView(const FSpawnTabArgs& Args);
 	TSharedRef<class SDockTab> SpawnTab_Details(const FSpawnTabArgs& Args);
@@ -54,9 +60,9 @@ protected:
 protected:
 	UPaperTileSet* TileSetBeingEdited;
 
-	TSharedPtr<class STileSetSelectorViewport> TileSetViewport;
-	TSharedPtr<class SSingleTileEditorViewport> TileEditorViewport;
-	TSharedPtr<class FSingleTileEditorViewportClient> TileEditorViewportClient;
+	TSharedPtr<STileSetSelectorViewport> TileSetViewport;
+	TSharedPtr<SSingleTileEditorViewport> TileEditorViewport;
+	TSharedPtr<FSingleTileEditorViewportClient> TileEditorViewportClient;
 
 	FDelegateHandle OnPropertyChangedHandle;
 };
diff --git a/Engine/Plugins/2D/Paper2D/Source/Paper2DEditor/Private/TileSetEditor/SingleTileEditorViewportClient.cpp b/Engine/Plugins/2D/Paper2D/Source/Paper2DEditor/Private/TileSetEditor/SingleTileEditorViewportClient.cpp
--- a/Engine/Plugins/2D/Paper2D/Source/Paper2DEditor/Private/TileSetEditor/SingleTileEditorViewportClient.cpp
+++ b/Engine/Plugins/2D/Paper2D/Source/Paper2DEditor/Private/TileSetEditor/SingleTileEditorViewportClient.cpp
@@ -259,11 +259,14 @@ void FSingleTileEditorViewportClient::SetTileIndex(int32 InTileIndex)
 		RequestFocusOnSelection(/*bInstant=*/ true);
 	}
 
+	// Trigger a details panel customization rebuild
+	OnSingleTileIndexChanged.Broadcast(TileBeingEditedIndex, OldTileIndex);
+
 	// Redraw the viewport
 	Invalidate();
 }
 
-void FSingleTileEditorViewportClient::OnActiveTileIndexChanged(const FIntPoint& TopLeft, const FIntPoint& Dimensions)
+void FSingleTileEditorViewportClient::OnTileSelectionRegionChanged(const FIntPoint& TopLeft, const FIntPoint& Dimensions)
 {
 	const int32 TopLeftIndex = (TileSet->GetTileCountX() * TopLeft.Y) + TopLeft.X;
 	const int32 SelectionArea = Dimensions.X * Dimensions.Y;
diff --git a/Engine/Plugins/2D/Paper2D/Source/Paper2DEditor/Private/TileSetEditor/SingleTileEditorViewportClient.h b/Engine/Plugins/2D/Paper2D/Source/Paper2DEditor/Private/TileSetEditor/SingleTileEditorViewportClient.h
--- a/Engine/Plugins/2D/Paper2D/Source/Paper2DEditor/Private/TileSetEditor/SingleTileEditorViewportClient.h
+++ b/Engine/Plugins/2D/Paper2D/Source/Paper2DEditor/Private/TileSetEditor/SingleTileEditorViewportClient.h
@@ -6,6 +6,8 @@
 #include "PaperEditorViewportClient.h"
 #include "SpriteEditor/SpriteEditorSelections.h"
 
+DECLARE_MULTICAST_DELEGATE_TwoParams(FOnSingleTileIndexChanged, int32 /*NewIndex*/, int32 /*OldIndex*/);
+
 //////////////////////////////////////////////////////////////////////////
 // FSingleTileEditorViewportClient
 
@@ -38,7 +40,7 @@ public:
 
 	void SetTileIndex(int32 InTileIndex);
 	int32 GetTileIndex() const;
-	void OnActiveTileIndexChanged(const FIntPoint& TopLeft, const FIntPoint& Dimensions);
+	void OnTileSelectionRegionChanged(const FIntPoint& TopLeft, const FIntPoint& Dimensions);
 
 	void ActivateEditMode(TSharedPtr<FUICommandList> InCommandList);
 
@@ -47,15 +49,22 @@ public:
 	void ToggleShowStats();
 	bool IsShowStatsChecked() const;
 
+	// Delegate for when the index of the single tile being edited changes
+	FOnSingleTileIndexChanged& GetOnSingleTileIndexChanged()
+	{
+		return OnSingleTileIndexChanged;
+	}
+
 protected:
 	// FPaperEditorViewportClient interface
 	virtual FBox GetDesiredFocusBounds() const override;
 	// End of FPaperEditorViewportClient
 
-public:
+private:
 	// Tile set
 	UPaperTileSet* TileSet;
 
+	// The current tile being edited or INDEX_NONE
 	int32 TileBeingEditedIndex;
 
 	// Are we currently manipulating something?
@@ -78,4 +87,7 @@ public:
 
 	// The preview sprite in the scene
 	UPaperSpriteComponent* PreviewTileSpriteComponent;
+
+	// Called when TileBeingEditedIndex changes
+	FOnSingleTileIndexChanged OnSingleTileIndexChanged;
 };
diff --git a/Engine/Plugins/2D/Paper2D/Source/Paper2DEditor/Private/TileSetEditor/TileSetDetailsCustomization.cpp b/Engine/Plugins/2D/Paper2D/Source/Paper2DEditor/Private/TileSetEditor/TileSetDetailsCustomization.cpp
--- a/Engine/Plugins/2D/Paper2D/Source/Paper2DEditor/Private/TileSetEditor/TileSetDetailsCustomization.cpp
+++ b/Engine/Plugins/2D/Paper2D/Source/Paper2DEditor/Private/TileSetEditor/TileSetDetailsCustomization.cpp
@@ -13,9 +13,21 @@
 //////////////////////////////////////////////////////////////////////////
 // FTileSetDetailsCustomization
 
+FTileSetDetailsCustomization::FTileSetDetailsCustomization(bool bInIsEmbedded)
+	: bIsEmbeddedInTileSetEditor(bInIsEmbedded)
+	, SelectedSingleTileIndex(INDEX_NONE)
+	, MyDetailLayout(nullptr)
+{
+}
+
 TSharedRef<IDetailCustomization> FTileSetDetailsCustomization::MakeInstance()
 {
-	return MakeShareable(new FTileSetDetailsCustomization);
+	return MakeShareable(new FTileSetDetailsCustomization(/*bIsEmbedded=*/ false));
+}
+
+TSharedRef<FTileSetDetailsCustomization> FTileSetDetailsCustomization::MakeEmbeddedInstance()
+{
+	return MakeShareable(new FTileSetDetailsCustomization(/*bIsEmbedded=*/ true));
 }
 
 void FTileSetDetailsCustomization::CustomizeDetails(IDetailLayoutBuilder& DetailLayout)
@@ -54,28 +66,51 @@ void FTileSetDetailsCustomization::CustomizeDetails(IDetailLayoutBuilder& Detail
 	);
 
 
-	// Add all of the properties from the inline tilemap
-// 	if ((TileComponent != nullptr) && (TileComponent->OwnsTileMap()))
-// 	{
-// 		TArray<UObject*> ListOfTileMaps;
-// 		ListOfTileMaps.Add(TileMap);
-// 
-// 		for (TFieldIterator<UProperty> PropIt(UPaperTileMap::StaticClass()); PropIt; ++PropIt)
-// 		{
-// 			UProperty* TestProperty = *PropIt;
-// 
-// 			if (TestProperty->HasAnyPropertyFlags(CPF_Edit))
-// 			{
-// 				FName CategoryName(*TestProperty->GetMetaData(TEXT("Category")));
-// 				IDetailCategoryBuilder& Category = DetailLayout.EditCategory(CategoryName);
-// 
-// 				if (IDetailPropertyRow* ExternalRow = Category.AddExternalProperty(ListOfTileMaps, TestProperty->GetFName()))
-// 				{
-// 					ExternalRow->Visibility(InternalInstanceVis);
-// 				}
-// 			}
-// 		}
-// 	}
+	if (bIsEmbeddedInTileSetEditor)
+	{
+		// Hide the array to start with
+		const FName MetadataArrayName = UPaperTileSet::GetPerTilePropertyName();
+		DetailLayout.HideProperty(MetadataArrayName);
+
+		if (SelectedSingleTileIndex != INDEX_NONE)
+		{
+			// Customize for the single tile being edited
+			IDetailCategoryBuilder& SingleTileCategory = DetailLayout.EditCategory("SingleTileEditor", FText::GetEmpty());
+
+			const FString ArrayEntryPathPrefix = FString::Printf(TEXT("%s[%d]."), *MetadataArrayName.ToString(), SelectedSingleTileIndex);
+
+			// Add all of the editable properties in the array entry
+			for (TFieldIterator<UProperty> PropIt(FPaperTileMetadata::StaticStruct()); PropIt; ++PropIt)
+			{
+				UProperty* TestProperty = *PropIt;
+				if (TestProperty->HasAnyPropertyFlags(CPF_Edit))
+				{
+					const FString PropertyPath = ArrayEntryPathPrefix + TestProperty->GetName();
+					TSharedRef<IPropertyHandle> PropertyHandle = DetailLayout.GetProperty(*PropertyPath);
+					SingleTileCategory.AddProperty(PropertyHandle);
+				}
+			}
+
+			// Add a display of the tile index being edited to the header
+			const FText TileIndexHeaderText = FText::Format(LOCTEXT("SingleTileSectionHeader", "Editing Tile #{0}"), FText::AsNumber(SelectedSingleTileIndex));
+			SingleTileCategory.HeaderContent
+			(
+				SNew(SBox)
+				.HAlign(HAlign_Right)
+				[
+					SNew(SHorizontalBox)
+					+SHorizontalBox::Slot()
+					.Padding(FMargin(5.0f, 0.0f))
+					.AutoWidth()
+					[
+						SNew(STextBlock)
+						.Font(FEditorStyle::GetFontStyle("TinyText"))
+						.Text(TileIndexHeaderText)
+					]
+				]
+			);
+		}
+	}
 }
 
 FText FTileSetDetailsCustomization::GetCellDimensionHeaderText() const
@@ -123,6 +158,15 @@ FSlateColor FTileSetDetailsCustomization::GetCellDimensionHeaderColor() const
 	return FSlateColor::UseForeground();
 }
 
+void FTileSetDetailsCustomization::OnTileIndexChanged(int32 NewIndex, int32 OldIndex)
+{
+	SelectedSingleTileIndex = NewIndex;
+	if (MyDetailLayout != nullptr)
+	{
+		MyDetailLayout->ForceRefreshDetails();
+	}
+}
+
 //////////////////////////////////////////////////////////////////////////
 
 #undef LOCTEXT_NAMESPACE
\ No newline at end of file
diff --git a/Engine/Plugins/2D/Paper2D/Source/Paper2DEditor/Private/TileSetEditor/TileSetDetailsCustomization.h b/Engine/Plugins/2D/Paper2D/Source/Paper2DEditor/Private/TileSetEditor/TileSetDetailsCustomization.h
--- a/Engine/Plugins/2D/Paper2D/Source/Paper2DEditor/Private/TileSetEditor/TileSetDetailsCustomization.h
+++ b/Engine/Plugins/2D/Paper2D/Source/Paper2DEditor/Private/TileSetEditor/TileSetDetailsCustomization.h
@@ -13,16 +13,32 @@ public:
 	// Makes a new instance of this detail layout class for a specific detail view requesting it
 	static TSharedRef<IDetailCustomization> MakeInstance();
 
+	// Makes a new instance of this detail layout class for use inside of the tile set editor (with a route to the tile being edited)
+	static TSharedRef<FTileSetDetailsCustomization> MakeEmbeddedInstance();
+
 	// IDetailCustomization interface
 	virtual void CustomizeDetails(IDetailLayoutBuilder& DetailLayout) override;
 	// End of IDetailCustomization interface
 
+	// Called when the tile being edited index changes
+	void OnTileIndexChanged(int32 NewIndex, int32 OldIndex);
+
 private:
+	// Are we embedded in the tile set editor?
+	bool bIsEmbeddedInTileSetEditor;
+
+	// Current index being viewed
+	int32 SelectedSingleTileIndex;
+
+	// The tile set being edited
 	TWeakObjectPtr<class UPaperTileSet> TileSetPtr;
 
+	// The detail layout builder that is using us
 	IDetailLayoutBuilder* MyDetailLayout;
 
 private:
+	FTileSetDetailsCustomization(bool bInIsEmbedded);
+
 	FText GetCellDimensionHeaderText() const;
 	FSlateColor GetCellDimensionHeaderColor() const;
 };
