Tile Maps: Refreshed the details customization when promoting an instance to an asset or vis versa [UE-14445]

diff --git a/Engine/Plugins/2D/Paper2D/Source/Paper2DEditor/Private/TileMapEditing/PaperTileMapDetailsCustomization.cpp b/Engine/Plugins/2D/Paper2D/Source/Paper2DEditor/Private/TileMapEditing/PaperTileMapDetailsCustomization.cpp
--- a/Engine/Plugins/2D/Paper2D/Source/Paper2DEditor/Private/TileMapEditing/PaperTileMapDetailsCustomization.cpp
+++ b/Engine/Plugins/2D/Paper2D/Source/Paper2DEditor/Private/TileMapEditing/PaperTileMapDetailsCustomization.cpp
@@ -294,6 +294,8 @@ FReply FPaperTileMapDetailsCustomization::OnPromoteToAssetButtonClicked()
 		}
 	}
 
+	MyDetailLayout->ForceRefreshDetails();
+
 	return FReply::Handled();
 }
 
@@ -310,6 +312,8 @@ FReply FPaperTileMapDetailsCustomization::OnMakeInstanceFromAssetButtonClicked()
 		}
 	}
 
+	MyDetailLayout->ForceRefreshDetails();
+
 	return FReply::Handled();
 }
 
