Tile Map Editor: Added respect for the AdvancedDisplay flag when adding per-instance properties to the details customization
Tile Set Editor: Ditto for per-tile metadata properties
[UE-14446]

diff --git a/Engine/Plugins/2D/Paper2D/Source/Paper2DEditor/Private/TileMapEditing/PaperTileMapDetailsCustomization.cpp b/Engine/Plugins/2D/Paper2D/Source/Paper2DEditor/Private/TileMapEditing/PaperTileMapDetailsCustomization.cpp
--- a/Engine/Plugins/2D/Paper2D/Source/Paper2DEditor/Private/TileMapEditing/PaperTileMapDetailsCustomization.cpp
+++ b/Engine/Plugins/2D/Paper2D/Source/Paper2DEditor/Private/TileMapEditing/PaperTileMapDetailsCustomization.cpp
@@ -227,10 +227,13 @@ void FPaperTileMapDetailsCustomization::CustomizeDetails(IDetailLayoutBuilder& D
 
 			if (TestProperty->HasAnyPropertyFlags(CPF_Edit))
 			{
-				FName CategoryName(*TestProperty->GetMetaData(TEXT("Category")));
+				const bool bAdvancedDisplay = TestProperty->HasAnyPropertyFlags(CPF_AdvancedDisplay);
+				const EPropertyLocation::Type PropertyLocation = bAdvancedDisplay ? EPropertyLocation::Advanced : EPropertyLocation::Common;
+
+				const FName CategoryName(*TestProperty->GetMetaData(TEXT("Category")));
 				IDetailCategoryBuilder& Category = DetailLayout.EditCategory(CategoryName);
 
-				if (IDetailPropertyRow* ExternalRow = Category.AddExternalProperty(ListOfTileMaps, TestProperty->GetFName()))
+				if (IDetailPropertyRow* ExternalRow = Category.AddExternalProperty(ListOfTileMaps, TestProperty->GetFName(), PropertyLocation))
 				{
 					ExternalRow->Visibility(InternalInstanceVis);
 				}
diff --git a/Engine/Plugins/2D/Paper2D/Source/Paper2DEditor/Private/TileSetEditor/TileSetDetailsCustomization.cpp b/Engine/Plugins/2D/Paper2D/Source/Paper2DEditor/Private/TileSetEditor/TileSetDetailsCustomization.cpp
--- a/Engine/Plugins/2D/Paper2D/Source/Paper2DEditor/Private/TileSetEditor/TileSetDetailsCustomization.cpp
+++ b/Engine/Plugins/2D/Paper2D/Source/Paper2DEditor/Private/TileSetEditor/TileSetDetailsCustomization.cpp
@@ -85,9 +85,12 @@ void FTileSetDetailsCustomization::CustomizeDetails(IDetailLayoutBuilder& Detail
 				UProperty* TestProperty = *PropIt;
 				if (TestProperty->HasAnyPropertyFlags(CPF_Edit))
 				{
+					const bool bAdvancedDisplay = TestProperty->HasAnyPropertyFlags(CPF_AdvancedDisplay);
+					const EPropertyLocation::Type PropertyLocation = bAdvancedDisplay ? EPropertyLocation::Advanced : EPropertyLocation::Common;
+
 					const FString PropertyPath = ArrayEntryPathPrefix + TestProperty->GetName();
 					TSharedRef<IPropertyHandle> PropertyHandle = DetailLayout.GetProperty(*PropertyPath);
-					SingleTileCategory.AddProperty(PropertyHandle);
+					SingleTileCategory.AddProperty(PropertyHandle, PropertyLocation);
 				}
 			}
 
