Fixed some bad font references

UE-8718 - Warning about loading font darta from Roboto-Regular.ttf & Roboto-Bold.ttf

diff --git a/Engine/Plugins/2D/Paper2D/Source/Paper2DEditor/Private/PaperStyle.cpp b/Engine/Plugins/2D/Paper2D/Source/Paper2DEditor/Private/PaperStyle.cpp
--- a/Engine/Plugins/2D/Paper2D/Source/Paper2DEditor/Private/PaperStyle.cpp
+++ b/Engine/Plugins/2D/Paper2D/Source/Paper2DEditor/Private/PaperStyle.cpp
@@ -7,7 +7,8 @@
 #define IMAGE_PLUGIN_BRUSH( RelativePath, ... ) FSlateImageBrush( FPaperStyle::InContent( RelativePath, ".png" ), __VA_ARGS__ )
 #define IMAGE_BRUSH(RelativePath, ...) FSlateImageBrush(StyleSet->RootToContentDir(RelativePath, TEXT(".png")), __VA_ARGS__)
 #define BOX_BRUSH(RelativePath, ...) FSlateBoxBrush(StyleSet->RootToContentDir(RelativePath, TEXT(".png")), __VA_ARGS__)
-#define TTF_FONT( RelativePath, ... ) FSlateFontInfo(StyleSet->RootToContentDir(RelativePath, TEXT(".ttf")), __VA_ARGS__)
+#define TTF_FONT(RelativePath, ...) FSlateFontInfo(StyleSet->RootToContentDir(RelativePath, TEXT(".ttf")), __VA_ARGS__)
+#define TTF_CORE_FONT(RelativePath, ...) FSlateFontInfo(StyleSet->RootToCoreContentDir(RelativePath, TEXT(".ttf") ), __VA_ARGS__)
 
 FString FPaperStyle::InContent(const FString& RelativePath, const ANSICHAR* Extension)
 {
@@ -45,7 +46,7 @@ void FPaperStyle::Initialize()
 			);
 
 		StyleSet->Set("Paper2D.Common.ViewportTitleTextStyle", FTextBlockStyle(NormalText)
-			.SetFont(TTF_FONT("Fonts/Roboto-Regular", 18))
+			.SetFont(TTF_CORE_FONT("Fonts/Roboto-Regular", 18))
 			.SetColorAndOpacity(FLinearColor(1.0, 1.0f, 1.0f, 0.5f))
 			);
 
@@ -225,6 +226,8 @@ void FPaperStyle::Initialize()
 #undef IMAGE_PLUGIN_BRUSH
 #undef IMAGE_BRUSH
 #undef BOX_BRUSH
+#undef TTF_FONT
+#undef TTF_CORE_FONT
 
 void FPaperStyle::Shutdown()
 {
@@ -234,4 +237,4 @@ void FPaperStyle::Shutdown()
 		ensure(StyleSet.IsUnique());
 		StyleSet.Reset();
 	}
-}
\ No newline at end of file
+}
