Fixed tile map ed mode consuming all input when control was down
Fixed crash when pressing shift while already painting

diff --git a/Engine/Plugins/2D/Paper2D/Source/Paper2DEditor/Private/TileMapEditing/EdModeTileMap.cpp b/Engine/Plugins/2D/Paper2D/Source/Paper2DEditor/Private/TileMapEditing/EdModeTileMap.cpp
--- a/Engine/Plugins/2D/Paper2D/Source/Paper2DEditor/Private/TileMapEditing/EdModeTileMap.cpp
+++ b/Engine/Plugins/2D/Paper2D/Source/Paper2DEditor/Private/TileMapEditing/EdModeTileMap.cpp
@@ -356,48 +356,6 @@ bool FEdModeTileMap::InputKey(FEditorViewportClient* InViewportClient, FViewport
 			bAnyPaintAbleActorsUnderCursor = UseActiveToolAtLocation(Ray);
 		}
 		bWasPainting = bIsPainting;
-
-		// Also absorb other mouse buttons, and Ctrl/Alt/Shift events that occur while we're painting as these would cause
-		// the editor viewport to start panning/dollying the camera
-		{
-			const bool bIsOtherMouseButtonEvent = ( InKey == EKeys::MiddleMouseButton || InKey == EKeys::RightMouseButton );
-			const bool bCtrlButtonEvent = (InKey == EKeys::LeftControl || InKey == EKeys::RightControl);
-			const bool bShiftButtonEvent = (InKey == EKeys::LeftShift || InKey == EKeys::RightShift);
-			const bool bAltButtonEvent = (InKey == EKeys::LeftAlt || InKey == EKeys::RightAlt);
-			if( bIsPainting && ( bIsOtherMouseButtonEvent || bShiftButtonEvent || bAltButtonEvent ) )
-			{
-				bHandled = true;
-			}
-
-			if (bCtrlButtonEvent && !bIsPainting)
-			{
-				bHandled = false;
-			}
-			else if (bIsCtrlDown)
-			{
-				//default to assuming this is a paint command
-				bHandled = true;
-
-				// If no other button was pressed && if a first press and we click OFF of an actor and we will let this pass through so multi-select can attempt to handle it 
-				if ((!(bShiftButtonEvent || bAltButtonEvent || bIsOtherMouseButtonEvent)) && ((InKey == EKeys::LeftMouseButton) && ((InEvent == IE_Pressed) || (InEvent == IE_Released)) && (!bAnyPaintAbleActorsUnderCursor)))
-				{
-					bHandled = false;
-					bIsPainting = false;
-				}
-
-				// Allow Ctrl+B to pass through so we can support the finding of a selected static mesh in the content browser.
-				if ( !(bShiftButtonEvent || bAltButtonEvent || bIsOtherMouseButtonEvent) && ( (InKey == EKeys::B) && (InEvent == IE_Pressed) ) )
-				{
-					bHandled = false;
-				}
-
-				// If we are not painting, we will let the CTRL-Z and CTRL-Y key presses through to support undo/redo.
-				if (!bIsPainting && ((InKey == EKeys::Z) || (InKey == EKeys::Y)))
-				{
-					bHandled = false;
-				}
-			}
-		}
 	}
 
 	if (!bHandled)
@@ -1504,7 +1462,7 @@ void FEdModeTileMap::SetActiveTool(ETileMapEditorTool::Type NewTool)
 ETileMapEditorTool::Type FEdModeTileMap::GetActiveTool() const
 {
 	// Force the eyedropper active when Shift is held (or if it was held when painting started, even if it was released later)
-	const bool bHoldingShift = FSlateApplication::Get().GetModifierKeys().IsShiftDown();
+	const bool bHoldingShift = !bIsPainting && FSlateApplication::Get().GetModifierKeys().IsShiftDown();
 	const bool bWasHoldingShift = bIsPainting && bWasHoldingSelectWhenPaintingStarted;
 	
 	return (bHoldingShift || bWasHoldingShift) ? ETileMapEditorTool::EyeDropper : ActiveTool;
