Paper2D: Added support for F to zoom to fit in the tile selector viewport and a good default zoom
- Added more zoom ranges to SPaperEditorViewport (16x, 32x, 1:16, 1:32)
- Reparented SPaperEditorViewport to SEditorViewport and removed some duplicated code
- Remove obsolete dependencies on GraphEditor and LevelEditor from Paper2DEditor
[UE-14428]

diff --git a/Engine/Plugins/2D/Paper2D/Source/Paper2DEditor/Paper2DEditor.Build.cs b/Engine/Plugins/2D/Paper2D/Source/Paper2DEditor/Paper2DEditor.Build.cs
--- a/Engine/Plugins/2D/Paper2D/Source/Paper2DEditor/Paper2DEditor.Build.cs
+++ b/Engine/Plugins/2D/Paper2D/Source/Paper2DEditor/Paper2DEditor.Build.cs
@@ -19,11 +19,9 @@ public class Paper2DEditor : ModuleRules
                 "InputCore",
 				"UnrealEd", // for FAssetEditorManager
 				"KismetWidgets",
-				"GraphEditor",
 				"Kismet",  // for FWorkflowCentricApplication
 				"PropertyEditor",
 				"RenderCore",
-				"LevelEditor", // for EdModes to get a toolkit host  //@TODO: PAPER: Should be a better way to do this (see the @todo in EdModeTileMap.cpp)
 				"Paper2D",
 				"ContentBrowser",
 				"WorkspaceMenuStructure",
diff --git a/Engine/Plugins/2D/Paper2D/Source/Paper2DEditor/Private/ExtractSprites/SPaperExtractSpritesDialog.cpp b/Engine/Plugins/2D/Paper2D/Source/Paper2DEditor/Private/ExtractSprites/SPaperExtractSpritesDialog.cpp
--- a/Engine/Plugins/2D/Paper2D/Source/Paper2DEditor/Private/ExtractSprites/SPaperExtractSpritesDialog.cpp
+++ b/Engine/Plugins/2D/Paper2D/Source/Paper2DEditor/Private/ExtractSprites/SPaperExtractSpritesDialog.cpp
@@ -125,7 +125,7 @@ void SPaperExtractSpritesViewport::Construct(const FArguments& InArgs, UTexture2
 	// Make sure we get input instead of the viewport stealing it
 	ViewportWidget->SetVisibility(EVisibility::HitTestInvisible);
 
-	RefreshViewport();
+	Invalidate();
 }
 
 FText SPaperExtractSpritesViewport::GetTitleText() const
diff --git a/Engine/Plugins/2D/Paper2D/Source/Paper2DEditor/Private/SPaperEditorViewport.cpp b/Engine/Plugins/2D/Paper2D/Source/Paper2DEditor/Private/SPaperEditorViewport.cpp
--- a/Engine/Plugins/2D/Paper2D/Source/Paper2DEditor/Private/SPaperEditorViewport.cpp
+++ b/Engine/Plugins/2D/Paper2D/Source/Paper2DEditor/Private/SPaperEditorViewport.cpp
@@ -5,8 +5,8 @@
 #include "PaperEditorViewportClient.h"
 #include "PaperStyle.h"
 
-const int32 DefaultZoomLevel = 5;
-const int32 NumZoomLevels = 13;
+const int32 DefaultZoomLevel = 7;
+const int32 NumZoomLevels = 17;
 
 struct FZoomLevelEntry
 {
@@ -22,6 +22,8 @@ struct FZoomLevelEntry
 
 static const FZoomLevelEntry ZoomLevels[NumZoomLevels] =
 {
+	FZoomLevelEntry(0.03125f, NSLOCTEXT("PaperEditor", "ZoomLevel", "1:32")),
+	FZoomLevelEntry(0.0625f, NSLOCTEXT("PaperEditor", "ZoomLevel", "1:16")),
 	FZoomLevelEntry(0.125f, NSLOCTEXT("PaperEditor", "ZoomLevel", "1:8")),
 	FZoomLevelEntry(0.250f, NSLOCTEXT("PaperEditor", "ZoomLevel", "1:4")),
 	FZoomLevelEntry(0.500f, NSLOCTEXT("PaperEditor", "ZoomLevel", "1:2")),
@@ -34,7 +36,9 @@ static const FZoomLevelEntry ZoomLevels[NumZoomLevels] =
 	FZoomLevelEntry(5.000f, NSLOCTEXT("PaperEditor", "ZoomLevel", "5x")),
 	FZoomLevelEntry(6.000f, NSLOCTEXT("PaperEditor", "ZoomLevel", "6x")),
 	FZoomLevelEntry(7.000f, NSLOCTEXT("PaperEditor", "ZoomLevel", "7x")),
-	FZoomLevelEntry(8.000f, NSLOCTEXT("PaperEditor", "ZoomLevel", "8x"))
+	FZoomLevelEntry(8.000f, NSLOCTEXT("PaperEditor", "ZoomLevel", "8x")),
+	FZoomLevelEntry(16.000f, NSLOCTEXT("PaperEditor", "ZoomLevel", "16x")),
+	FZoomLevelEntry(32.000f, NSLOCTEXT("PaperEditor", "ZoomLevel", "32x"))
 };
 
 //////////////////////////////////////////////////////////////////////////
@@ -42,32 +46,25 @@ static const FZoomLevelEntry ZoomLevels[NumZoomLevels] =
 
 SPaperEditorViewport::~SPaperEditorViewport()
 {
-	// Close viewport
-	if (ViewportClient.IsValid())
-	{
-		ViewportClient->Viewport = nullptr;
-	}
+	PaperViewportClient.Reset();
+}
 
-	Viewport.Reset();
-	ViewportClient.Reset();
+TSharedRef<FEditorViewportClient> SPaperEditorViewport::MakeEditorViewportClient()
+{
+	return PaperViewportClient.ToSharedRef();
 }
 
 void SPaperEditorViewport::Construct(const FArguments& InArgs, TSharedRef<FPaperEditorViewportClient> InViewportClient)
 {
 	OnSelectionChanged = InArgs._OnSelectionChanged;
 
-	this->ChildSlot
-	[
-		SNew(SOverlay)
-		+SOverlay::Slot()
-		[
-			SAssignNew(ViewportWidget, SViewport)
-			.EnableGammaCorrection(false)
-			.IsEnabled(FSlateApplication::Get().GetNormalExecutionAttribute())
-			.ShowEffectWhenDisabled(false)
-		]
-		// Indicator of current zoom level
-		+SOverlay::Slot()
+	PaperViewportClient = InViewportClient;
+	PaperViewportClient->SetRealtime(false);
+
+	SEditorViewport::Construct(SEditorViewport::FArguments());
+
+	// Indicator of current zoom level
+	ViewportOverlay->AddSlot()
 		.Padding(5)
 		.VAlign(VAlign_Top)
 		[
@@ -75,8 +72,10 @@ void SPaperEditorViewport::Construct(const FArguments& InArgs, TSharedRef<FPaper
 			.TextStyle(FPaperStyle::Get(), "Paper2D.Common.ViewportZoomTextStyle")
 			.Text(this, &SPaperEditorViewport::GetZoomText)
 			.ColorAndOpacity(this, &SPaperEditorViewport::GetZoomTextColorAndOpacity)
-		]
-		+SOverlay::Slot()
+		];
+
+	// Title for the viewport
+	ViewportOverlay->AddSlot()
 		.VAlign(VAlign_Top)
 		[
 			SNew(SBorder)
@@ -100,38 +99,20 @@ void SPaperEditorViewport::Construct(const FArguments& InArgs, TSharedRef<FPaper
 					]
 				]
 			]
-		]
-	];
-
-	ViewportClient = InViewportClient;
-	ViewportClient->SetRealtime(false);
-
-	Viewport = MakeShareable(new FSceneViewport(ViewportClient.Get(), ViewportWidget));
-	ViewportClient->Viewport = Viewport.Get();
-
-	// The viewport widget needs an interface so it knows what should render
-	ViewportWidget->SetViewportInterface( Viewport.ToSharedRef() );
-
+		];
 
 	ZoomLevel = DefaultZoomLevel;
 	PreviousZoomLevel = DefaultZoomLevel;
 	ViewOffset = FVector2D::ZeroVector;
 	TotalMouseDelta = 0;
-	//bAllowContinousZoomInterpolation = false;
 
 	bIsPanning = false;
 
 	ZoomLevelFade = FCurveSequence( 0.0f, 0.75f );
 	ZoomLevelFade.Play(this->AsShared());
 
-	ZoomLevelGraphFade = FCurveSequence( 0.0f, 0.5f );
-	ZoomLevelGraphFade.Play(this->AsShared());
-
 	DeferredPanPosition = FVector2D::ZeroVector;
 	bRequestDeferredPan = false;
-
-	CommandList = MakeShareable(new FUICommandList);
-	BindCommands();
 }
 
 void SPaperEditorViewport::Tick( const FGeometry& AllottedGeometry, const double InCurrentTime, const float InDeltaTime )
@@ -149,8 +130,8 @@ void SPaperEditorViewport::Tick( const FGeometry& AllottedGeometry, const double
 		bIsPanning = false;
 	}
 
-	ViewportClient->SetZoomPos( ViewOffset, GetZoomAmount() );
-	ViewportClient->bNeedsRedraw = true;
+	PaperViewportClient->SetZoomPos( ViewOffset, GetZoomAmount() );
+	PaperViewportClient->bNeedsRedraw = true;
 
 	bool bSelectionModified = false;
 	if (Marquee.IsValid())
@@ -163,11 +144,11 @@ void SPaperEditorViewport::Tick( const FGeometry& AllottedGeometry, const double
 	if (bSelectionModified || bIsPanning || FSlateThrottleManager::Get().IsAllowingExpensiveTasks())
 	{
 		// Setup the selection set for the viewport
-		ViewportClient->SelectionRectangles.Empty();
+		PaperViewportClient->SelectionRectangles.Empty();
 
 		if (Marquee.IsValid())
 		{
-			FViewportSelectionRectangle& Rect = *(new (ViewportClient->SelectionRectangles) FViewportSelectionRectangle);
+			FViewportSelectionRectangle& Rect = *(new (PaperViewportClient->SelectionRectangles) FViewportSelectionRectangle);
 			Rect.Color = FColorList::Yellow;
 			Rect.Color.A = 0.45f;
 			Rect.TopLeft = Marquee.Rect.GetUpperLeft();
@@ -175,8 +156,8 @@ void SPaperEditorViewport::Tick( const FGeometry& AllottedGeometry, const double
 		}
 
 		// Tick and render the viewport
-		ViewportClient->Tick(InDeltaTime);
-		GEditor->UpdateSingleViewportClient(ViewportClient.Get(), /*bInAllowNonRealtimeViewportToDraw=*/ true, /*bLinkedOrthoMovement=*/ false);
+		PaperViewportClient->Tick(InDeltaTime);
+		GEditor->UpdateSingleViewportClient(PaperViewportClient.Get(), /*bInAllowNonRealtimeViewportToDraw=*/ true, /*bLinkedOrthoMovement=*/ false);
 	}
 }
 
@@ -206,7 +187,7 @@ FReply SPaperEditorViewport::OnMouseButtonDown(const FGeometry& MyGeometry, cons
 
 		// Trigger a selection update now so that single-clicks without a drag still select something
 		OnSelectionChanged.ExecuteIfBound(Marquee, true);
-		ViewportClient->Invalidate();
+		PaperViewportClient->Invalidate();
 
 		return FReply::Handled().CaptureMouse( SharedThis(this) );
 	}
@@ -227,7 +208,7 @@ FReply SPaperEditorViewport::OnMouseButtonUp(const FGeometry& MyGeometry, const
 	{
 		FReply ReplyState = FReply::Handled();
 
-		if (this->HasMouseCapture())
+		if (HasMouseCapture())
 		{
 			FSlateRect ThisPanelScreenSpaceRect = MyGeometry.GetClippingRect();
 			const FVector2D ScreenSpaceCursorPos = MyGeometry.LocalToAbsolute( GraphCoordToPanelCoord( SoftwareCursorPosition ) );
@@ -251,7 +232,7 @@ FReply SPaperEditorViewport::OnMouseButtonUp(const FGeometry& MyGeometry, const
 
 		bShowSoftwareCursor = false;
 
-		this->bIsPanning = false;
+		bIsPanning = false;
 		return (WidgetToFocus.IsValid())
 			? ReplyState.ReleaseMouseCapture().SetUserFocus(WidgetToFocus.ToSharedRef(), EFocusCause::SetDirectly)
 			: ReplyState.ReleaseMouseCapture();
@@ -276,7 +257,7 @@ FReply SPaperEditorViewport::OnMouseButtonUp(const FGeometry& MyGeometry, const
 // 				NodeUnderMousePtr.Reset();
 // 			}
 // 			else 
-			if (this->HasMouseCapture())
+			if (HasMouseCapture())
  			{
  				// We clicked on the panel background
 
@@ -316,7 +297,7 @@ FReply SPaperEditorViewport::OnMouseMove(const FGeometry& MyGeometry, const FPoi
 	const bool bIsRightMouseButtonDown = MouseEvent.IsMouseButtonDown(EKeys::RightMouseButton);
 	const bool bIsLeftMouseButtonDown = MouseEvent.IsMouseButtonDown(EKeys::LeftMouseButton);
 	
-	if (this->HasMouseCapture())
+	if (HasMouseCapture())
 	{
 		// Track how much the mouse moved since the mouse down.
 		const FVector2D CursorDelta = MouseEvent.GetCursorDelta();
@@ -345,48 +326,11 @@ FReply SPaperEditorViewport::OnMouseMove(const FGeometry& MyGeometry, const FPoi
 
 			const bool bCursorInDeadZone = TotalMouseDelta <= FSlateApplication::Get().GetDragTriggerDistance();
 
-// 			if (NodeBeingDragged.IsValid())
-// 			{
-// 				if (!bCursorInDeadZone)
-// 				{
-// 					// Note, NodeGrabOffset() comes from the node itself, so it's already scaled correctly.
-// 					FVector2D AnchorNodeNewPos = PanelCoordToGraphCoord( MyGeometry.AbsoluteToLocal( MouseEvent.GetScreenSpacePosition() ) ) - NodeGrabOffset;
-// 
-// 					// Snap to grid
-// 					const float SnapSize = GetSnapGridSize();
-// 					AnchorNodeNewPos.X = SnapSize * FMath::RoundToInt(AnchorNodeNewPos.X/SnapSize);
-// 					AnchorNodeNewPos.Y = SnapSize * FMath::RoundToInt(AnchorNodeNewPos.Y/SnapSize);
-// 
-// 					// Dragging an unselected node automatically selects it.
-// 					SelectionManager.StartDraggingNode(NodeBeingDragged->GetObjectBeingDisplayed(), MouseEvent);
-// 
-// 					// Move all the selected nodes.
-// 					{
-// 						const FVector2D AnchorNodeOldPos = NodeBeingDragged->GetPosition();
-// 						const FVector2D DeltaPos = AnchorNodeNewPos - AnchorNodeOldPos;
-// 
-// 						for (FGraphPanelSelectionSet::TIterator NodeIt(SelectionManager.SelectedNodes); NodeIt; ++NodeIt)
-// 						{
-// 							TSharedRef<SNode>* pWidget = NodeToWidgetLookup.Find(*NodeIt);
-// 							if (pWidget != nullptr)
-// 							{
-// 								SNode& Widget = pWidget->Get();
-// 								Widget.MoveTo(Widget.GetPosition() + DeltaPos);
-// 							}
-// 						}
-// 					}
-// 				}
-// 
-// 				return FReply::Handled();
-// 			}
-// 
-// 			if (!NodeBeingDragged.IsValid())
 			{
 				// We are marquee selecting
 				const FVector2D GraphMousePos = PanelCoordToGraphCoord( MyGeometry.AbsoluteToLocal( MouseEvent.GetScreenSpacePosition() ) );
 				Marquee.Rect.UpdateEndPoint(GraphMousePos);
 
-//				FindNodesAffectedByMarquee( /*out*/ Marquee.AffectedNodes );
 				return FReply::Handled();
 			}
 		}
@@ -405,12 +349,6 @@ FReply SPaperEditorViewport::OnMouseWheel(const FGeometry& MyGeometry, const FPo
 	const int32 ZoomLevelDelta = FMath::FloorToInt(MouseEvent.GetWheelDelta());
 
 	const bool bAllowFullZoomRange = true;
-/*
-		// To zoom out past 1:1 the user must press control
-		(ZoomLevel == DefaultZoomLevel && (ZoomLevelDelta > 0) && MouseEvent.IsControlDown()) ||
-		// If they are already zoomed out past 1:1, user may zoom freely
-		(ZoomLevel > DefaultZoomLevel);
-*/
 	const float OldZoomLevel = ZoomLevel;
 
 	if (bAllowFullZoomRange)
@@ -427,15 +365,7 @@ FReply SPaperEditorViewport::OnMouseWheel(const FGeometry& MyGeometry, const FPo
 
 
 	// Re-center the screen so that it feels like zooming around the cursor.
-	{
- 		FSlateRect GraphBounds = ComputeSensibleGraphBounds();
- 		// Make sure we are not zooming into/out into emptiness; otherwise the user will get lost..
- 		const FVector2D ClampedPointToMaintainGraphSpace(
- 			FMath::Clamp(PointToMaintainGraphSpace.X, GraphBounds.Left, GraphBounds.Right),
- 			FMath::Clamp(PointToMaintainGraphSpace.Y, GraphBounds.Top, GraphBounds.Bottom)
- 			);
-		this->ViewOffset = ClampedPointToMaintainGraphSpace - WidgetSpaceCursorPos / GetZoomAmount();
-	}
+	ViewOffset = PointToMaintainGraphSpace - WidgetSpaceCursorPos / GetZoomAmount();
 
 	return FReply::Handled();
 }
@@ -447,39 +377,43 @@ FCursorReply SPaperEditorViewport::OnCursorQuery( const FGeometry& MyGeometry, c
 		FCursorReply::Cursor( EMouseCursor::Default );
 }
 
-void SPaperEditorViewport::RefreshViewport()
+int32 SPaperEditorViewport::FindNearestZoomLevel(float InZoomAmount, bool bRoundDown) const
 {
-	Viewport->Invalidate();
-	Viewport->InvalidateDisplay();
-}
+	bool bFoundItem = false;
+	int32 Result = 0;
 
-/** Search through the list of zoom levels and find the one closest to what was passed in. */
-int32 SPaperEditorViewport::FindNearestZoomLevel(int32 CurrentZoomLevel, float InZoomAmount) const
-{
 	for (int32 ZoomLevelIndex=0; ZoomLevelIndex < NumZoomLevels; ++ZoomLevelIndex)
 	{
 		if (InZoomAmount <= ZoomLevels[ZoomLevelIndex].ZoomAmount)
 		{
-			return ZoomLevelIndex;
+			Result = ZoomLevelIndex;
+			bFoundItem = true;
+			break;
 		}
 	}
 
-	return DefaultZoomLevel;
-}
-
-
-float SPaperEditorViewport::GetZoomAmount() const
-{
-	if (false)//bAllowContinousZoomInterpolation)
+	if (!bFoundItem)
 	{
-		return FMath::Lerp(ZoomLevels[PreviousZoomLevel].ZoomAmount, ZoomLevels[ZoomLevel].ZoomAmount, ZoomLevelGraphFade.GetLerp());
+		return DefaultZoomLevel;
 	}
 	else
 	{
-		return ZoomLevels[ZoomLevel].ZoomAmount;
+		if (bFoundItem && bRoundDown && (Result > 0))
+		{
+			return Result - 1;
+		}
+		else
+		{
+			return Result;
+		}
 	}
 }
 
+float SPaperEditorViewport::GetZoomAmount() const
+{
+	return ZoomLevels[ZoomLevel].ZoomAmount;
+}
+
 FText SPaperEditorViewport::GetZoomText() const
 {
 	return ZoomLevels[ZoomLevel].DisplayText;
@@ -590,21 +524,6 @@ int32 SPaperEditorViewport::OnPaint(const FPaintArgs& Args, const FGeometry& All
 	return MaxLayerId;
 }
 
-FSlateRect SPaperEditorViewport::ComputeSensibleGraphBounds() const
-{
-	float Left = 0.0f;
-	float Top = 0.0f;
-
-	//@TODO: PAPER: Use real values!
-	float Right = 512.0f;//0.0f;
-	float Bottom = 1024.0f;
-
-	// Pad it out in every direction, to roughly account for nodes being of non-zero extent
-	const float Padding = 100.0f;
-
-	return FSlateRect( Left - Padding, Top - Padding, Right + Padding, Bottom + Padding );
-}
-
 FText SPaperEditorViewport::GetTitleText() const
 {
 	return NSLOCTEXT("PaperEditor", "TileSetPaletteTitle", "tile set palette");
diff --git a/Engine/Plugins/2D/Paper2D/Source/Paper2DEditor/Private/SPaperEditorViewport.h b/Engine/Plugins/2D/Paper2D/Source/Paper2DEditor/Private/SPaperEditorViewport.h
--- a/Engine/Plugins/2D/Paper2D/Source/Paper2DEditor/Private/SPaperEditorViewport.h
+++ b/Engine/Plugins/2D/Paper2D/Source/Paper2DEditor/Private/SPaperEditorViewport.h
@@ -2,8 +2,8 @@
 
 #pragma once
 
-#include "GraphEditor.h"
 #include "MarqueeRect.h"
+#include "SEditorViewport.h"
 
 /** Helper for managing marquee operations */
 struct FMarqueeOperation
@@ -68,7 +68,7 @@ public:
 //////////////////////////////////////////////////////////////////////////
 // SPaperEditorViewport
 
-class SPaperEditorViewport : public SCompoundWidget
+class SPaperEditorViewport : public SEditorViewport
 {
 public:
 	DECLARE_DELEGATE_TwoParams(FOnSelectionChanged, FMarqueeOperation /*MarqueeAction*/, bool /*bPreview*/);
@@ -92,21 +92,19 @@ public:
 
 	void Construct(const FArguments& InArgs, TSharedRef<class FPaperEditorViewportClient> InViewportClient);
 
-	/** Refreshes the viewport */
-	void RefreshViewport();
-
 	float GetZoomAmount() const;
 	FText GetZoomText() const;
 	FSlateColor GetZoomTextColorAndOpacity() const;
 	FVector2D GetViewOffset() const;
 
-	virtual void BindCommands() { }
+protected:
 
-	/** @return The viewport command list */
-	const TSharedPtr<FUICommandList> GetCommandList() const { return CommandList; }
+	// SEditorViewport interface
+	virtual TSharedRef<FEditorViewportClient> MakeEditorViewportClient() override;
+	// End of SEditorViewport interface
 
-protected:
-	int32 FindNearestZoomLevel(int32 CurrentZoomLevel, float InZoomAmount) const;
+
+	int32 FindNearestZoomLevel(float InZoomAmount, bool bRoundDown) const;
 
 
 	FVector2D ComputeEdgePanAmount(const FGeometry& MyGeometry, const FVector2D& TargetPosition);
@@ -119,8 +117,6 @@ protected:
 
 	void PaintSoftwareCursor(const FGeometry& AllottedGeometry, const FSlateRect& MyClippingRect, FSlateWindowElementList& OutDrawElements, int32 DrawLayerId) const;
 
-	FSlateRect ComputeSensibleGraphBounds() const;
-
 	virtual FText GetTitleText() const;
 protected:
 	/** The position within the graph at which the user is looking */
@@ -141,12 +137,6 @@ protected:
 	/** A pending marquee operation if it's active */
 	FMarqueeOperation Marquee;
 
-	/** Allow continuous zoom interpolation? */
-	bool bAllowContinousZoomInterpolation;
-
-	/** Fade on zoom for graph */
-	FCurveSequence ZoomLevelGraphFade;
-
 	/** Curve that handles fading the 'Zoom +X' text */
 	FCurveSequence ZoomLevelFade;
 
@@ -163,17 +153,8 @@ protected:
 	bool bShowSoftwareCursor;
 
 	/** Level viewport client */
-	TSharedPtr<class FPaperEditorViewportClient> ViewportClient;
-
-	/** Slate viewport for rendering and I/O */
-	TSharedPtr<FSceneViewport> Viewport;
-
-	/** Viewport widget*/
-	TSharedPtr<SViewport> ViewportWidget;
+	TSharedPtr<class FPaperEditorViewportClient> PaperViewportClient;
 
 	// Selection changed delegate
 	FOnSelectionChanged OnSelectionChanged;
-
-	// The command list
-	TSharedPtr<FUICommandList> CommandList;
 };
diff --git a/Engine/Plugins/2D/Paper2D/Source/Paper2DEditor/Private/TileMapEditing/EdModeTileMap.cpp b/Engine/Plugins/2D/Paper2D/Source/Paper2DEditor/Private/TileMapEditing/EdModeTileMap.cpp
--- a/Engine/Plugins/2D/Paper2D/Source/Paper2DEditor/Private/TileMapEditing/EdModeTileMap.cpp
+++ b/Engine/Plugins/2D/Paper2D/Source/Paper2DEditor/Private/TileMapEditing/EdModeTileMap.cpp
@@ -3,7 +3,6 @@
 #include "Paper2DEditorPrivatePCH.h"
 #include "EdModeTileMap.h"
 #include "TileMapEdModeToolkit.h"
-#include "LevelEditor.h"
 #include "Toolkits/ToolkitManager.h"
 #include "../PaperEditorCommands.h"
 #include "ScopedTransaction.h"
diff --git a/Engine/Plugins/2D/Paper2D/Source/Paper2DEditor/Private/TileSetEditor/TileSetSelectorViewport.cpp b/Engine/Plugins/2D/Paper2D/Source/Paper2DEditor/Private/TileSetEditor/TileSetSelectorViewport.cpp
--- a/Engine/Plugins/2D/Paper2D/Source/Paper2DEditor/Private/TileSetEditor/TileSetSelectorViewport.cpp
+++ b/Engine/Plugins/2D/Paper2D/Source/Paper2DEditor/Private/TileSetEditor/TileSetSelectorViewport.cpp
@@ -18,6 +18,8 @@ STileSetSelectorViewport::~STileSetSelectorViewport()
 
 void STileSetSelectorViewport::Construct(const FArguments& InArgs, UPaperTileSet* InTileSet, FEdModeTileMap* InTileMapEditor)
 {
+	bPendingZoom = true;
+
 	SelectionTopLeft = FIntPoint::ZeroValue;
 	SelectionDimensions = FIntPoint::ZeroValue;
 
@@ -79,6 +81,44 @@ void STileSetSelectorViewport::BindCommands()
 		FIsActionChecked::CreateSP(EditorViewportClientRef, &FTileSetEditorViewportClient::IsShowTilesWithMetaDataChecked));
 }
 
+void STileSetSelectorViewport::Tick(const FGeometry& AllottedGeometry, const double InCurrentTime, const float InDeltaTime)
+{
+	if (bPendingZoom && (SceneViewport->GetSize().GetMin() > 0))
+	{
+		OnFocusViewportToSelection();
+		bPendingZoom = false;
+	}
+
+	SPaperEditorViewport::Tick(AllottedGeometry, InCurrentTime, InDeltaTime);
+}
+
+void STileSetSelectorViewport::OnFocusViewportToSelection()
+{
+	//@TODO: The rest of this method would get replaced by the following if we use a component to render the tile set
+	//TypedViewportClient->RequestFocusOnSelection(/*bInstant=*/ false);
+
+	if (UPaperTileSet* TileSetBeingEdited = TileSetPtr.Get())
+	{
+		const FIntPoint ViewportSize = SceneViewport->GetSize();
+		if (UTexture* Texture = TileSetBeingEdited->TileSheet)
+		{
+			const FIntPoint TextureSize(Texture->Source.GetSizeX(), Texture->Source.GetSizeY());
+
+			const float CopiesInX = ViewportSize.X / (float)TextureSize.X;
+			const float CopiesInY = ViewportSize.Y / (float)TextureSize.Y;
+
+			const float NumCopies = FMath::Min<float>(CopiesInX, CopiesInY);
+
+			// Find the zoom level that would match
+			ZoomLevel = FindNearestZoomLevel(NumCopies, /*bRoundDown=*/ true);
+
+			// Figure out the desired zoom position
+			const float ZoomAmount = GetZoomAmount();
+			ViewOffset.X = -(ViewportSize.X / ZoomAmount - TextureSize.X) * 0.5f;
+			ViewOffset.Y = -(ViewportSize.Y / ZoomAmount - TextureSize.Y) * 0.5f;
+		}
+	}
+}
 
 void STileSetSelectorViewport::OnSelectionChanged(FMarqueeOperation Marquee, bool bIsPreview)
 {
diff --git a/Engine/Plugins/2D/Paper2D/Source/Paper2DEditor/Private/TileSetEditor/TileSetSelectorViewport.h b/Engine/Plugins/2D/Paper2D/Source/Paper2DEditor/Private/TileSetEditor/TileSetSelectorViewport.h
--- a/Engine/Plugins/2D/Paper2D/Source/Paper2DEditor/Private/TileSetEditor/TileSetSelectorViewport.h
+++ b/Engine/Plugins/2D/Paper2D/Source/Paper2DEditor/Private/TileSetEditor/TileSetSelectorViewport.h
@@ -29,11 +29,19 @@ public:
 	}
 
 protected:
+	// SWidget interface
+	virtual void Tick(const FGeometry& AllottedGeometry, const double InCurrentTime, const float InDeltaTime) override;
+	// End of SWidget interface
+
 	// SPaperEditorViewport interface
 	virtual FText GetTitleText() const override;
-	virtual void BindCommands() override;
 	// End of SPaperEditorViewport interface
 
+	// SEditorViewport interface
+	virtual void BindCommands() override;
+	virtual void OnFocusViewportToSelection() override;
+	// End of SEditorViewport interface
+
 private:
 	void OnSelectionChanged(FMarqueeOperation Marquee, bool bIsPreview);
 	void RefreshSelectionRectangle();
@@ -44,6 +52,7 @@ private:
 	class FEdModeTileMap* TileMapEditor;
 	FIntPoint SelectionTopLeft;
 	FIntPoint SelectionDimensions;
+	bool bPendingZoom;
 
 	FOnTileViewportSelectionChanged OnTileSelectionChanged;
 };
\ No newline at end of file
