Tile Set Editor: Fixed a rare crash when changing the active tile [UE-14612]

diff --git a/Engine/Plugins/2D/Paper2D/Source/Paper2DEditor/Private/TileSetEditor.cpp b/Engine/Plugins/2D/Paper2D/Source/Paper2DEditor/Private/TileSetEditor.cpp
--- a/Engine/Plugins/2D/Paper2D/Source/Paper2DEditor/Private/TileSetEditor.cpp
+++ b/Engine/Plugins/2D/Paper2D/Source/Paper2DEditor/Private/TileSetEditor.cpp
@@ -56,6 +56,7 @@ private:
 	TSharedRef<IDetailCustomization> MakeEmbeddedInstance()
 	{
 		TSharedRef<FTileSetDetailsCustomization> Customization = FTileSetDetailsCustomization::MakeEmbeddedInstance();
+		CurrentCustomizationPtr = Customization;
 
 		// Make sure the customization starts off looking at the right tile index
 		TSharedPtr<FSingleTileEditorViewportClient> SingleTileEditor = TileSetEditorPtr.Pin()->GetSingleTileEditor();
@@ -66,15 +67,24 @@ private:
 			Customization->OnTileIndexChanged(TileIndex, INDEX_NONE);
 		}
 
-		// And that it gets updated when the tile index changes
-		SingleTileEditor->GetOnSingleTileIndexChanged().AddSP(Customization, &FTileSetDetailsCustomization::OnTileIndexChanged);
-
 		return Customization;
 	}
 
+	void OnTileIndexChanged(int32 NewIndex, int32 OldIndex)
+	{
+		if (FTileSetDetailsCustomization* CurrentCustomization = CurrentCustomizationPtr.Pin().Get())
+		{
+			CurrentCustomization->OnTileIndexChanged(NewIndex, OldIndex);
+		}
+	}
+
 private:
 	// Pointer back to owning TileSet editor instance (the keeper of state)
 	TWeakPtr<class FTileSetEditor> TileSetEditorPtr;
+
+	// Pointer to the allocated customization
+	TWeakPtr<FTileSetDetailsCustomization> CurrentCustomizationPtr;
+
 public:
 	void Construct(const FArguments& InArgs, TSharedPtr<FTileSetEditor> InTileSetEditor)
 	{
@@ -82,6 +92,11 @@ public:
 
 		SSingleObjectDetailsPanel::Construct(SSingleObjectDetailsPanel::FArguments().HostCommandList(InTileSetEditor->GetToolkitCommands()));
 
+		// Register for index change notifications
+		TSharedPtr<FSingleTileEditorViewportClient> SingleTileEditor = TileSetEditorPtr.Pin()->GetSingleTileEditor();
+		SingleTileEditor->GetOnSingleTileIndexChanged().AddSP(this, &STileSetPropertiesTabBody::OnTileIndexChanged);
+
+		// Register our customization that will be notified of tile index changes
 		FOnGetDetailCustomizationInstance LayoutTileSetDetails = FOnGetDetailCustomizationInstance::CreateSP(this, &STileSetPropertiesTabBody::MakeEmbeddedInstance);
 		PropertyView->RegisterInstancedCustomPropertyLayout(UPaperTileSet::StaticClass(), LayoutTileSetDetails);
 	}
