Paper2D - Fixed missing collision on Paper Terrain spline [UE-11445]

diff --git a/Engine/Plugins/2D/Paper2D/Source/Paper2D/Private/Terrain/PaperTerrainComponent.cpp b/Engine/Plugins/2D/Paper2D/Source/Paper2D/Private/Terrain/PaperTerrainComponent.cpp
--- a/Engine/Plugins/2D/Paper2D/Source/Paper2D/Private/Terrain/PaperTerrainComponent.cpp
+++ b/Engine/Plugins/2D/Paper2D/Source/Paper2D/Private/Terrain/PaperTerrainComponent.cpp
@@ -692,6 +692,15 @@ void UPaperTerrainComponent::OnSplineEdited()
 #endif
 	}
 
+	if (CachedBodySetup != nullptr)
+	{
+		// Finalize the BodySetup
+#if WITH_RUNTIME_PHYSICS_COOKING || WITH_EDITOR
+		CachedBodySetup->InvalidatePhysicsData();
+#endif
+		CachedBodySetup->CreatePhysicsMeshes();
+	}
+
 	RecreateRenderState_Concurrent();
 }
 
@@ -703,6 +712,9 @@ void UPaperTerrainComponent::SpawnSegments(const TArray<FTerrainSegment>& Terrai
 	TArray<FVector2D> CollisionPolygonPoints;
 #endif
 
+	// The tangent from the first box added in this segment
+	FVector2D StartTangent;
+
 	for (const FTerrainSegment& Segment : TerrainSegments)
 	{
 		for (const FTerrainSpriteStamp& Stamp : Segment.Stamps)
@@ -739,9 +751,17 @@ void UPaperTerrainComponent::SpawnSegments(const TArray<FTerrainSegment>& Terrai
 					int InsertPoint = CollisionPolygonPoints.Num() / 2 - 1;
 					float LengthV0 = FVector2D::Distance(CollisionPolygonPoints[InsertPoint], BoxPoints2D[0]);
 					float LengthV1 = FVector2D::Distance(CollisionPolygonPoints[InsertPoint + 1], BoxPoints2D[3]);
+
+					FVector2D CurrentSegmentTangent = BoxPoints2D[1] - BoxPoints2D[0];
+					CurrentSegmentTangent.Normalize();
+
+					bool bNewSegmentStraightEnough = FVector2D::DotProduct(CurrentSegmentTangent, StartTangent) > FMath::Acos(45.0f);
+
+					// TODO: Arbitray number needs to come from somewhere...
 					float MergeThreshold = 10;
 					bool bMergeIntoPolygon = LengthV0 < MergeThreshold && LengthV1 < MergeThreshold;
-					if (bMergeIntoPolygon)
+
+					if (bNewSegmentStraightEnough && bMergeIntoPolygon)
 					{
 						CollisionPolygonPoints.Insert(BoxPoints2D[2], InsertPoint + 1);
 						CollisionPolygonPoints.Insert(BoxPoints2D[1], InsertPoint + 1);
@@ -754,6 +774,8 @@ void UPaperTerrainComponent::SpawnSegments(const TArray<FTerrainSegment>& Terrai
 						CollisionPolygonPoints.Add(BoxPoints2D[1]);
 						CollisionPolygonPoints.Add(BoxPoints2D[2]);
 						CollisionPolygonPoints.Add(BoxPoints2D[3]);
+						StartTangent = BoxPoints2D[1] - BoxPoints2D[0];
+						StartTangent.Normalize();
 					}
 				}
 				else
@@ -762,6 +784,8 @@ void UPaperTerrainComponent::SpawnSegments(const TArray<FTerrainSegment>& Terrai
 					CollisionPolygonPoints.Add(BoxPoints2D[1]);
 					CollisionPolygonPoints.Add(BoxPoints2D[2]);
 					CollisionPolygonPoints.Add(BoxPoints2D[3]);
+					StartTangent = BoxPoints2D[1] - BoxPoints2D[0];
+					StartTangent.Normalize();
 				}
 #else
 				FKBoxElem Box;
@@ -945,14 +969,25 @@ void UPaperTerrainComponent::InsertConvexCollisionDataFromPolygon(const TArray<F
 		// Simplify polygon
 		TArray<float> EmptyOffsetsList;
 		TArray<FVector2D> LocalPolyVertices = ClosedPolyVertices2D;
-		SimplifyPolygon(LocalPolyVertices, EmptyOffsetsList);
+
+		// The merge / weld threshold should not be any lower / less than half the thickness
+		float PolygonThickness = (ClosedPolyVertices2D[0] - ClosedPolyVertices2D[ClosedPolyVertices2D.Num() - 1]).Size();
+		float SimplifyThreshold = PolygonThickness * 0.5f;
+		SimplifyPolygon(LocalPolyVertices, EmptyOffsetsList, SimplifyThreshold);
 
 		// Always CCW and facing forward regardless of spline winding
 		TArray<FVector2D> CorrectedSplineVertices;
 		PaperGeomTools::CorrectPolygonWinding(CorrectedSplineVertices, LocalPolyVertices, false);
 
 		TArray<FVector2D> TriangulatedPolygonVertices;
-		PaperGeomTools::TriangulatePoly(/*out*/TriangulatedPolygonVertices, CorrectedSplineVertices, false);
+		if (!PaperGeomTools::TriangulatePoly(/*out*/TriangulatedPolygonVertices, CorrectedSplineVertices, false))
+		{
+			// Triangulation failed, try triangulating the original non simplified polygon
+			CorrectedSplineVertices.Empty();
+			PaperGeomTools::CorrectPolygonWinding(/*out*/CorrectedSplineVertices, ClosedPolyVertices2D, false);
+			TriangulatedPolygonVertices.Empty();
+			PaperGeomTools::TriangulatePoly(/*out*/TriangulatedPolygonVertices, CorrectedSplineVertices, false);
+		}
 
 		TArray<TArray<FVector2D>> ConvexHulls;
 		PaperGeomTools::GenerateConvexPolygonsFromTriangles(ConvexHulls, TriangulatedPolygonVertices);
