Tile Set Editor: Fixed an issue where the color picker would immediately close when click-dragging in the color wheel [UE-14616]

diff --git a/Engine/Plugins/2D/Paper2D/Source/Paper2DEditor/Private/TileSetEditor/TileSetDetailsCustomization.cpp b/Engine/Plugins/2D/Paper2D/Source/Paper2DEditor/Private/TileSetEditor/TileSetDetailsCustomization.cpp
--- a/Engine/Plugins/2D/Paper2D/Source/Paper2DEditor/Private/TileSetEditor/TileSetDetailsCustomization.cpp
+++ b/Engine/Plugins/2D/Paper2D/Source/Paper2DEditor/Private/TileSetEditor/TileSetDetailsCustomization.cpp
@@ -164,9 +164,12 @@ FSlateColor FTileSetDetailsCustomization::GetCellDimensionHeaderColor() const
 void FTileSetDetailsCustomization::OnTileIndexChanged(int32 NewIndex, int32 OldIndex)
 {
 	SelectedSingleTileIndex = NewIndex;
-	if (MyDetailLayout != nullptr)
+	if (NewIndex != OldIndex)
 	{
-		MyDetailLayout->ForceRefreshDetails();
+		if (MyDetailLayout != nullptr)
+		{
+			MyDetailLayout->ForceRefreshDetails();
+		}
 	}
 }
 
