Changed EditorUserSettings config classes to EditorPerProjectUserSettings
EditorUserSettings had been renamed just as these classes went in

diff --git a/Engine/Plugins/2D/Paper2D/Source/Paper2DEditor/Private/FlipbookEditor/FlipbookEditorSettings.h b/Engine/Plugins/2D/Paper2D/Source/Paper2DEditor/Private/FlipbookEditor/FlipbookEditorSettings.h
--- a/Engine/Plugins/2D/Paper2D/Source/Paper2DEditor/Private/FlipbookEditor/FlipbookEditorSettings.h
+++ b/Engine/Plugins/2D/Paper2D/Source/Paper2DEditor/Private/FlipbookEditor/FlipbookEditorSettings.h
@@ -5,7 +5,7 @@
 #include "FlipbookEditorSettings.generated.h"
 
 // Settings for the Paper2D sprite editor
-UCLASS(config=EditorUserSettings)
+UCLASS(config=EditorPerProjectUserSettings)
 class UFlipbookEditorSettings : public UObject
 {
 	GENERATED_BODY()
diff --git a/Engine/Plugins/2D/Paper2D/Source/Paper2DEditor/Private/SpriteEditor/SpriteEditorSettings.h b/Engine/Plugins/2D/Paper2D/Source/Paper2DEditor/Private/SpriteEditor/SpriteEditorSettings.h
--- a/Engine/Plugins/2D/Paper2D/Source/Paper2DEditor/Private/SpriteEditor/SpriteEditorSettings.h
+++ b/Engine/Plugins/2D/Paper2D/Source/Paper2DEditor/Private/SpriteEditor/SpriteEditorSettings.h
@@ -5,7 +5,7 @@
 #include "SpriteEditorSettings.generated.h"
 
 // Settings for the Paper2D sprite editor
-UCLASS(config=EditorUserSettings)
+UCLASS(config=EditorPerProjectUserSettings)
 class USpriteEditorSettings : public UObject
 {
 	GENERATED_BODY()
diff --git a/Engine/Plugins/2D/Paper2D/Source/Paper2DEditor/Private/TileMapEditing/TileMapEditorSettings.h b/Engine/Plugins/2D/Paper2D/Source/Paper2DEditor/Private/TileMapEditing/TileMapEditorSettings.h
--- a/Engine/Plugins/2D/Paper2D/Source/Paper2DEditor/Private/TileMapEditing/TileMapEditorSettings.h
+++ b/Engine/Plugins/2D/Paper2D/Source/Paper2DEditor/Private/TileMapEditing/TileMapEditorSettings.h
@@ -5,7 +5,7 @@
 #include "TileMapEditorSettings.generated.h"
 
 // Settings for the Paper2D tile map editor
-UCLASS(config=EditorUserSettings)
+UCLASS(config=EditorPerProjectUserSettings)
 class UTileMapEditorSettings : public UObject
 {
 	GENERATED_BODY()
diff --git a/Engine/Plugins/2D/Paper2D/Source/Paper2DEditor/Private/TileSetEditor/TileSetEditorSettings.h b/Engine/Plugins/2D/Paper2D/Source/Paper2DEditor/Private/TileSetEditor/TileSetEditorSettings.h
--- a/Engine/Plugins/2D/Paper2D/Source/Paper2DEditor/Private/TileSetEditor/TileSetEditorSettings.h
+++ b/Engine/Plugins/2D/Paper2D/Source/Paper2DEditor/Private/TileSetEditor/TileSetEditorSettings.h
@@ -5,7 +5,7 @@
 #include "TileSetEditorSettings.generated.h"
 
 // Settings for the Paper2D tile set editor
-UCLASS(config=EditorUserSettings)
+UCLASS(config=EditorPerProjectUserSettings)
 class UTileSetEditorSettings : public UObject
 {
 	GENERATED_BODY()
