Cursor no longer pops back once you finish painting in a TileMap.

diff --git a/Engine/Plugins/2D/Paper2D/Source/Paper2DEditor/Private/TileMapEditing/EdModeTileMap.cpp b/Engine/Plugins/2D/Paper2D/Source/Paper2DEditor/Private/TileMapEditing/EdModeTileMap.cpp
--- a/Engine/Plugins/2D/Paper2D/Source/Paper2DEditor/Private/TileMapEditing/EdModeTileMap.cpp
+++ b/Engine/Plugins/2D/Paper2D/Source/Paper2DEditor/Private/TileMapEditing/EdModeTileMap.cpp
@@ -343,6 +343,7 @@ bool FEdModeTileMap::InputKey(FEditorViewportClient* InViewportClient, FViewport
 		else if (bWasPainting && !bIsPainting)
 		{
 			// Stopping painting
+			InViewportClient->Viewport->SetPreCaptureMousePosFromSlateCursor();
 		}
 
 		const FViewportCursorLocation Ray = CalculateViewRay(InViewportClient, InViewport);
