Property Editor: Made it possible to force auto-expansion on an IDetailPropertyRow
Tile Set Editor: Switched how the per-tile properties are being customized to avoid staleness when the # of tiles changes
Tile Set Editor: Fixed the display of the dimensions (in tiles) in the header to show the correct number
[UE-14703]

diff --git a/Engine/Plugins/2D/Paper2D/Source/Paper2D/Classes/PaperTileSet.h b/Engine/Plugins/2D/Paper2D/Source/Paper2D/Classes/PaperTileSet.h
--- a/Engine/Plugins/2D/Paper2D/Source/Paper2D/Classes/PaperTileSet.h
+++ b/Engine/Plugins/2D/Paper2D/Source/Paper2D/Classes/PaperTileSet.h
@@ -11,7 +11,7 @@
 struct FPaperTileInfo;
 
 // Information about a single tile in a tile set
-USTRUCT()
+USTRUCT(meta=(ShowOnlyInnerProperties))
 struct PAPER2D_API FPaperTileMetadata
 {
 public:
diff --git a/Engine/Plugins/2D/Paper2D/Source/Paper2DEditor/Private/TileSetEditor/TileSetDetailsCustomization.cpp b/Engine/Plugins/2D/Paper2D/Source/Paper2DEditor/Private/TileSetEditor/TileSetDetailsCustomization.cpp
--- a/Engine/Plugins/2D/Paper2D/Source/Paper2DEditor/Private/TileSetEditor/TileSetDetailsCustomization.cpp
+++ b/Engine/Plugins/2D/Paper2D/Source/Paper2DEditor/Private/TileSetEditor/TileSetDetailsCustomization.cpp
@@ -70,29 +70,18 @@ void FTileSetDetailsCustomization::CustomizeDetails(IDetailLayoutBuilder& Detail
 	{
 		// Hide the array to start with
 		const FName MetadataArrayName = UPaperTileSet::GetPerTilePropertyName();
-		DetailLayout.HideProperty(MetadataArrayName);
+		TSharedPtr<IPropertyHandle> PerTileArrayProperty = DetailLayout.GetProperty(MetadataArrayName);
+		DetailLayout.HideProperty(PerTileArrayProperty);
 
 		if (SelectedSingleTileIndex != INDEX_NONE)
 		{
 			// Customize for the single tile being edited
 			IDetailCategoryBuilder& SingleTileCategory = DetailLayout.EditCategory("SingleTileEditor", FText::GetEmpty());
+			
+			TSharedPtr<IPropertyHandle> OneTileEntry = PerTileArrayProperty->GetChildHandle(SelectedSingleTileIndex);
+			SingleTileCategory.AddProperty(OneTileEntry)
+				.ShouldAutoExpand(true);
 
-			const FString ArrayEntryPathPrefix = FString::Printf(TEXT("%s[%d]."), *MetadataArrayName.ToString(), SelectedSingleTileIndex);
-
-			// Add all of the editable properties in the array entry
-			for (TFieldIterator<UProperty> PropIt(FPaperTileMetadata::StaticStruct()); PropIt; ++PropIt)
-			{
-				UProperty* TestProperty = *PropIt;
-				if (TestProperty->HasAnyPropertyFlags(CPF_Edit))
-				{
-					const bool bAdvancedDisplay = TestProperty->HasAnyPropertyFlags(CPF_AdvancedDisplay);
-					const EPropertyLocation::Type PropertyLocation = bAdvancedDisplay ? EPropertyLocation::Advanced : EPropertyLocation::Common;
-
-					const FString PropertyPath = ArrayEntryPathPrefix + TestProperty->GetName();
-					TSharedRef<IPropertyHandle> PropertyHandle = DetailLayout.GetProperty(*PropertyPath);
-					SingleTileCategory.AddProperty(PropertyHandle, PropertyLocation);
-				}
-			}
 
 			// Add a display of the tile index being edited to the header
 			const FText TileIndexHeaderText = FText::Format(LOCTEXT("SingleTileSectionHeader", "Editing Tile #{0}"), FText::AsNumber(SelectedSingleTileIndex));
@@ -141,7 +130,7 @@ FText FTileSetDetailsCustomization::GetCellDimensionHeaderText() const
 		{
 			const FText TextNumTilesX = FText::AsNumber(NumTilesX, &FNumberFormattingOptions::DefaultNoGrouping());
 			const FText TextNumTilesY = FText::AsNumber(NumTilesY, &FNumberFormattingOptions::DefaultNoGrouping());
-			Result = FText::Format(LOCTEXT("CellDimensions", "{0} x {0} tiles"), TextNumTilesX, TextNumTilesY);
+			Result = FText::Format(LOCTEXT("CellDimensions", "{0} x {1} tiles"), TextNumTilesX, TextNumTilesY);
 		}
 	}
 
diff --git a/Engine/Source/Editor/PropertyEditor/Private/DetailItemNode.cpp b/Engine/Source/Editor/PropertyEditor/Private/DetailItemNode.cpp
--- a/Engine/Source/Editor/PropertyEditor/Private/DetailItemNode.cpp
+++ b/Engine/Source/Editor/PropertyEditor/Private/DetailItemNode.cpp
@@ -52,6 +52,11 @@ void FDetailItemNode::Initialize()
 		InitGroup();
 	}
 
+	if (Customization.PropertyRow.IsValid() && Customization.PropertyRow->GetForceAutoExpansion())
+	{
+		SetExpansionState(true);
+	}
+
 	// Cache the visibility of customizations that can set it
 	if( Customization.HasCustomWidget() )
 	{	
@@ -133,12 +138,17 @@ bool FDetailItemNode::HasMultiColumnWidget() const
 
 void FDetailItemNode::ToggleExpansion()
 {
-	bIsExpanded = !bIsExpanded;
-	
+	SetExpansionState( !bIsExpanded );
+}
+
+void FDetailItemNode::SetExpansionState(bool bWantsExpanded)
+{
+	bIsExpanded = bWantsExpanded;
+
 	// Expand the child after filtering if it wants to be expanded
-	ParentCategory.Pin()->RequestItemExpanded( AsShared(), bIsExpanded );
+	ParentCategory.Pin()->RequestItemExpanded(AsShared(), bIsExpanded);
 
-	OnItemExpansionChanged( bIsExpanded );
+	OnItemExpansionChanged(bIsExpanded);
 }
 
 TSharedRef< ITableRow > FDetailItemNode::GenerateNodeWidget( const TSharedRef<STableViewBase>& OwnerTable, const FDetailColumnSizeData& ColumnSizeData, const TSharedRef<IPropertyUtilities>& PropertyUtilities )
diff --git a/Engine/Source/Editor/PropertyEditor/Private/DetailItemNode.h b/Engine/Source/Editor/PropertyEditor/Private/DetailItemNode.h
--- a/Engine/Source/Editor/PropertyEditor/Private/DetailItemNode.h
+++ b/Engine/Source/Editor/PropertyEditor/Private/DetailItemNode.h
@@ -18,6 +18,8 @@ public:
 
 	void ToggleExpansion();
 
+	void SetExpansionState(bool bWantsExpanded);
+
 	/**
 	 * Generates children for this node
 	 *
diff --git a/Engine/Source/Editor/PropertyEditor/Private/DetailPropertyRow.cpp b/Engine/Source/Editor/PropertyEditor/Private/DetailPropertyRow.cpp
--- a/Engine/Source/Editor/PropertyEditor/Private/DetailPropertyRow.cpp
+++ b/Engine/Source/Editor/PropertyEditor/Private/DetailPropertyRow.cpp
@@ -16,6 +16,7 @@ FDetailPropertyRow::FDetailPropertyRow(TSharedPtr<FPropertyNode> InPropertyNode,
 	, ExternalRootNode( InExternalRootNode )
 	, bShowPropertyButtons( true )
 	, bShowCustomPropertyChildren( true )
+	, bForceAutoExpansion( false )
 {
 	if( InPropertyNode.IsValid() )
 	{
@@ -91,6 +92,12 @@ IDetailPropertyRow& FDetailPropertyRow::IsEnabled( TAttribute<bool> InIsEnabled
 	return *this;
 }
 
+IDetailPropertyRow& FDetailPropertyRow::ShouldAutoExpand(bool bInForceExpansion)
+{
+	bForceAutoExpansion = bInForceExpansion;
+	return *this;
+}
+
 IDetailPropertyRow& FDetailPropertyRow::Visibility( TAttribute<EVisibility> Visibility )
 {
 	PropertyVisibility = Visibility;
@@ -313,6 +320,11 @@ bool FDetailPropertyRow::GetEnabledState() const
 	return Result;
 }
 
+bool FDetailPropertyRow::GetForceAutoExpansion() const
+{
+	return bForceAutoExpansion;
+}
+
 void FDetailPropertyRow::MakeNameWidget( FDetailWidgetRow& Row, const TSharedPtr<FDetailWidgetRow> InCustomRow ) const
 {
 	EVerticalAlignment VerticalAlignment = VAlign_Center;
diff --git a/Engine/Source/Editor/PropertyEditor/Private/DetailPropertyRow.h b/Engine/Source/Editor/PropertyEditor/Private/DetailPropertyRow.h
--- a/Engine/Source/Editor/PropertyEditor/Private/DetailPropertyRow.h
+++ b/Engine/Source/Editor/PropertyEditor/Private/DetailPropertyRow.h
@@ -12,7 +12,8 @@ public:
 	virtual IDetailPropertyRow& ToolTip( const FText& InToolTip ) override;
 	virtual IDetailPropertyRow& ShowPropertyButtons( bool bInShowPropertyButtons ) override;
 	virtual IDetailPropertyRow& EditCondition( TAttribute<bool> EditConditionValue, FOnBooleanValueChanged OnEditConditionValueChanged ) override;
-	virtual IDetailPropertyRow& IsEnabled( TAttribute<bool> InIsEnabled ) override;
+	virtual IDetailPropertyRow& IsEnabled(TAttribute<bool> InIsEnabled) override;
+	virtual IDetailPropertyRow& ShouldAutoExpand(bool bForceExpansion) override;
 	virtual IDetailPropertyRow& Visibility( TAttribute<EVisibility> Visibility ) override;
 	virtual IDetailPropertyRow& OverrideResetToDefault( TAttribute<bool> IsResetToDefaultVisible, FSimpleDelegate OnResetToDefaultClicked ) override;
 	virtual FDetailWidgetRow& CustomWidget( bool bShowChildren = false ) override;
@@ -54,6 +55,11 @@ public:
 	 * @param OutChildren	The list of children created
 	 */
 	void OnGenerateChildren( FDetailNodeList& OutChildren );
+	
+	/**
+	 * @return Whether or not this row wants to force expansion
+	 */
+	bool GetForceAutoExpansion() const;
 
 	/** 
 	 * @return The visibility of this property
@@ -125,4 +131,6 @@ private:
 	bool bShowPropertyButtons;
 	/** True to show custom property children */
 	bool bShowCustomPropertyChildren;
+	/** True to force auto-expansion */
+	bool bForceAutoExpansion;
 };
diff --git a/Engine/Source/Editor/PropertyEditor/Public/IDetailPropertyRow.h b/Engine/Source/Editor/PropertyEditor/Public/IDetailPropertyRow.h
--- a/Engine/Source/Editor/PropertyEditor/Public/IDetailPropertyRow.h
+++ b/Engine/Source/Editor/PropertyEditor/Public/IDetailPropertyRow.h
@@ -51,6 +51,14 @@ public:
 	 */
 	virtual IDetailPropertyRow& IsEnabled( TAttribute<bool> InIsEnabled ) = 0;
 	
+
+	/**
+	 * Sets whether or not this property should auto-expand
+	 *
+	 * @param bForceExpansion	true to force the property to be expanded
+	 */
+	virtual IDetailPropertyRow& ShouldAutoExpand(bool bForceExpansion = true) = 0;
+
 	/**
 	 * Sets the visibility of this property
 	 *
