Tile Maps: Hid the terrain list dropdown button when there are no terrains defined [UE-14684]

diff --git a/Engine/Plugins/2D/Paper2D/Source/Paper2DEditor/Private/TileMapEditing/TileMapEdModeToolkit.cpp b/Engine/Plugins/2D/Paper2D/Source/Paper2DEditor/Private/TileMapEditing/TileMapEdModeToolkit.cpp
--- a/Engine/Plugins/2D/Paper2D/Source/Paper2DEditor/Private/TileMapEditing/TileMapEdModeToolkit.cpp
+++ b/Engine/Plugins/2D/Paper2D/Source/Paper2DEditor/Private/TileMapEditing/TileMapEdModeToolkit.cpp
@@ -302,8 +302,9 @@ TSharedRef<SWidget> FTileMapEdModeToolkit::BuildToolBar() const
 		ToolsToolbar.AddToolBarButton(Commands.SelectTerrainTool);
 
 		//@TODO: TileMapTerrain: Ugly styling
-		FUIAction DummyAction;
-		ToolsToolbar.AddComboButton(DummyAction, FOnGetContent::CreateSP(this, &FTileMapEdModeToolkit::GenerateTerrainMenu));
+		FUIAction TerrainTypeDropdownAction;
+		TerrainTypeDropdownAction.IsActionVisibleDelegate = FIsActionButtonVisible::CreateSP(this, &FTileMapEdModeToolkit::DoesSelectedTileSetHaveTerrains);
+		ToolsToolbar.AddComboButton(TerrainTypeDropdownAction, FOnGetContent::CreateSP(this, &FTileMapEdModeToolkit::GenerateTerrainMenu));
 	}
 
 	return
