Tile Map Editor: Fixed issue where erased tiles were not treated the same as newly created tiles, and changed newly created maps to use INDEX_NONE instead of 0 for their empty cells [UE-14713]

diff --git a/Engine/Plugins/2D/Paper2D/Source/Paper2D/Private/PaperTileLayer.cpp b/Engine/Plugins/2D/Paper2D/Source/Paper2D/Private/PaperTileLayer.cpp
--- a/Engine/Plugins/2D/Paper2D/Source/Paper2D/Private/PaperTileLayer.cpp
+++ b/Engine/Plugins/2D/Paper2D/Source/Paper2D/Private/PaperTileLayer.cpp
@@ -157,7 +157,7 @@ void UPaperTileLayer::DestructiveAllocateMap(int32 NewWidth, int32 NewHeight)
 
 	const int32 NumCells = NewWidth * NewHeight;
 	AllocatedCells.Empty(NumCells);
-	AllocatedCells.AddZeroed(NumCells);
+	AllocatedCells.AddDefaulted(NumCells);
 
 	AllocatedWidth = NewWidth;
 	AllocatedHeight = NewHeight;
diff --git a/Engine/Plugins/2D/Paper2D/Source/Paper2DEditor/Private/TileMapEditing/EdModeTileMap.cpp b/Engine/Plugins/2D/Paper2D/Source/Paper2DEditor/Private/TileMapEditing/EdModeTileMap.cpp
--- a/Engine/Plugins/2D/Paper2D/Source/Paper2DEditor/Private/TileMapEditing/EdModeTileMap.cpp
+++ b/Engine/Plugins/2D/Paper2D/Source/Paper2DEditor/Private/TileMapEditing/EdModeTileMap.cpp
@@ -87,7 +87,9 @@ struct FHorizontalSpan
 		// Go left
 		for (int32 TestX = X0 - 1; TestX >= 0; --TestX)
 		{
-			if ((Layer->GetCell(TestX, Y) == RequiredInk) && !Reach(Layer, Reachability, TestX, Y))
+			const FPaperTileInfo ExistingCell = Layer->GetCell(TestX, Y);
+			const bool bCellMatches = (ExistingCell == RequiredInk) || (!ExistingCell.IsValid() && !RequiredInk.IsValid());
+			if (bCellMatches && !Reach(Layer, Reachability, TestX, Y))
 			{
 				X0 = TestX;
 			}
@@ -100,7 +102,9 @@ struct FHorizontalSpan
 		// Go right
 		for (int32 TestX = X1 + 1; TestX < Layer->LayerWidth; ++TestX)
 		{
-			if ((Layer->GetCell(TestX, Y) == RequiredInk) && !Reach(Layer, Reachability, TestX, Y))
+			const FPaperTileInfo ExistingCell = Layer->GetCell(TestX, Y);
+			const bool bCellMatches = (ExistingCell == RequiredInk) || (!ExistingCell.IsValid() && !RequiredInk.IsValid());
+			if (bCellMatches && !Reach(Layer, Reachability, TestX, Y))
 			{
 				X1 = TestX;
 			}
@@ -980,7 +984,10 @@ bool FEdModeTileMap::FloodFillTiles(const FViewportCursorLocation& Ray)
 				for (int32 X = Span.X0; X <= Span.X1; ++X)
 				{
 					// If it is the right color and not already visited, create a span there
-					if ((TargetLayer->GetCell(X, Y) == RequiredInk) && !FHorizontalSpan::Reach(TargetLayer, TileReachability, X, Y))
+					FPaperTileInfo ExistingCell = TargetLayer->GetCell(X, Y);
+					const bool bCellMatches = (ExistingCell == RequiredInk) || (!ExistingCell.IsValid() && !RequiredInk.IsValid());
+
+					if (bCellMatches && !FHorizontalSpan::Reach(TargetLayer, TileReachability, X, Y))
 					{
 						FHorizontalSpan& NewSpan = *(new (OutstandingSpans) FHorizontalSpan(X, Y));
 						NewSpan.GrowSpan(RequiredInk, TargetLayer, TileReachability);
