Tile Set Editor: Changed the display name of TileSheet to be "Tile Sheet Texture" [UE-14712]

diff --git a/Engine/Plugins/2D/Paper2D/Source/Paper2D/Classes/PaperTileSet.h b/Engine/Plugins/2D/Paper2D/Source/Paper2D/Classes/PaperTileSet.h
--- a/Engine/Plugins/2D/Paper2D/Source/Paper2D/Classes/PaperTileSet.h
+++ b/Engine/Plugins/2D/Paper2D/Source/Paper2D/Classes/PaperTileSet.h
@@ -86,7 +86,7 @@ class PAPER2D_API UPaperTileSet : public UObject
 	int32 TileHeight;
 
 	// The tile sheet texture associated with this tile set
-	UPROPERTY(Category=TileSet, BlueprintReadOnly, EditAnywhere)
+	UPROPERTY(Category=TileSet, BlueprintReadOnly, EditAnywhere, meta=(DisplayName="Tile Sheet Texture"))
 	UTexture2D* TileSheet;
 
 	// The amount of padding around the perimeter of the tile sheet (in pixels)
