Sprite Geometry Editor: Improved the look and feel of the Add Polygon mode
- Now colorizes the text indicating the mode is active
- Polygons can be closed by double-clicking anywhere or by single-clicking on the first vertex when there are at least 3 verts
- There is a preview vertex and connecting lines indicating what will happen when you click, including a + or close cursor annotation
- Ensured winding is always facing outward when the 3rd vertex is first added
- Prevent adding vertices that are nearly equal to the existing ones (0.25f texture-space pixel tolerance right now)
- Selected the shape when the polygon is closed out
[UE-13222]

diff --git a/Engine/Plugins/2D/Paper2D/Source/Paper2DEditor/Private/PaperEditorShared/SpriteGeometryEditMode.cpp b/Engine/Plugins/2D/Paper2D/Source/Paper2DEditor/Private/PaperEditorShared/SpriteGeometryEditMode.cpp
--- a/Engine/Plugins/2D/Paper2D/Source/Paper2DEditor/Private/PaperEditorShared/SpriteGeometryEditMode.cpp
+++ b/Engine/Plugins/2D/Paper2D/Source/Paper2DEditor/Private/PaperEditorShared/SpriteGeometryEditMode.cpp
@@ -189,21 +189,25 @@ bool FSpriteGeometryEditMode::InputKey(FEditorViewportClient* ViewportClient, FV
 	{
 		if (SpriteGeometryHelper.IsAddingPolygon())
 		{
-			if (Key == EKeys::LeftMouseButton && Event == IE_Pressed)
+			if (Key == EKeys::LeftMouseButton)
 			{
 				const int32 HitX = Viewport->GetMouseX();
 				const int32 HitY = Viewport->GetMouseY();
 
-				// Calculate world space positions
+				// Calculate the texture space position of the mouse click
 				FSceneViewFamilyContext ViewFamily(FSceneViewFamily::ConstructionValues(Viewport, ViewportClient->GetScene(), ViewportClient->EngineShowFlags));
 				FSceneView* View = ViewportClient->CalcSceneView(&ViewFamily);
-				FVector WorldPoint = View->PixelToWorld(HitX, HitY, 0);
-
+				const FVector WorldPoint = View->PixelToWorld(HitX, HitY, 0);
 				const FVector2D TexturePoint = SpriteGeometryHelper.GetEditorContext()->WorldSpaceToTextureSpace(WorldPoint);
 
+				// Add or close the polygon (depending on where the click happened and how)
 				const bool bMakeSubtractiveIfAllowed = Viewport->KeyState(EKeys::LeftControl) || Viewport->KeyState(EKeys::RightControl);
-				SpriteGeometryHelper.TEMP_HandleAddPolygonClick(TexturePoint, bMakeSubtractiveIfAllowed);
+				SpriteGeometryHelper.HandleAddPolygonClick(TexturePoint, bMakeSubtractiveIfAllowed, *View, Event);
 			}
+			else if ((Key == EKeys::BackSpace) && (Event == IE_Pressed))
+			{
+				SpriteGeometryHelper.DeleteLastVertexFromAddPolygonMode();
+			}			
 			else if (Key == EKeys::Enter)
 			{
 				SpriteGeometryHelper.ResetAddPolygonMode();
diff --git a/Engine/Plugins/2D/Paper2D/Source/Paper2DEditor/Private/PaperEditorShared/SpriteGeometryEditing.cpp b/Engine/Plugins/2D/Paper2D/Source/Paper2DEditor/Private/PaperEditorShared/SpriteGeometryEditing.cpp
--- a/Engine/Plugins/2D/Paper2D/Source/Paper2DEditor/Private/PaperEditorShared/SpriteGeometryEditing.cpp
+++ b/Engine/Plugins/2D/Paper2D/Source/Paper2DEditor/Private/PaperEditorShared/SpriteGeometryEditing.cpp
@@ -21,6 +21,9 @@ namespace SpriteEditingConstantsEX
 	const float GeometryVertexSize = 8.0f;
 	const float GeometryBorderLineThickness = 2.0f;
 	const FLinearColor GeometrySelectedColor(FLinearColor::White);
+
+	// Add polygon mode
+	const float AddPolygonVertexWeldScreenSpaceDistance = 6.0f;
 }
 
 //////////////////////////////////////////////////////////////////////////
@@ -197,7 +200,10 @@ void FSpriteGeometryEditingHelper::ClearSelectionSet()
 	FSpriteSelectionHelper::ClearSelectionSet();
 	SelectedIDSet.Empty();
 
-	ResetAddPolygonMode();
+	if (bIsAddingPolygon)
+	{
+		ResetAddPolygonMode();
+	}
 
 	EditorContext->InvalidateViewportAndHitProxies();
 }
@@ -404,6 +410,10 @@ void FSpriteGeometryEditingHelper::DrawGeometry_CanvasPass(FViewport& InViewport
 		return;
 	}
 
+	// Calculate the texture-space position of the mouse
+	const FVector MousePositionWorldSpace = View.PixelToWorld(InViewport.GetMouseX(), InViewport.GetMouseY(), 0);
+	const FVector2D MousePositionTextureSpace = EditorContext->WorldSpaceToTextureSpace(MousePositionWorldSpace);
+
 	//@TODO: Display 'action required' text in a bright color for Add Polygon mode
 	//@TODO: Move all of the line drawing to the PDI pass
 	FSpriteGeometryCollection& Geometry = GetGeometryChecked();
@@ -413,7 +423,8 @@ void FSpriteGeometryEditingHelper::DrawGeometry_CanvasPass(FViewport& InViewport
 		static const FText GeomHelpStr = LOCTEXT("GeomEditHelp", "Shift + click to insert a vertex.\nSelect one or more vertices and press Delete to remove them.\nDouble click a vertex to select a polygon\n");
 		static const FText GeomClickAddPolygon_NoSubtractive = LOCTEXT("GeomClickAddPolygon_NoSubtractive", "Click to start creating a polygon\n");
 		static const FText GeomClickAddPolygon_AllowSubtractive = LOCTEXT("GeomClickAddPolygon_AllowSubtractive", "Click to start creating a polygon\nCtrl + Click to start creating a subtractive polygon\n");
-		static const FText GeomAddVerticesHelpStr = LOCTEXT("GeomClickAddVertices", "Click to add points to the polygon\nEnter to finish adding a polygon\nEscape to cancel this polygon\n");
+		static const FText GeomAddVerticesHelpStr = LOCTEXT("GeomClickAddVertices", "Click to add points to the polygon\nDouble-click to add a point and close the shape\nClick again on the first point or press Enter to close the shape\nPress Backspace to remove the last added point or Escape to remove the shape\n");
+		FLinearColor ToolTextColor = FLinearColor::White;
 
 		const FText* HelpStr;
 		if (IsAddingPolygon())
@@ -426,13 +437,14 @@ void FSpriteGeometryEditingHelper::DrawGeometry_CanvasPass(FViewport& InViewport
 			{
 				HelpStr = &GeomAddVerticesHelpStr;
 			}
+			ToolTextColor = FLinearColor::Yellow;
 		}
 		else
 		{
 			HelpStr = &GeomHelpStr;
 		}
 
-		FCanvasTextItem TextItem(FVector2D(6, YPos), *HelpStr, GEngine->GetSmallFont(), FLinearColor::White);
+		FCanvasTextItem TextItem(FVector2D(6, YPos), *HelpStr, GEngine->GetSmallFont(), ToolTextColor);
 		TextItem.EnableShadow(FLinearColor::Black);
 		TextItem.Draw(&Canvas);
 		YPos += 54;
@@ -566,6 +578,63 @@ void FSpriteGeometryEditingHelper::DrawGeometry_CanvasPass(FViewport& InViewport
 			}
 		}
 	}
+
+	// Draw a preview cursor for the add polygon tool
+	if (IsAddingPolygon())
+	{
+		// Figure out where the mouse is back in screen space
+		const FVector2D PotentialVertexScreenPos = TextureSpaceToScreenSpace(View, MousePositionTextureSpace);
+
+		bool bWillCloseByClicking = false;
+		if (Geometry.Shapes.IsValidIndex(AddingPolygonIndex))
+		{
+			const FSpriteGeometryShape& Shape = Geometry.Shapes[AddingPolygonIndex];
+
+			const FLinearColor LineColorRaw = Shape.bNegativeWinding ? NegativeGeometryVertexColor : GeometryVertexColor;
+			const FLinearColor LineColorValidity = Shape.IsShapeValid() ? LineColorRaw : FMath::Lerp(LineColorRaw, FLinearColor::Red, 0.8f);
+			const FLinearColor LineColor = FMath::Lerp(LineColorValidity, SpriteEditingConstantsEX::GeometrySelectedColor, 0.2f);
+
+			if (Shape.Vertices.Num() > 0)
+			{
+				// Draw a line from the last vertex to the potential insertion point for the new one
+				{
+					const FVector2D LastScreenPos = TextureSpaceToScreenSpace(View, Shape.ConvertShapeSpaceToTextureSpace(Shape.Vertices[Shape.Vertices.Num() - 1]));
+
+					FCanvasLineItem LineItem(LastScreenPos, PotentialVertexScreenPos);
+					LineItem.SetColor(LineColor);
+					LineItem.LineThickness = SpriteEditingConstantsEX::GeometryBorderLineThickness;
+					Canvas.DrawItem(LineItem);
+				}
+
+				// And to the first vertex if there were at least 2
+				if (Shape.Vertices.Num() >= 2)
+				{
+					const FVector2D FirstScreenPos = TextureSpaceToScreenSpace(View, Shape.ConvertShapeSpaceToTextureSpace(Shape.Vertices[0]));
+
+					FCanvasLineItem LineItem(PotentialVertexScreenPos, FirstScreenPos);
+					LineItem.SetColor(LineColor);
+					LineItem.LineThickness = SpriteEditingConstantsEX::GeometryBorderLineThickness;
+					Canvas.DrawItem(LineItem);
+
+					// Determine how close we are to the first vertex (will we close the shape by clicking)?
+					bWillCloseByClicking = (Shape.Vertices.Num() >= 3) && (FVector2D::Distance(FirstScreenPos, PotentialVertexScreenPos) < SpriteEditingConstantsEX::AddPolygonVertexWeldScreenSpaceDistance);
+				}
+			}
+		}
+
+		// Draw the prospective vert
+		const float VertSize = SpriteEditingConstantsEX::GeometryVertexSize;
+		Canvas.DrawTile(PotentialVertexScreenPos.X - VertSize*0.5f, PotentialVertexScreenPos.Y - VertSize*0.5f, VertSize, VertSize, 0.f, 0.f, 1.f, 1.f, SpriteEditingConstantsEX::GeometrySelectedColor, GWhiteTexture);
+
+		// Draw a prompt above and to the right of the cursor
+		static const FText CloseButton(LOCTEXT("ClosePolygonPrompt", "Close"));
+		static const FText AddButton(LOCTEXT("AddVertexToPolygonPrompt", "+"));
+
+		const FText PromptText(bWillCloseByClicking ? CloseButton : AddButton);
+		FCanvasTextItem PromptTextItem(FVector2D(PotentialVertexScreenPos.X + VertSize, PotentialVertexScreenPos.Y - VertSize), PromptText, GEngine->GetSmallFont(), FLinearColor::White);
+		PromptTextItem.EnableShadow(FLinearColor::Black);
+		PromptTextItem.Draw(&Canvas);
+	}
 }
 
 void FSpriteGeometryEditingHelper::AddPointToGeometry(const FVector2D& TextureSpacePoint, const int32 SelectedPolygonIndex)
@@ -763,6 +832,19 @@ void FSpriteGeometryEditingHelper::AddNewBoxShape(const FVector2D& BoxLocation,
 	EditorContext->EndTransaction();
 }
 
+void FSpriteGeometryEditingHelper::ResetAddPolygonMode()
+{
+	if (bIsAddingPolygon)
+	{
+		bIsAddingPolygon = false;
+		if (AddingPolygonIndex != INDEX_NONE)
+		{
+			ClearSelectionSet();
+			AddShapeToSelection(AddingPolygonIndex);
+		}
+	}
+}
+
 void FSpriteGeometryEditingHelper::ToggleAddPolygonMode()
 {
 	if (IsAddingPolygon())
@@ -837,36 +919,136 @@ void FSpriteGeometryEditingHelper::SnapAllVerticesToPixelGrid()
 	EditorContext->EndTransaction();
 }
 
-void FSpriteGeometryEditingHelper::TEMP_HandleAddPolygonClick(const FVector2D& TexturePoint, bool bWantsSubtractive)
+void FSpriteGeometryEditingHelper::HandleAddPolygonClick(const FVector2D& TexturePoint, bool bWantsSubtractive, const FSceneView& View, EInputEvent Event)
 {
-	EditorContext->BeginTransaction(LOCTEXT("AddPolygonVertexTransaction", "Add Vertex to Polygon"));
-
 	FSpriteGeometryCollection& Geometry = GetGeometryChecked();
-	if (AddingPolygonIndex == INDEX_NONE)
+
+	// Determine what the action is
+	bool bCloseShape = false;
+	bool bAddVertex = Event == IE_Pressed;
+
+	// When we've already got at least one vertex in the shape, we disallow identical clicks and eventually allow closing the shape via various means
+	if ((AddingPolygonIndex != INDEX_NONE) && Geometry.Shapes.IsValidIndex(AddingPolygonIndex))
+	{
+		const FSpriteGeometryShape& Shape = Geometry.Shapes[AddingPolygonIndex];
+
+		// See if we're allowed to close the shape yet
+		if (Shape.Vertices.Num() >= 3)
+		{
+			if (Event == IE_DoubleClick)
+			{
+				// Double-clicking with enough verts to form a polygon will close the shape
+				bCloseShape = true;
+			}
+			else if (Event == IE_Pressed)
+			{
+				// Clicking on the starting vertex if we're close enough and have enough points will finish the shape
+				const FVector2D ClickScreenPos = TextureSpaceToScreenSpace(View, TexturePoint);
+				const FVector2D StartingScreenPos = TextureSpaceToScreenSpace(View, Shape.ConvertShapeSpaceToTextureSpace(Shape.Vertices[0]));
+				bCloseShape = FVector2D::Distance(StartingScreenPos, ClickScreenPos) < SpriteEditingConstantsEX::AddPolygonVertexWeldScreenSpaceDistance;
+			}
+		}
+
+		// Prevent adding if we're closing a shape
+		if (bCloseShape)
+		{
+			bAddVertex = false;
+		}
+
+		// Prevent adding if we're really close to an existing vertex in the shape
+		if (bAddVertex)
+		{
+			for (const FVector2D& ExistingShapeSpaceVertex : Shape.Vertices)
+			{
+				const FVector2D ExistingTextureSpaceVertex = Shape.ConvertShapeSpaceToTextureSpace(ExistingShapeSpaceVertex);
+				const float DistanceToExisting = FVector2D::Distance(ExistingTextureSpaceVertex, TexturePoint);
+				if (DistanceToExisting < 0.25f)
+				{
+					bAddVertex = false;
+					break;
+				}
+			}
+		}
+	}
+
+	if (bCloseShape)
+	{
+		ResetAddPolygonMode();
+	}
+	else if (bAddVertex)
 	{
-		FSpriteGeometryShape NewPolygon;
-		NewPolygon.ShapeType = ESpriteShapeType::Polygon;
-		NewPolygon.bNegativeWinding = CanAddSubtractivePolygon() && bWantsSubtractive;
+		EditorContext->BeginTransaction(LOCTEXT("AddPolygonVertexTransaction", "Add Vertex to Polygon"));
+
+		if (AddingPolygonIndex == INDEX_NONE)
+		{
+			FSpriteGeometryShape NewPolygon;
+			NewPolygon.ShapeType = ESpriteShapeType::Polygon;
+			NewPolygon.bNegativeWinding = CanAddSubtractivePolygon() && bWantsSubtractive;
+
+			Geometry.Shapes.Add(NewPolygon);
+			AddingPolygonIndex = Geometry.Shapes.Num() - 1;
+		}
+
+		if (Geometry.Shapes.IsValidIndex(AddingPolygonIndex))
+		{
+			FSpriteGeometryShape& Shape = Geometry.Shapes[AddingPolygonIndex];
+			Shape.ShapeType = ESpriteShapeType::Polygon;
+			Shape.Vertices.Add(Shape.ConvertTextureSpaceToShapeSpace(TexturePoint));
+
+			// Reorder the vertices when a triangle is first made to make sure the winding is facing outwards
+			// After that it is up to the user to add verts in the order they want
+			if (Shape.Vertices.Num() == 3)
+			{
+				FVector2D& A = Shape.Vertices[0];
+				FVector2D& B = Shape.Vertices[1];
+				FVector2D& C = Shape.Vertices[2];
+
+				if (FVector2D::CrossProduct(B - A, C - A) < 0.0f)
+				{
+					Swap(B, C);
+				}
+			}
+		}
+		else
+		{
+			ResetAddPolygonMode();
+		}
 
-		Geometry.Shapes.Add(NewPolygon);
-		AddingPolygonIndex = Geometry.Shapes.Num() - 1;
+		Geometry.GeometryType = ESpritePolygonMode::FullyCustom;
+
+		EditorContext->MarkTransactionAsDirty();
+		EditorContext->EndTransaction();
 	}
+}
 
-	if (Geometry.Shapes.IsValidIndex(AddingPolygonIndex))
+void FSpriteGeometryEditingHelper::DeleteLastVertexFromAddPolygonMode()
+{
+	check(bIsAddingPolygon);
+
+	FSpriteGeometryCollection& Geometry = GetGeometryChecked();
+
+	if ((AddingPolygonIndex != INDEX_NONE) && Geometry.Shapes.IsValidIndex(AddingPolygonIndex))
 	{
+		EditorContext->BeginTransaction(LOCTEXT("DeleteLastAddedPoint", "Delete last added point"));
+
 		FSpriteGeometryShape& Shape = Geometry.Shapes[AddingPolygonIndex];
-		Shape.ShapeType = ESpriteShapeType::Polygon;
-		Shape.Vertices.Add(Shape.ConvertTextureSpaceToShapeSpace(TexturePoint));
+		if (Shape.Vertices.Num() > 0)
+		{
+			Shape.Vertices.Pop();
+		}
+		else
+		{
+			Geometry.Shapes.RemoveAt(AddingPolygonIndex);
+			AddingPolygonIndex = INDEX_NONE;
+		}
+
+		EditorContext->MarkTransactionAsDirty();
+		EditorContext->EndTransaction();
 	}
 	else
 	{
 		ResetAddPolygonMode();
 	}
-
-	Geometry.GeometryType = ESpritePolygonMode::FullyCustom;
-
-	EditorContext->MarkTransactionAsDirty();
-	EditorContext->EndTransaction();
 }
 
 bool FSpriteGeometryEditingHelper::ClosestPointOnLine(const FVector2D& Point, const FVector2D& LineStart, const FVector2D& LineEnd, FVector2D& OutClosestPoint)
diff --git a/Engine/Plugins/2D/Paper2D/Source/Paper2DEditor/Private/PaperEditorShared/SpriteGeometryEditing.h b/Engine/Plugins/2D/Paper2D/Source/Paper2DEditor/Private/PaperEditorShared/SpriteGeometryEditing.h
--- a/Engine/Plugins/2D/Paper2D/Source/Paper2DEditor/Private/PaperEditorShared/SpriteGeometryEditing.h
+++ b/Engine/Plugins/2D/Paper2D/Source/Paper2DEditor/Private/PaperEditorShared/SpriteGeometryEditing.h
@@ -120,7 +120,7 @@ public:
 	void AddNewBoxShape(const FVector2D& BoxLocation, const FVector2D& BoxSize);
 	bool CanAddBoxShape() const { return GeometryBeingEdited != nullptr; }
 
-	void ResetAddPolygonMode() { bIsAddingPolygon = false; }
+	void ResetAddPolygonMode();
 	void ToggleAddPolygonMode();
 	bool IsAddingPolygon() const { return bIsAddingPolygon; }
 	bool CanAddPolygon() const { return GeometryBeingEdited != nullptr; }
@@ -130,8 +130,9 @@ public:
 	void SnapAllVerticesToPixelGrid();
 	bool CanSnapVerticesToPixelGrid() const { return GeometryBeingEdited != nullptr; }
 
-	//@TODO: The code calling this will eventually be folded in here too
-	void TEMP_HandleAddPolygonClick(const FVector2D& TexturePoint, bool bWantsSubtractive);
+	void HandleAddPolygonClick(const FVector2D& TexturePoint, bool bWantsSubtractive, const FSceneView& View, EInputEvent Event);
+
+	void DeleteLastVertexFromAddPolygonMode();
 
 	bool IsEditingGeometry() const
 	{
