Sprite Geometry Editing: Changed centroid for polygon shapes to be a bounding box center rather than a weighted average [UE-14721]

diff --git a/Engine/Plugins/2D/Paper2D/Source/Paper2D/Classes/SpriteEditorOnlyTypes.h b/Engine/Plugins/2D/Paper2D/Source/Paper2D/Classes/SpriteEditorOnlyTypes.h
--- a/Engine/Plugins/2D/Paper2D/Source/Paper2D/Classes/SpriteEditorOnlyTypes.h
+++ b/Engine/Plugins/2D/Paper2D/Source/Paper2D/Classes/SpriteEditorOnlyTypes.h
@@ -96,16 +96,18 @@ public:
 	// Returns the polygon centroid in texture space
 	FVector2D GetPolygonCentroid() const
 	{
+		FBox2D Bounds(ForceInit);
+
 		FVector2D LocalSpaceResult = FVector2D::ZeroVector;
 
 		for (const FVector2D& Vertex : Vertices)
 		{
-			LocalSpaceResult += Vertex;
+			Bounds += Vertex;
 		}
 
 		if (Vertices.Num() > 0)
 		{
-			LocalSpaceResult *= 1.0f / Vertices.Num();
+			LocalSpaceResult = Bounds.GetCenter();
 		}
 
 		return ConvertShapeSpaceToTextureSpace(LocalSpaceResult);
