Paper2D: Refactored UPaperSprite::InitSprite and FSpriteAssetInitParameters to move texture analysis and default settings out of runtime code into editor code, and changed default pixels/uu for new assets to 1.0 [UE-14452]

Upgrade nodes:
- All default material analysis settings and the default pixels/uu setting have been moved from UPaperRuntimeSettings to UPaperImporterSetings
- The default pixels/uu setting for newly created assets is now 1.0 instead of 2.56 (old value in engine) or 1.28 (old value in 2D templates).  It can be edited in Editor Settings .. Paper 2D Import
- Old default settings are not carried across, any custom settings for materials or pixels/uu will need to be reapplied; but this change will not affect any existing assets

diff --git a/Engine/Plugins/2D/Paper2D/Source/Paper2D/Classes/PaperRuntimeSettings.h b/Engine/Plugins/2D/Paper2D/Source/Paper2D/Classes/PaperRuntimeSettings.h
--- a/Engine/Plugins/2D/Paper2D/Source/Paper2D/Classes/PaperRuntimeSettings.h
+++ b/Engine/Plugins/2D/Paper2D/Source/Paper2D/Classes/PaperRuntimeSettings.h
@@ -7,33 +7,11 @@
 /**
  * Implements the settings for the Paper2D plugin.
  */
-UCLASS(config = Engine, defaultconfig)
+UCLASS(config=Engine, defaultconfig)
 class PAPER2D_API UPaperRuntimeSettings : public UObject
 {
 	GENERATED_UCLASS_BODY()
 
-	// The default scaling factor between pixels and Unreal units (cm) to use for newly created sprite assets (e.g., 0.64 would make a 64 pixel wide sprite take up 100 cm)
-	UPROPERTY(GlobalConfig, EditAnywhere, Category=Settings)
-	float DefaultPixelsPerUnrealUnit;
-
-#if WITH_EDITORONLY_DATA
-	// Should the source texture be scanned when creating new sprites to determine the appropriate material? (if false, the Default Masked Material is always used)
-	UPROPERTY(config, EditAnywhere, Category=Settings)
-	bool bPickBestMaterialWhenCreatingSprite;
-
-	// The default masked material for newly created sprites (masked means binary opacity: things are either opaque or see-thru, with nothing in between)
-	UPROPERTY(config, EditAnywhere, Category=Settings, meta=(AllowedClasses="MaterialInterface", DisplayName="Default Masked Material"))
-	FStringAssetReference DefaultMaskedMaterialName;
-
-	// The default translucent material for newly created sprites (translucent means smooth opacity which can vary continuously from 0..1, but translucent rendering is more expensive that opaque or masked rendering and has different sorting rules)
-	UPROPERTY(config, EditAnywhere, Category=Settings, meta=(AllowedClasses="MaterialInterface", DisplayName="Default Translucent Material", EditCondition="bPickBestMaterialWhenCreatingSprite"))
-	FStringAssetReference DefaultTranslucentMaterialName;
-
-	// The default opaque material for newly created sprites
-	UPROPERTY(config, EditAnywhere, Category=Settings, meta=(AllowedClasses="MaterialInterface", DisplayName="Default Opaque Sprite Material"))
-	FStringAssetReference DefaultOpaqueMaterialName;
-#endif
-
 	// Enables experimental *incomplete and unsupported* texture atlas groups that sprites can be assigned to
 	UPROPERTY(EditAnywhere, config, Category=Experimental)
 	bool bEnableSpriteAtlasGroups;
@@ -41,4 +19,8 @@ class PAPER2D_API UPaperRuntimeSettings : public UObject
 	// Enables experimental *incomplete and unsupported* 2D terrain spline editing. Note: You need to restart the editor when enabling this setting for the change to fully take effect.
 	UPROPERTY(EditAnywhere, config, Category=Experimental, meta=(ConfigRestartRequired=true))
 	bool bEnableTerrainSplineEditing;
+
+	// Enables automatic resizing of various sprite data that is authored in texture space if the source texture gets resized (sockets, the pivot, render and collision geometry, etc...)
+	UPROPERTY(EditAnywhere, config, Category=Settings)
+	bool bResizeSpriteDataToMatchTextures;
 };
diff --git a/Engine/Plugins/2D/Paper2D/Source/Paper2D/Classes/SpriteEditorOnlyTypes.h b/Engine/Plugins/2D/Paper2D/Source/Paper2D/Classes/SpriteEditorOnlyTypes.h
--- a/Engine/Plugins/2D/Paper2D/Source/Paper2D/Classes/SpriteEditorOnlyTypes.h
+++ b/Engine/Plugins/2D/Paper2D/Source/Paper2D/Classes/SpriteEditorOnlyTypes.h
@@ -258,10 +258,10 @@ struct FSpriteAssetInitParameters
 		: Texture(nullptr)
 		, Offset(FVector2D::ZeroVector)
 		, Dimension(FVector2D::ZeroVector)
-		, bNewlyCreated(false)
-		, OpaqueMaterialOverride(nullptr)
-		, TranslucentMaterialOverride(nullptr)
-		, MaskedMaterialOverride(nullptr)
+		, bOverridePixelsPerUnrealUnit(false)
+		, PixelsPerUnrealUnit(1.0f)
+		, DefaultMaterialOverride(nullptr)
+		, AlternateMaterialOverride(nullptr)
 	{
 	}
 
@@ -281,6 +281,11 @@ struct FSpriteAssetInitParameters
 		}
 	}
 
+	void SetPixelsPerUnrealUnit(float NewPixelsPerUU)
+	{
+		bOverridePixelsPerUnrealUnit = true;
+		PixelsPerUnrealUnit = NewPixelsPerUU;
+	}
 public:
 	// The texture to use
 	UTexture2D* Texture;
@@ -294,17 +299,17 @@ public:
 	// The dimension of the subregion within the texture (in pixels)
 	FVector2D Dimension;
 
-	// Is this sprite newly created (should we pull the default pixels/uu and materials from the project settings)?
-	bool bNewlyCreated;
+	// Should we apply the PixelsPerUnrealUnit, or leave the value as it was
+	bool bOverridePixelsPerUnrealUnit;
 
-	// The material to use for opaque sprites if set
-	UMaterialInterface* OpaqueMaterialOverride;
+	// The scaling factor between pixels and Unreal units (cm) to use (e.g., 0.64 would make a 64 pixel wide sprite take up 100 cm)
+	float PixelsPerUnrealUnit;
 
-	// The material to use for translucent sprites if set
-	UMaterialInterface* TranslucentMaterialOverride;
+	// The material to override the default value with, if set
+	UMaterialInterface* DefaultMaterialOverride;
 
-	// The material to use for masked sprites if set
-	UMaterialInterface* MaskedMaterialOverride;
+	// The material to override the alternate (opaque) value with, if set
+	UMaterialInterface* AlternateMaterialOverride;
 };
 
 UENUM()
diff --git a/Engine/Plugins/2D/Paper2D/Source/Paper2D/Private/AlphaBitmap.h b/Engine/Plugins/2D/Paper2D/Source/Paper2D/Private/AlphaBitmap.h
deleted file mode 100644
diff --git a/Engine/Plugins/2D/Paper2D/Source/Paper2D/Private/PaperRuntimeSettings.cpp b/Engine/Plugins/2D/Paper2D/Source/Paper2D/Private/PaperRuntimeSettings.cpp
--- a/Engine/Plugins/2D/Paper2D/Source/Paper2D/Private/PaperRuntimeSettings.cpp
+++ b/Engine/Plugins/2D/Paper2D/Source/Paper2D/Private/PaperRuntimeSettings.cpp
@@ -8,17 +8,8 @@
 
 UPaperRuntimeSettings::UPaperRuntimeSettings(const FObjectInitializer& ObjectInitializer)
 	: Super(ObjectInitializer)
-	, DefaultPixelsPerUnrealUnit(2.56f)
 	, bEnableSpriteAtlasGroups(false)
 	, bEnableTerrainSplineEditing(false)
+	, bResizeSpriteDataToMatchTextures(true)
 {
-#if WITH_EDITORONLY_DATA
-	if (!IsRunningCommandlet())
-	{
-		bPickBestMaterialWhenCreatingSprite = true;
-		DefaultMaskedMaterialName = FStringAssetReference("/Paper2D/MaskedUnlitSpriteMaterial.MaskedUnlitSpriteMaterial");
-		DefaultTranslucentMaterialName = FStringAssetReference("/Paper2D/TranslucentUnlitSpriteMaterial.TranslucentUnlitSpriteMaterial");
-		DefaultOpaqueMaterialName = FStringAssetReference("/Paper2D/OpaqueUnlitSpriteMaterial.OpaqueUnlitSpriteMaterial");
-	}
-#endif
 }
diff --git a/Engine/Plugins/2D/Paper2D/Source/Paper2D/Private/PaperSprite.cpp b/Engine/Plugins/2D/Paper2D/Source/Paper2D/Private/PaperSprite.cpp
--- a/Engine/Plugins/2D/Paper2D/Source/Paper2D/Private/PaperSprite.cpp
+++ b/Engine/Plugins/2D/Paper2D/Source/Paper2D/Private/PaperSprite.cpp
@@ -7,7 +7,6 @@
 
 #include "PaperCustomVersion.h"
 #include "PaperGeomTools.h"
-#include "PaperRuntimeSettings.h"
 #include "PaperSpriteComponent.h"
 #include "PaperFlipbookComponent.h"
 #include "SpriteDrawCall.h"
@@ -22,6 +21,7 @@
 #include "GeomTools.h"
 #include "BitmapUtils.h"
 #include "ComponentReregisterContext.h"
+#include "PaperRuntimeSettings.h"
 
 //////////////////////////////////////////////////////////////////////////
 // maf
@@ -375,7 +375,7 @@ void UPaperSprite::OnObjectReimported(UObject* InObject)
 	// If SourceTetxureDimension == 0, we don't have a previous dimension to work off, so can't
 	// rescale sensibly
 	UTexture2D* Texture = Cast<UTexture2D>(InObject);
-	if (Texture != nullptr && Texture == GetSourceTexture())
+	if ((Texture != nullptr) && (Texture == GetSourceTexture()))
 	{
 		if (NeedRescaleSpriteData())
 		{
@@ -691,12 +691,17 @@ void UPaperSprite::RescaleSpriteData(UTexture2D* Texture)
 
 bool UPaperSprite::NeedRescaleSpriteData()
 {
-	if (UTexture2D* Texture = GetSourceTexture())
+	const bool bSupportsRescaling = GetDefault<UPaperRuntimeSettings>()->bResizeSpriteDataToMatchTextures;
+
+	if (bSupportsRescaling)
 	{
-		Texture->ConditionalPostLoad();
-		const FIntPoint TextureSize = Texture->GetImportedSize();
-		const bool bTextureSizeIsZero = (TextureSize.X == 0) || (TextureSize.Y == 0);
-		return !SourceTextureDimension.IsZero() && !bTextureSizeIsZero && ((TextureSize.X != SourceTextureDimension.X) || (TextureSize.Y != SourceTextureDimension.Y));
+		if (UTexture2D* Texture = GetSourceTexture())
+		{
+			Texture->ConditionalPostLoad();
+			const FIntPoint TextureSize = Texture->GetImportedSize();
+			const bool bTextureSizeIsZero = (TextureSize.X == 0) || (TextureSize.Y == 0);
+			return !SourceTextureDimension.IsZero() && !bTextureSizeIsZero && ((TextureSize.X != SourceTextureDimension.X) || (TextureSize.Y != SourceTextureDimension.Y));
+		}
 	}
 
 	return false;
@@ -1351,58 +1356,19 @@ void UPaperSprite::ExtractRectsFromTexture(UTexture2D* Texture, TArray<FIntRect>
 
 void UPaperSprite::InitializeSprite(const FSpriteAssetInitParameters& InitParams)
 {
-	if (InitParams.bNewlyCreated)
+	if (InitParams.bOverridePixelsPerUnrealUnit)
 	{
-		const UPaperRuntimeSettings* DefaultSettings = GetDefault<UPaperRuntimeSettings>();
-		PixelsPerUnrealUnit = DefaultSettings->DefaultPixelsPerUnrealUnit;
-
-		bool bUseMaskedTexture = true;
-
-		// Analyze the texture if desired (to see if it's got greyscale alpha or just binary alpha, picking either a translucent or masked material)
-		if (DefaultSettings->bPickBestMaterialWhenCreatingSprite)
-		{
-			if (InitParams.Texture != nullptr)
-			{
-				FAlphaBitmap AlphaBitmap(InitParams.Texture);
-				bool bHasIntermediateValues;
-				bool bHasZeros;
-				AlphaBitmap.AnalyzeImage((int32)InitParams.Offset.X, (int32)InitParams.Offset.Y, (int32)InitParams.Dimension.X, (int32)InitParams.Dimension.Y, /*out*/ bHasZeros, /*out*/ bHasIntermediateValues);
-
-				bUseMaskedTexture = !bHasIntermediateValues;
-			}
-		}
+		PixelsPerUnrealUnit = InitParams.PixelsPerUnrealUnit;
+	}
 
-		if (bUseMaskedTexture)
-		{
-			if (InitParams.MaskedMaterialOverride != nullptr)
-			{
-				DefaultMaterial = InitParams.MaskedMaterialOverride;
-			}
-			else if (UMaterialInterface* MaskedMaterial = Cast<UMaterialInterface>(DefaultSettings->DefaultMaskedMaterialName.TryLoad()))
-			{
-				DefaultMaterial = MaskedMaterial;
-			}
-		}
-		else
-		{
-			if (InitParams.TranslucentMaterialOverride != nullptr)
-			{
-				DefaultMaterial = InitParams.TranslucentMaterialOverride;
-			}
-			else if (UMaterialInterface* TranslucentMaterial = Cast<UMaterialInterface>(DefaultSettings->DefaultTranslucentMaterialName.TryLoad()))
-			{
-				DefaultMaterial = TranslucentMaterial;
-			}
-		}
+	if (InitParams.DefaultMaterialOverride != nullptr)
+	{
+		DefaultMaterial = InitParams.DefaultMaterialOverride;
+	}
 
-		if (InitParams.OpaqueMaterialOverride != nullptr)
-		{
-			AlternateMaterial = InitParams.OpaqueMaterialOverride;
-		}
-		else if (UMaterialInterface* OpaqueMaterial = Cast<UMaterialInterface>(DefaultSettings->DefaultOpaqueMaterialName.TryLoad()))
-		{
-			AlternateMaterial = OpaqueMaterial;
-		}
+	if (InitParams.AlternateMaterialOverride != nullptr)
+	{
+		AlternateMaterial = InitParams.AlternateMaterialOverride;
 	}
 
 	SourceTexture = InitParams.Texture;
diff --git a/Engine/Plugins/2D/Paper2D/Source/Paper2D/Private/PaperTileMap.cpp b/Engine/Plugins/2D/Paper2D/Source/Paper2D/Private/PaperTileMap.cpp
--- a/Engine/Plugins/2D/Paper2D/Source/Paper2D/Private/PaperTileMap.cpp
+++ b/Engine/Plugins/2D/Paper2D/Source/Paper2D/Private/PaperTileMap.cpp
@@ -443,9 +443,8 @@ void UPaperTileMap::ResizeMap(int32 NewWidth, int32 NewHeight, bool bForceResize
 
 void UPaperTileMap::InitializeNewEmptyTileMap()
 {
+	//@TODO: Consider removing this function / making it a 'prepare' method just like the sprite one on the importer settings
 	AddNewLayer();
-
-	PixelsPerUnrealUnit = GetDefault<UPaperRuntimeSettings>()->DefaultPixelsPerUnrealUnit;
 }
 
 UPaperTileMap* UPaperTileMap::CloneTileMap(UObject* OuterForClone)
diff --git a/Engine/Plugins/2D/Paper2D/Source/Paper2D/Public/AlphaBitmap.h b/Engine/Plugins/2D/Paper2D/Source/Paper2D/Public/AlphaBitmap.h
new file mode 100644
--- /dev/null
+++ b/Engine/Plugins/2D/Paper2D/Source/Paper2D/Public/AlphaBitmap.h
@@ -0,0 +1,359 @@
+// Copyright 1998-2015 Epic Games, Inc. All Rights Reserved.
+
+#pragma once
+
+//////////////////////////////////////////////////////////////////////////
+// FAlphaBitmap
+
+struct PAPER2D_API FAlphaBitmap
+{
+	FAlphaBitmap(UTexture2D* SourceTexture, uint8 InDefaultValue = 0)
+		: Width(0)
+		, Height(0)
+		, DefaultValue(InDefaultValue)
+	{
+		ExtractFromTexture(SourceTexture);
+	}
+
+	void ExtractFromTexture(UTexture2D* SourceTexture)
+	{
+		// use the source art if it exists
+		FTextureSource* TextureSource = nullptr;
+		if ((SourceTexture != nullptr) && SourceTexture->Source.IsValid())
+		{
+			switch (SourceTexture->Source.GetFormat())
+			{
+			case TSF_G8:
+			case TSF_BGRA8:
+				TextureSource = &(SourceTexture->Source);
+				break;
+			default:
+				break;
+			};
+		}
+
+		if (TextureSource != nullptr)
+		{
+			TArray<uint8> TextureRawData;
+			TextureSource->GetMipData(TextureRawData, 0);
+			int32 BytesPerPixel = TextureSource->GetBytesPerPixel();
+			ETextureSourceFormat PixelFormat = TextureSource->GetFormat();
+			Width = TextureSource->GetSizeX();
+			Height = TextureSource->GetSizeY();
+			RawData.SetNumZeroed(Width * Height);
+
+			for (int Y = 0; Y < Height; ++Y)
+			{
+				for (int X = 0; X < Width; ++X)
+				{
+					int32 PixelByteOffset = (X + Y * Width) * BytesPerPixel;
+					const uint8* PixelPtr = TextureRawData.GetData() + PixelByteOffset;
+					FColor Color;
+					if (PixelFormat == TSF_BGRA8)
+					{
+						Color = *((FColor*)PixelPtr);
+					}
+					else
+					{
+						checkSlow(PixelFormat == TSF_G8);
+						const uint8 Intensity = *PixelPtr;
+						Color = FColor(Intensity, Intensity, Intensity, Intensity);
+					}
+
+					RawData[Y * Width + X] = Color.A;
+				}
+			}
+		}
+	}
+
+	// Create an empty bitmap
+	FAlphaBitmap(int InWidth, int InHeight, int InDefaultValue = 0)
+		: Width(InWidth)
+		, Height(InHeight)
+		, DefaultValue(InDefaultValue)
+	{
+		RawData.SetNumZeroed(Width * Height);
+		ClearToDefaultValue();
+	}
+
+	bool IsValid() const
+	{
+		return (Width != 0) && (Height != 0);
+	}
+
+	uint8 GetPixel(int X, int Y) const
+	{
+		if ((X >= 0) && (X < Width) && (Y >= 0) && (Y < Height))
+		{
+			return RawData[Y * Width + X];
+		}
+		else
+		{
+			return DefaultValue;
+		}
+	}
+
+	void SetPixel(int X, int Y, uint8 Value)
+	{
+		if ((X >= 0) && (X < Width) && (Y >= 0) && (Y < Height))
+		{
+			RawData[Y * Width + X] = Value;
+		}
+	}
+
+	void ClearToDefaultValue()
+	{
+		int RawDataCount = Width * Height;
+		for (int RawDataIndex = 0; RawDataIndex < RawDataCount; ++RawDataIndex)
+		{
+			RawData[RawDataIndex] = DefaultValue;
+		}
+	}
+
+	bool IsColumnEqual(int X, int Y0, int Y1, int Target) const
+	{
+		for (int Y = Y0; Y <= Y1; ++Y)
+		{
+			if (GetPixel(X, Y) != Target)
+			{
+				return false;
+			}
+		}
+		return true;
+	}
+
+	bool IsRowEqual(int X0, int X1, int Y, int Target) const
+	{
+		for (int X = X0; X <= X1; ++X)
+		{
+			if (GetPixel(X, Y) != Target)
+			{
+				return false;
+			}
+		}
+		return true;
+	}
+
+	// Is the rectangle empty (X0..X1, Y0..Y1 inclusive)
+	bool IsRegionEqual(int32 X0, int32 Y0, int32 X1, int32 Y1, int32 Target)
+	{
+		bool bEqual = true;
+		for (int32 Y = Y0; (Y < Y1) && bEqual; ++Y)
+		{
+			bEqual = bEqual && IsRowEqual(X0, X1, Y, Target);
+		}
+		return bEqual;
+	}
+
+	bool IsColumnEmpty(int X, int Y0, int Y1) const
+	{
+		return IsColumnEqual(X, Y0, Y1, 0);
+	}
+
+	bool IsRowEmpty(int X0, int X1, int Y)
+	{
+		return IsRowEqual(X0, X1, Y, 0);
+	}
+
+	// Is the rectangle empty (X0..X1, Y0..Y1 inclusive)
+	bool IsRegionEmpty(int32 X0, int32 Y0, int32 X1, int32 Y1)
+	{
+		return IsRegionEqual(X0, Y0, X1, Y1, 0);
+	}
+
+	// Returns the tight bounding box around pixels that are not 0
+	void GetTightBounds(FIntPoint& OutOrigin, FIntPoint& OutDimension)
+	{
+		OutOrigin.X = 0;
+		OutOrigin.Y = 0;
+		OutDimension.X = Width;
+		OutDimension.Y = Height;
+		TightenBounds(OutOrigin, OutDimension);
+	}
+
+	// Returns the tight bounding box around pixels that are not 0
+	void TightenBounds(FIntPoint& InOutOrigin, FIntPoint& InOutDimension)
+	{
+		int Top = InOutOrigin.Y;
+		int Bottom = InOutOrigin.Y + InOutDimension.Y - 1;
+		int Left = InOutOrigin.X;
+		int Right = InOutOrigin.X + InOutDimension.X - 1;
+
+		while (Top < Bottom && IsRowEmpty(Left, Right, Top))
+		{
+			++Top;
+		}
+		while (Bottom >= Top && IsRowEmpty(Left, Right, Bottom))
+		{
+			--Bottom;
+		}
+		while (Left < Right && IsColumnEmpty(Left, Top, Bottom))
+		{
+			++Left;
+		}
+		while (Right >= Left && IsColumnEmpty(Right, Top, Bottom))
+		{
+			--Right;
+		}
+
+		InOutOrigin.X = Left;
+		InOutOrigin.Y = Top;
+		InOutDimension.X = Right - Left + 1;
+		InOutDimension.Y = Bottom - Top + 1;
+	}
+
+	// Sets pixels to 1 marking the rectangle outline
+	void DrawRectOutline(int StartX, int StartY, int InWidth, int InHeight)
+	{
+		const int X0 = StartX;
+		const int Y0 = StartY;
+		const int X1 = StartX + InWidth - 1;
+		const int Y1 = StartY + InHeight - 1;
+		for (int Y = Y0; Y <= Y1; ++Y)
+		{
+			SetPixel(X0, Y, 1);
+			SetPixel(X1, Y, 1);
+		}
+		for (int X = X0; X <= X1; ++X)
+		{
+			SetPixel(X, Y0, 1);
+			SetPixel(X, Y1, 1);
+		}
+	}
+
+	void FillRect(int StartX, int StartY, int InWidth, int InHeight)
+	{
+		const int X0 = StartX;
+		const int Y0 = StartY;
+		const int X1 = StartX + InWidth - 1;
+		const int Y1 = StartY + InHeight - 1;
+		for (int Y = Y0; Y <= Y1; ++Y)
+		{
+			for (int X = X0; X <= X1; ++X)
+			{
+				SetPixel(X, Y, 1);
+			}
+		}
+	}
+
+	// Wind through the bitmap from StartX, StartY to find the closest hit point
+	bool FoundClosestValidPoint(int StartX, int StartY, const int MaxAllowedSearchDistance, FIntPoint& OutHitPoint)
+	{
+		// Constrain the point within the image bounds
+		StartX = FMath::Clamp(StartX, 0, Width - 1);
+		StartY = FMath::Clamp(StartY, 0, Height - 1);
+		OutHitPoint.X = StartX;
+		OutHitPoint.Y = StartY;
+
+		// Should probably calculate a better max based on StartX and StartY
+		int RequiredSearchDistance = (FMath::Max(Width, Height) + 1) / 2;
+		int MaxSearchDistance = FMath::Min(RequiredSearchDistance, MaxAllowedSearchDistance);
+		int SearchDistance = 0;
+		while (SearchDistance < MaxSearchDistance)
+		{
+			int X0 = FMath::Max(StartX - SearchDistance, 0);
+			int X1 = FMath::Min(StartX + SearchDistance, Width - 1);
+			int Y0 = FMath::Max(StartY - SearchDistance, 0);
+			int Y1 = FMath::Min(StartY + SearchDistance, Height - 1);
+			// Search along the rectangular edge
+			for (int Y = Y0; Y <= Y1; ++Y)
+			{
+				if (GetPixel(X0, Y) != 0) 
+				{
+					OutHitPoint.X = X0;
+					OutHitPoint.Y = Y;
+					return true;
+				}
+				if (GetPixel(X1, Y) != 0) 
+				{
+					OutHitPoint.X = X1;
+					OutHitPoint.Y = Y;
+					return true;
+				}
+			}
+			for (int X = X0; X <= X1; ++X)
+			{
+				if (GetPixel(X, Y0) != 0) 
+				{
+					OutHitPoint.X = X;
+					OutHitPoint.Y = Y0;
+					return true;
+				}
+				if (GetPixel(X, Y1) != 0) 
+				{
+					OutHitPoint.X = X;
+					OutHitPoint.Y = Y1;
+					return true;
+				}
+			}
+			SearchDistance += 1;
+		}
+		return false;
+	}
+
+	void ThresholdImage(int32 AlphaThreshold)
+	{
+		for (uint8& Pixel : RawData)
+		{
+			if (Pixel > AlphaThreshold)
+			{
+				Pixel = 1;
+			}
+			else
+			{
+				Pixel = DefaultValue;
+			}
+		}
+	}
+
+	// Flushes values smaller or equal to LowAlphaThreshold to 0, and values greater or equal to HighAlphaThreshold to 255
+	void ThresholdImageBothWays(int32 LowAlphaThreshold, int32 HighAlphaThreshold)
+	{
+		for (uint8& Pixel : RawData)
+		{
+			if (Pixel <= LowAlphaThreshold)
+			{
+				Pixel = 0;
+			}
+			else if (Pixel >= HighAlphaThreshold)
+			{
+				Pixel = 255;
+			}
+		}
+	}
+
+	// Checks the image to determine if it is suitable for opaque, masked or translucent rendering
+	void AnalyzeImage(int32 StartX, int32 StartY, int32 InWidth, int32 InHeight, bool& bOutHasZeros, bool& bOutHasIntermediateValues)
+	{
+		bOutHasZeros = false;
+		bOutHasIntermediateValues = false;
+
+		const int32 X0 = StartX;
+		const int32 Y0 = StartY;
+		const int32 X1 = StartX + InWidth;
+		const int32 Y1 = StartY + InHeight;
+		for (int Y = Y0; Y < Y1; ++Y)
+		{
+			for (int X = X0; X < X1; ++X)
+			{
+				const uint8 Value = GetPixel(X, Y);
+
+				if ((Value > 0) && (Value < 255))
+				{
+					bOutHasIntermediateValues = true;
+				}
+
+				if (Value == 0)
+				{
+					bOutHasZeros = true;
+				}
+			}
+		}
+	}
+
+public:
+	TArray<uint8> RawData;
+	int32 Width;
+	int32 Height;
+	uint8 DefaultValue;
+};
diff --git a/Engine/Plugins/2D/Paper2D/Source/Paper2DEditor/Classes/PaperImporterSettings.h b/Engine/Plugins/2D/Paper2D/Source/Paper2DEditor/Classes/PaperImporterSettings.h
--- a/Engine/Plugins/2D/Paper2D/Source/Paper2DEditor/Classes/PaperImporterSettings.h
+++ b/Engine/Plugins/2D/Paper2D/Source/Paper2DEditor/Classes/PaperImporterSettings.h
@@ -4,6 +4,36 @@
 
 #include "PaperImporterSettings.generated.h"
 
+enum class ESpriteInitMaterialLightingMode : uint8
+{
+	// Use the default in the importer settings
+	Automatic,
+
+	// Force unlit
+	ForceUnlit,
+
+	// Force lit
+	ForceLit
+};
+
+enum class ESpriteInitMaterialType : uint8
+{
+	// Leave the material alone
+	LeaveAsIs,
+
+	// Use the default in the importer settings (typically auto-analyze)
+	Automatic,
+
+	// Force masked
+	Masked,
+
+	// Force translucent
+	Translucent,
+
+	// Force opaque
+	Opaque
+};
+
 /**
  * Implements the settings for imported Paper2D assets, such as sprite sheet textures.
  */
@@ -12,7 +42,19 @@ class PAPER2DEDITOR_API UPaperImporterSettings : public UObject
 {
 	GENERATED_BODY()
 
-public:
+protected:
+	// Should the source texture be scanned when creating new sprites to determine the appropriate material? (if false, the Default Masked Material is always used)
+	UPROPERTY(config, EditAnywhere, Category=Settings)
+	bool bPickBestMaterialWhenCreatingSprite;
+
+	// Can opaque materials be applied as part of the 'best material' analysis?
+	UPROPERTY(config, EditAnywhere, Category=Settings, meta=(EditCondition=bPickBestMaterialWhenCreatingSprite))
+	bool bAnalysisCanUseOpaque;
+
+	// The default scaling factor between pixels and Unreal units (cm) to use for newly created sprite assets (e.g., 0.64 would make a 64 pixel wide sprite take up 100 cm)
+	UPROPERTY(config, EditAnywhere, Category=Settings)
+	float DefaultPixelsPerUnrealUnit;
+
 	// A list of default suffixes to use when looking for associated normal maps while importing sprites or creating sprites from textures
 	UPROPERTY(config, EditAnywhere, Category=ImportSettings)
 	TArray<FString> NormalMapTextureSuffixes;
@@ -31,18 +73,30 @@ public:
 	
 	// Compression settings to use when building the texture.
 	// The default texture group for imported sprite textures, tile sheets, etc... (typically set to UI for 'modern 2D' or 2D pixels for 'retro 2D')
-	UPROPERTY(config, EditAnywhere, Category = ImportSettings, meta = (EditCondition = bOverrideTextureCompression))
+	UPROPERTY(config, EditAnywhere, Category=ImportSettings, meta=(EditCondition=bOverrideTextureCompression))
 	TEnumAsByte<TextureCompressionSettings> DefaultSpriteTextureCompression;
 
-	// The default masked material for newly created sprites (masked means binary opacity: things are either opaque or see-thru, with nothing in between)
+	// The unlit default masked material for newly created sprites (masked means binary opacity: things are either opaque or see-thru, with nothing in between)
+	UPROPERTY(config, EditAnywhere, Category=MaterialSettings, meta=(AllowedClasses="MaterialInterface", DisplayName="Unlit Default Masked Material"))
+	FStringAssetReference UnlitDefaultMaskedMaterialName;
+
+	// The unlit default translucent material for newly created sprites (translucent means smooth opacity which can vary continuously from 0..1, but translucent rendering is more expensive that opaque or masked rendering and has different sorting rules)
+	UPROPERTY(config, EditAnywhere, Category=MaterialSettings, meta=(AllowedClasses="MaterialInterface", DisplayName="Unlit Default Translucent Material"))
+	FStringAssetReference UnlitDefaultTranslucentMaterialName;
+
+	// The unlit default opaque material for newly created sprites
+	UPROPERTY(config, EditAnywhere, Category=MaterialSettings, meta=(AllowedClasses="MaterialInterface", DisplayName="Unlit Default Opaque Sprite Material"))
+	FStringAssetReference UnlitDefaultOpaqueMaterialName;
+
+	// The lit default masked material for newly created sprites (masked means binary opacity: things are either opaque or see-thru, with nothing in between)
 	UPROPERTY(config, EditAnywhere, Category=MaterialSettings, meta=(AllowedClasses="MaterialInterface", DisplayName="Lit Default Masked Material"))
 	FStringAssetReference LitDefaultMaskedMaterialName;
 
-	// The default translucent material for newly created sprites (translucent means smooth opacity which can vary continuously from 0..1, but translucent rendering is more expensive that opaque or masked rendering and has different sorting rules)
+	// The lit default translucent material for newly created sprites (translucent means smooth opacity which can vary continuously from 0..1, but translucent rendering is more expensive that opaque or masked rendering and has different sorting rules)
 	UPROPERTY(config, EditAnywhere, Category=MaterialSettings, meta=(AllowedClasses="MaterialInterface", DisplayName="Lit Default Translucent Material"))
 	FStringAssetReference LitDefaultTranslucentMaterialName;
 
-	// The default opaque material for newly created sprites
+	// The lit default opaque material for newly created sprites
 	UPROPERTY(config, EditAnywhere, Category=MaterialSettings, meta=(AllowedClasses="MaterialInterface", DisplayName="Lit Default Opaque Material"))
 	FStringAssetReference LitDefaultOpaqueMaterialName;
 
@@ -55,9 +109,22 @@ public:
 	// Generates names to test for a normal map using NormalMapTextureSuffixes
 	void GenerateNormalMapNamesToTest(const FString& InRoot, TArray<FString>& InOutNames) const;
 
-	// Pushes the lit materials
-	void PopulateMaterialsIntoInitParams(FSpriteAssetInitParameters& InitParams) const;
-
 	// Applies the compression settings to the specified texture
 	void ApplyTextureSettings(UTexture2D* Texture) const;
+
+	// Applies the material settings to the 
+	// Note: This should be called after the texture has been set, as that will be analyzed if the lighting/material type flags are set to automatic
+	void ApplySettingsForSpriteInit(FSpriteAssetInitParameters& InitParams, ESpriteInitMaterialLightingMode LightingMode = ESpriteInitMaterialLightingMode::Automatic, ESpriteInitMaterialType MaterialTypeMode = ESpriteInitMaterialType::Automatic) const;
+
+	// Returns the default translucent material
+	UMaterialInterface* GetDefaultTranslucentMaterial(bool bLit) const;
+
+	// Returns the default opaque material
+	UMaterialInterface* GetDefaultOpaqueMaterial(bool bLit) const;
+
+	// Returns the default masked material
+	UMaterialInterface* GetDefaultMaskedMaterial(bool bLit) const;
+
+	// Returns the default pixels/uu setting
+	float GetDefaultPixelsPerUnrealUnit() const { return DefaultPixelsPerUnrealUnit; }
 };
diff --git a/Engine/Plugins/2D/Paper2D/Source/Paper2DEditor/Private/PaperImporterSettings.cpp b/Engine/Plugins/2D/Paper2D/Source/Paper2DEditor/Private/PaperImporterSettings.cpp
--- a/Engine/Plugins/2D/Paper2D/Source/Paper2DEditor/Private/PaperImporterSettings.cpp
+++ b/Engine/Plugins/2D/Paper2D/Source/Paper2DEditor/Private/PaperImporterSettings.cpp
@@ -2,12 +2,15 @@
 
 #include "Paper2DEditorPrivatePCH.h"
 #include "PaperImporterSettings.h"
+#include "AlphaBitmap.h"
 
 //////////////////////////////////////////////////////////////////////////
 // UPaperImporterSettings
 
 UPaperImporterSettings::UPaperImporterSettings()
-: DefaultSpriteTextureGroup(TEXTUREGROUP_Pixels2D)
+	: bPickBestMaterialWhenCreatingSprite(true)
+	, DefaultPixelsPerUnrealUnit(1.0f)
+	, DefaultSpriteTextureGroup(TEXTUREGROUP_Pixels2D)
 	, bOverrideTextureCompression(true)
 	, DefaultSpriteTextureCompression(TC_EditorIcon)
 {
@@ -19,6 +22,10 @@ UPaperImporterSettings::UPaperImporterSettings()
 
 	if (!IsRunningCommandlet())
 	{
+		UnlitDefaultMaskedMaterialName = FStringAssetReference("/Paper2D/MaskedUnlitSpriteMaterial.MaskedUnlitSpriteMaterial");
+		UnlitDefaultTranslucentMaterialName = FStringAssetReference("/Paper2D/TranslucentUnlitSpriteMaterial.TranslucentUnlitSpriteMaterial");
+		UnlitDefaultOpaqueMaterialName = FStringAssetReference("/Paper2D/OpaqueUnlitSpriteMaterial.OpaqueUnlitSpriteMaterial");
+
 		LitDefaultMaskedMaterialName = FStringAssetReference("/Paper2D/MaskedLitSpriteMaterial.MaskedLitSpriteMaterial");
 		LitDefaultTranslucentMaterialName = FStringAssetReference("/Paper2D/TranslucentLitSpriteMaterial.TranslucentLitSpriteMaterial");
 		LitDefaultOpaqueMaterialName = FStringAssetReference("/Paper2D/OpaqueLitSpriteMaterial.OpaqueLitSpriteMaterial");
@@ -48,13 +55,6 @@ void UPaperImporterSettings::GenerateNormalMapNamesToTest(const FString& InRoot,
 	}
 }
 
-void UPaperImporterSettings::PopulateMaterialsIntoInitParams(FSpriteAssetInitParameters& InitParams) const
-{
-	InitParams.MaskedMaterialOverride = Cast<UMaterialInterface>(LitDefaultMaskedMaterialName.TryLoad());
-	InitParams.OpaqueMaterialOverride = Cast<UMaterialInterface>(LitDefaultOpaqueMaterialName.TryLoad());
-	InitParams.TranslucentMaterialOverride = Cast<UMaterialInterface>(LitDefaultTranslucentMaterialName.TryLoad());
-}
-
 void UPaperImporterSettings::ApplyTextureSettings(UTexture2D* Texture) const
 {
 	if (Texture->IsNormalMap())
@@ -75,3 +75,77 @@ void UPaperImporterSettings::ApplyTextureSettings(UTexture2D* Texture) const
 		Texture->PostEditChange();
 	}
 }
+
+void UPaperImporterSettings::ApplySettingsForSpriteInit(FSpriteAssetInitParameters& InitParams, ESpriteInitMaterialLightingMode LightingMode, ESpriteInitMaterialType MaterialTypeMode) const
+{
+	InitParams.SetPixelsPerUnrealUnit(DefaultPixelsPerUnrealUnit);
+
+	ESpriteInitMaterialType DesiredMaterialType = MaterialTypeMode;
+	if (DesiredMaterialType == ESpriteInitMaterialType::Automatic)
+	{
+		// Analyze the texture if desired (to see if it's got greyscale alpha or just binary alpha, picking either a translucent or masked material)
+		if (bPickBestMaterialWhenCreatingSprite && (InitParams.Texture != nullptr))
+		{
+			FAlphaBitmap AlphaBitmap(InitParams.Texture);
+			bool bHasIntermediateValues;
+			bool bHasZeros;
+			AlphaBitmap.AnalyzeImage((int32)InitParams.Offset.X, (int32)InitParams.Offset.Y, (int32)InitParams.Dimension.X, (int32)InitParams.Dimension.Y, /*out*/ bHasZeros, /*out*/ bHasIntermediateValues);
+
+			if (bAnalysisCanUseOpaque && !bHasIntermediateValues && !bHasZeros)
+			{
+				DesiredMaterialType = ESpriteInitMaterialType::Opaque;
+			}
+			else
+			{
+				DesiredMaterialType = bHasIntermediateValues ? ESpriteInitMaterialType::Translucent : ESpriteInitMaterialType::Masked;
+			}
+		}
+	}
+
+	if (DesiredMaterialType == ESpriteInitMaterialType::Automatic)
+	{
+		// Fall back to masked if we wanted automatic and couldn't analyze things
+		DesiredMaterialType = ESpriteInitMaterialType::Masked;
+	}
+
+	if (DesiredMaterialType != ESpriteInitMaterialType::LeaveAsIs)
+	{
+		// Determine whether to use lit or unlit materials
+		const bool bUseLitMaterial = LightingMode == ESpriteInitMaterialLightingMode::ForceLit;
+
+		// Apply the materials
+		switch (DesiredMaterialType)
+		{
+		default:
+		case ESpriteInitMaterialType::LeaveAsIs:
+		case ESpriteInitMaterialType::Automatic:
+			check(false);
+		case ESpriteInitMaterialType::Masked:
+			InitParams.DefaultMaterialOverride = GetDefaultMaskedMaterial(bUseLitMaterial);
+			break;
+		case ESpriteInitMaterialType::Translucent:
+			InitParams.DefaultMaterialOverride = GetDefaultTranslucentMaterial(bUseLitMaterial);
+			break;
+		case ESpriteInitMaterialType::Opaque:
+			InitParams.DefaultMaterialOverride = GetDefaultOpaqueMaterial(bUseLitMaterial);
+			break;
+		}
+
+		InitParams.AlternateMaterialOverride = GetDefaultOpaqueMaterial(bUseLitMaterial);
+	}
+}
+
+UMaterialInterface* UPaperImporterSettings::GetDefaultTranslucentMaterial(bool bLit) const
+{
+	return Cast<UMaterialInterface>((bLit ? LitDefaultTranslucentMaterialName : UnlitDefaultTranslucentMaterialName).TryLoad());
+}
+
+UMaterialInterface* UPaperImporterSettings::GetDefaultOpaqueMaterial(bool bLit) const
+{
+	return Cast<UMaterialInterface>((bLit ? LitDefaultOpaqueMaterialName : UnlitDefaultOpaqueMaterialName).TryLoad());
+}
+
+UMaterialInterface* UPaperImporterSettings::GetDefaultMaskedMaterial(bool bLit) const
+{
+	return Cast<UMaterialInterface>((bLit ? LitDefaultMaskedMaterialName : UnlitDefaultMaskedMaterialName).TryLoad());
+}
diff --git a/Engine/Plugins/2D/Paper2D/Source/Paper2DEditor/Private/PaperSpriteFactory.cpp b/Engine/Plugins/2D/Paper2D/Source/Paper2DEditor/Private/PaperSpriteFactory.cpp
--- a/Engine/Plugins/2D/Paper2D/Source/Paper2DEditor/Private/PaperSpriteFactory.cpp
+++ b/Engine/Plugins/2D/Paper2D/Source/Paper2DEditor/Private/PaperSpriteFactory.cpp
@@ -33,7 +33,6 @@ UObject* UPaperSpriteFactory::FactoryCreateNew(UClass* Class, UObject* InParent,
 	UPaperSprite* NewSprite = NewObject<UPaperSprite>(InParent, Class, Name, Flags | RF_Transactional);
 
 	FSpriteAssetInitParameters SpriteInitParams;
-	SpriteInitParams.bNewlyCreated = true;
 
 	if (bUseSourceRegion)
 	{
@@ -46,6 +45,9 @@ UObject* UPaperSpriteFactory::FactoryCreateNew(UClass* Class, UObject* InParent,
 		SpriteInitParams.SetTextureAndFill(InitialTexture);
 	}
 
+	const UPaperImporterSettings* ImporterSettings = GetDefault<UPaperImporterSettings>();
+
+	bool bFoundNormalMap = false;
 	if (InitialTexture != nullptr)
 	{
 		// Look for an associated normal map to go along with the base map
@@ -53,8 +55,6 @@ UObject* UPaperSpriteFactory::FactoryCreateNew(UClass* Class, UObject* InParent,
 		const FString PackagePath = FPackageName::GetLongPackagePath(SanitizedBasePackageName);
 		FAssetRegistryModule& AssetRegistryModule = FModuleManager::Get().LoadModuleChecked<FAssetRegistryModule>(TEXT("AssetRegistry"));
 
-		const UPaperImporterSettings* ImporterSettings = GetDefault<UPaperImporterSettings>();
-
 		const FString NormalMapNameNoSuffix = ImporterSettings->RemoveSuffixFromBaseMapName(InitialTexture->GetName());
 
 		TArray<FString> NamesToTest;
@@ -71,14 +71,15 @@ UObject* UPaperSpriteFactory::FactoryCreateNew(UClass* Class, UObject* InParent,
 			{
 				if (UTexture2D* NormalMapTexture = Cast<UTexture2D>(AssetData.GetAsset()))
 				{
+					bFoundNormalMap = true;
 					SpriteInitParams.AdditionalTextures.Add(NormalMapTexture);
-					ImporterSettings->PopulateMaterialsIntoInitParams(/*inout*/ SpriteInitParams);
 					break;
 				}
 			}
 		}
 	}
 
+	ImporterSettings->ApplySettingsForSpriteInit(/*inout*/ SpriteInitParams, bFoundNormalMap ? ESpriteInitMaterialLightingMode::ForceLit : ESpriteInitMaterialLightingMode::Automatic);
 	NewSprite->InitializeSprite(SpriteInitParams);
 
 	return NewSprite;
diff --git a/Engine/Plugins/2D/Paper2D/Source/Paper2DEditor/Private/SpriteEditor/SpriteEditorViewportClient.cpp b/Engine/Plugins/2D/Paper2D/Source/Paper2DEditor/Private/SpriteEditor/SpriteEditorViewportClient.cpp
--- a/Engine/Plugins/2D/Paper2D/Source/Paper2DEditor/Private/SpriteEditor/SpriteEditorViewportClient.cpp
+++ b/Engine/Plugins/2D/Paper2D/Source/Paper2DEditor/Private/SpriteEditor/SpriteEditorViewportClient.cpp
@@ -146,8 +146,11 @@ void FSpriteEditorViewportClient::UpdateSourceTextureSpriteFromSprite(UPaperSpri
 			FComponentReregisterContext ReregisterSprite(SourceTextureViewComponent);
 
 			FSpriteAssetInitParameters SpriteReinitParams;
+
 			SpriteReinitParams.SetTextureAndFill(SourceSprite->SourceTexture);
-			TargetSprite->PixelsPerUnrealUnit = SourceSprite->PixelsPerUnrealUnit;
+			SpriteReinitParams.DefaultMaterialOverride = SourceSprite->DefaultMaterial;
+			SpriteReinitParams.AlternateMaterialOverride = SourceSprite->AlternateMaterial;
+			SpriteReinitParams.SetPixelsPerUnrealUnit(SourceSprite->PixelsPerUnrealUnit);
 			TargetSprite->InitializeSprite(SpriteReinitParams);
 
 			RequestFocusOnSelection(/*bInstant=*/ true);
diff --git a/Engine/Plugins/2D/Paper2D/Source/Paper2DEditor/Private/TileMapEditing/TileMapActorFactory.cpp b/Engine/Plugins/2D/Paper2D/Source/Paper2DEditor/Private/TileMapEditing/TileMapActorFactory.cpp
--- a/Engine/Plugins/2D/Paper2D/Source/Paper2DEditor/Private/TileMapEditing/TileMapActorFactory.cpp
+++ b/Engine/Plugins/2D/Paper2D/Source/Paper2DEditor/Private/TileMapEditing/TileMapActorFactory.cpp
@@ -3,6 +3,7 @@
 #include "Paper2DEditorPrivatePCH.h"
 #include "AssetData.h"
 #include "TileMapActorFactory.h"
+#include "PaperImporterSettings.h"
 
 //////////////////////////////////////////////////////////////////////////
 // UTileMapActorFactory
@@ -36,6 +37,7 @@ void UTileMapActorFactory::PostSpawnActor(UObject* Asset, AActor* NewActor)
 			RenderComponent->TileMap->TileWidth = TileSet->TileWidth;
 			RenderComponent->TileMap->TileHeight = TileSet->TileHeight;
 			RenderComponent->TileMap->SelectedTileSet = TileSet;
+			//@TODO: Initialize the tile map material here too, based on analysis of the TileSheet texture
 			RenderComponent->RegisterComponent();
 		}
 	}
@@ -43,6 +45,7 @@ void UTileMapActorFactory::PostSpawnActor(UObject* Asset, AActor* NewActor)
 	if (RenderComponent->OwnsTileMap())
 	{
 		RenderComponent->TileMap->InitializeNewEmptyTileMap();
+		RenderComponent->TileMap->PixelsPerUnrealUnit = GetDefault<UPaperImporterSettings>()->GetDefaultPixelsPerUnrealUnit();
 	}
 }
 
@@ -64,12 +67,14 @@ void UTileMapActorFactory::PostCreateBlueprint(UObject* Asset, AActor* CDO)
 				RenderComponent->TileMap->TileWidth = TileSet->TileWidth;
 				RenderComponent->TileMap->TileHeight = TileSet->TileHeight;
 				RenderComponent->TileMap->SelectedTileSet = TileSet;
+				//@TODO: Initialize the tile map material here too, based on analysis of the TileSheet texture
 			}
 		}
 
 		if (RenderComponent->OwnsTileMap())
 		{
 			RenderComponent->TileMap->InitializeNewEmptyTileMap();
+			RenderComponent->TileMap->PixelsPerUnrealUnit = GetDefault<UPaperImporterSettings>()->GetDefaultPixelsPerUnrealUnit();
 		}
 	}
 }
diff --git a/Engine/Plugins/2D/Paper2D/Source/Paper2DEditor/Private/TileSetEditor/SingleTileEditorViewportClient.cpp b/Engine/Plugins/2D/Paper2D/Source/Paper2DEditor/Private/TileSetEditor/SingleTileEditorViewportClient.cpp
--- a/Engine/Plugins/2D/Paper2D/Source/Paper2DEditor/Private/TileSetEditor/SingleTileEditorViewportClient.cpp
+++ b/Engine/Plugins/2D/Paper2D/Source/Paper2DEditor/Private/TileSetEditor/SingleTileEditorViewportClient.cpp
@@ -228,20 +228,21 @@ void FSingleTileEditorViewportClient::SetTileIndex(int32 InTileIndex)
 	UPaperSprite* DummySprite = nullptr;
 	if (TileBeingEditedIndex != INDEX_NONE)
 	{
-		//@TODO: Should use this to pick the correct material
-		//const bool bUseTranslucentBlend = Texture->HasAlphaChannel();
-
 		DummySprite = NewObject<UPaperSprite>();
  		DummySprite->SpriteCollisionDomain = ESpriteCollisionMode::None;
  		DummySprite->PivotMode = ESpritePivotMode::Center_Center;
  		DummySprite->CollisionGeometry.GeometryType = ESpritePolygonMode::SourceBoundingBox;
  		DummySprite->RenderGeometry.GeometryType = ESpritePolygonMode::SourceBoundingBox;
-		DummySprite->PixelsPerUnrealUnit = 1.0f;
 
 		FSpriteAssetInitParameters SpriteReinitParams;
+
 		SpriteReinitParams.Texture = TileSet->TileSheet;
+
+		//@TODO: Should analyze the texture (*at a higher level, not per tile click!*) to pick the correct material
+
 		TileSet->GetTileUV(TileBeingEditedIndex, /*out*/ SpriteReinitParams.Offset);
 		SpriteReinitParams.Dimension = FVector2D(TileSet->TileWidth, TileSet->TileHeight);
+		SpriteReinitParams.SetPixelsPerUnrealUnit(1.0f);
 		DummySprite->InitializeSprite(SpriteReinitParams);
 	}
 	PreviewTileSpriteComponent->SetSprite(DummySprite);
diff --git a/Engine/Plugins/2D/Paper2D/Source/PaperSpriteSheetImporter/Private/PaperJsonSpriteSheetImporter.cpp b/Engine/Plugins/2D/Paper2D/Source/PaperSpriteSheetImporter/Private/PaperJsonSpriteSheetImporter.cpp
--- a/Engine/Plugins/2D/Paper2D/Source/PaperSpriteSheetImporter/Private/PaperJsonSpriteSheetImporter.cpp
+++ b/Engine/Plugins/2D/Paper2D/Source/PaperSpriteSheetImporter/Private/PaperJsonSpriteSheetImporter.cpp
@@ -566,14 +566,13 @@ bool FPaperJsonSpriteSheetImporter::PerformImport(const FString& LongPackagePath
 		{
 			// Put the normal map into the additional textures array and ask for a lit material instead of unlit
 			SpriteInitParams.AdditionalTextures.Add(NormalMapTexture);
-
-			const UPaperImporterSettings* ImporterSettings = GetDefault<UPaperImporterSettings>();
-			ImporterSettings->PopulateMaterialsIntoInitParams(/*inout*/ SpriteInitParams);
 		}
 
 		SpriteInitParams.Offset = Frame.SpritePosInSheet;
 		SpriteInitParams.Dimension = Frame.SpriteSizeInSheet;
-		SpriteInitParams.bNewlyCreated = true;
+
+		GetDefault<UPaperImporterSettings>()->ApplySettingsForSpriteInit(/*inout*/ SpriteInitParams, (NormalMapTexture != nullptr) ? ESpriteInitMaterialLightingMode::ForceLit : ESpriteInitMaterialLightingMode::Automatic);
+
 		TargetSprite->InitializeSprite(SpriteInitParams);
 
 		if (Frame.bRotated)
diff --git a/Engine/Plugins/2D/Paper2D/Source/PaperTiledImporter/Private/PaperTiledImporterFactory.cpp b/Engine/Plugins/2D/Paper2D/Source/PaperTiledImporter/Private/PaperTiledImporterFactory.cpp
--- a/Engine/Plugins/2D/Paper2D/Source/PaperTiledImporter/Private/PaperTiledImporterFactory.cpp
+++ b/Engine/Plugins/2D/Paper2D/Source/PaperTiledImporter/Private/PaperTiledImporterFactory.cpp
@@ -209,7 +209,7 @@ UObject* UPaperTiledImporterFactory::FactoryCreateText(UClass* InClass, UObject*
 		Result->SeparationPerTileY = 0.0f;
 		Result->SeparationPerLayer = 1.0f;
 		Result->ProjectionMode = GlobalInfo.GetOrientationType();
-		Result->PixelsPerUnrealUnit = GetDefault<UPaperRuntimeSettings>()->DefaultPixelsPerUnrealUnit;
+		Result->PixelsPerUnrealUnit = GetDefault<UPaperImporterSettings>()->GetDefaultPixelsPerUnrealUnit();
 		Result->BackgroundColor = GlobalInfo.BackgroundColor;
 		Result->HexSideLength = GlobalInfo.HexSideLength;
 
@@ -310,6 +310,8 @@ UObject* UPaperTiledImporterFactory::FactoryCreateText(UClass* InClass, UObject*
 		{
 			// Bind our selected tile set to the first tile set that was imported so something is already picked
 			Result->SelectedTileSet = GlobalInfo.CreatedTileSetAssets[0];
+		
+			//@TODO: Should analyze the material from the used tile sets and set it appropriately on the tile map
 		}
 
 		// Create the layers
diff --git a/Engine/Plugins/2D/Paper2D/Source/SpriterImporter/Private/SpriterImporterFactory.cpp b/Engine/Plugins/2D/Paper2D/Source/SpriterImporter/Private/SpriterImporterFactory.cpp
--- a/Engine/Plugins/2D/Paper2D/Source/SpriterImporter/Private/SpriterImporterFactory.cpp
+++ b/Engine/Plugins/2D/Paper2D/Source/SpriterImporter/Private/SpriterImporterFactory.cpp
@@ -302,7 +302,8 @@ UObject* USpriterImporterFactory::FactoryCreateText(UClass* InClass, UObject* In
 
 					FSpriteAssetInitParameters SpriteInitParams;
 					SpriteInitParams.SetTextureAndFill(ImportedTexture);
-					//@TODO: Need to make a way to force pixels/uu to 1.0 when doing this: SpriteInitParams.bNewlyCreated = true;
+					GetDefault<UPaperImporterSettings>()->ApplySettingsForSpriteInit(SpriteInitParams);
+					SpriteInitParams.SetPixelsPerUnrealUnit(1.0f);
 					ImportedSprite->InitializeSprite(SpriteInitParams);
 				}
 				else if (File.FileType == ESpriterFileType::Sound)
