Tile Map Editor: Improving the behavior of isometric staggered and hexagonal staggered projection modes
- Better bounds calculations
- Better preview wireframe cursor, showing the logical borders of all tiles below the preview cursor
- Removed secondary bogus wireframe due to the preview brush
- Preview brush is now centered correctly for single-row brushes
- Added UPaperTileMap::GetTilePolygon, which returns the logical borders of a tile (e.g., a hex in hexagonal staggered)
- Fixed UPaperTileMap::GetTileCenterInLocalSpace to return the actual center for staggered iso and hex tiles

Note: Multi-row brushes in staggered modes are still offset incorrectly on odd lines when the cursor is on an odd line due to the stagger behavior, but the wireframe preview shows where the ink will really go

WIP for [UE-11488]

diff --git a/Engine/Plugins/2D/Paper2D/Source/Paper2D/Classes/PaperTileMap.h b/Engine/Plugins/2D/Paper2D/Source/Paper2D/Classes/PaperTileMap.h
--- a/Engine/Plugins/2D/Paper2D/Source/Paper2D/Classes/PaperTileMap.h
+++ b/Engine/Plugins/2D/Paper2D/Source/Paper2D/Classes/PaperTileMap.h
@@ -141,6 +141,8 @@ public:
 	// Returns the center of the specified tile in local space
 	FVector GetTileCenterInLocalSpace(float TileX, float TileY, int32 LayerIndex = 0) const;
 
+	// Returns the polygon for the specified tile (will be 4 or 6 vertices as a rectangle, diamond, or hexagon)
+	void GetTilePolygon(int32 TileX, int32 TileY, int32 LayerIndex, TArray<FVector>& LocalSpacePoints);
 
 	void GetTileToLocalParameters(FVector& OutCornerPosition, FVector& OutStepX, FVector& OutStepY, FVector& OutOffsetYFactor) const;
 	void GetLocalToTileParameters(FVector& OutCornerPosition, FVector& OutStepX, FVector& OutStepY, FVector& OutOffsetYFactor) const;
diff --git a/Engine/Plugins/2D/Paper2D/Source/Paper2D/Private/PaperTileMap.cpp b/Engine/Plugins/2D/Paper2D/Source/Paper2D/Private/PaperTileMap.cpp
--- a/Engine/Plugins/2D/Paper2D/Source/Paper2D/Private/PaperTileMap.cpp
+++ b/Engine/Plugins/2D/Paper2D/Source/Paper2D/Private/PaperTileMap.cpp
@@ -343,9 +343,95 @@ FVector UPaperTileMap::GetTilePositionInLocalSpace(float TileX, float TileY, int
 	return LocalPos;
 }
 
+void UPaperTileMap::GetTilePolygon(int32 TileX, int32 TileY, int32 LayerIndex, TArray<FVector>& LocalSpacePoints)
+{
+	switch (ProjectionMode)
+	{
+	case ETileMapProjectionMode::Orthogonal:
+	case ETileMapProjectionMode::IsometricDiamond:
+	default:
+		LocalSpacePoints.Add(GetTilePositionInLocalSpace(TileX, TileY, LayerIndex));
+		LocalSpacePoints.Add(GetTilePositionInLocalSpace(TileX + 1, TileY, LayerIndex));
+		LocalSpacePoints.Add(GetTilePositionInLocalSpace(TileX, TileY + 1, LayerIndex));
+		LocalSpacePoints.Add(GetTilePositionInLocalSpace(TileX + 1, TileY + 1, LayerIndex));
+		break;
+
+	case ETileMapProjectionMode::IsometricStaggered:
+		{
+			const float UnrealUnitsPerPixel = GetUnrealUnitsPerPixel();
+			const float TileWidthInUU = TileWidth * UnrealUnitsPerPixel;
+			const float TileHeightInUU = TileHeight * UnrealUnitsPerPixel;
+
+			const FVector RecenterOffset = PaperAxisX*TileWidthInUU*0.5f;
+			const FVector LSTM = GetTilePositionInLocalSpace(TileX, TileY, LayerIndex) + RecenterOffset;
+
+			LocalSpacePoints.Add(LSTM);
+			LocalSpacePoints.Add(LSTM + PaperAxisX*TileWidthInUU*0.5f - PaperAxisY*TileHeightInUU*0.5f);
+			LocalSpacePoints.Add(LSTM - PaperAxisY*TileHeightInUU*1.0f);
+			LocalSpacePoints.Add(LSTM - PaperAxisX*TileWidthInUU*0.5f - PaperAxisY*TileHeightInUU*0.5f);
+		}
+		break;
+
+	case ETileMapProjectionMode::HexagonalStaggered:
+		{
+			const float UnrealUnitsPerPixel = GetUnrealUnitsPerPixel();
+			const float TileWidthInUU = TileWidth * UnrealUnitsPerPixel;
+			const float TileHeightInUU = TileHeight * UnrealUnitsPerPixel;
+
+			const FVector RecenterOffset = PaperAxisX*TileWidthInUU*0.5f;
+			const FVector LSTM = GetTilePositionInLocalSpace(TileX, TileY, LayerIndex) + RecenterOffset;
+
+			const float HexSideLengthInUU = HexSideLength * UnrealUnitsPerPixel;
+
+			const FVector Top(LSTM - PaperAxisY*HexSideLengthInUU);
+
+			const FVector RightTop(LSTM + PaperAxisX*TileWidthInUU*0.5f - PaperAxisY*TileHeightInUU*0.5f);
+			const FVector LeftTop(LSTM - PaperAxisX*TileWidthInUU*0.5f - PaperAxisY*TileHeightInUU*0.5f);
+
+			const FVector RightBottom(LSTM + PaperAxisX*TileWidthInUU*0.5f - PaperAxisY*(TileHeightInUU*0.5f + HexSideLengthInUU));
+			const FVector LeftBottom(LSTM - PaperAxisX*TileWidthInUU*0.5f - PaperAxisY*(TileHeightInUU*0.5f + HexSideLengthInUU));
+
+			const FVector Bottom(LSTM - PaperAxisY*TileHeightInUU*1.0f);
+
+			LocalSpacePoints.Add(Top);
+			LocalSpacePoints.Add(RightTop);
+			LocalSpacePoints.Add(RightBottom);
+			LocalSpacePoints.Add(Bottom);
+			LocalSpacePoints.Add(LeftBottom);
+			LocalSpacePoints.Add(LeftTop);
+		}
+		break;
+	}
+}
+
 FVector UPaperTileMap::GetTileCenterInLocalSpace(float TileX, float TileY, int32 LayerIndex) const
 {
-	return GetTilePositionInLocalSpace(TileX + 0.5f, TileY + 0.5f, LayerIndex);
+	switch (ProjectionMode)
+	{
+	case ETileMapProjectionMode::Orthogonal:
+	default:
+		return GetTilePositionInLocalSpace(TileX + 0.5f, TileY + 0.5f, LayerIndex);
+	case ETileMapProjectionMode::IsometricStaggered:
+		{
+			const float UnrealUnitsPerPixel = GetUnrealUnitsPerPixel();
+			const float TileWidthInUU = TileWidth * UnrealUnitsPerPixel;
+			const float TileHeightInUU = TileHeight * UnrealUnitsPerPixel;
+
+			const FVector RecenterOffset = PaperAxisX*TileWidthInUU*0.5f - PaperAxisY*TileHeightInUU*0.5f;
+			return GetTilePositionInLocalSpace(TileX, TileY, LayerIndex) + RecenterOffset;
+		}
+
+	case ETileMapProjectionMode::HexagonalStaggered:
+		{
+			const float UnrealUnitsPerPixel = GetUnrealUnitsPerPixel();
+			const float TileWidthInUU = TileWidth * UnrealUnitsPerPixel;
+			const float TileHeightInUU = TileHeight * UnrealUnitsPerPixel;
+			const float HexSideLengthInUU = HexSideLength * UnrealUnitsPerPixel;
+
+			const FVector RecenterOffset = PaperAxisX*TileWidthInUU*0.5f - PaperAxisY*(HexSideLengthInUU + TileHeightInUU)*0.5f;
+			return GetTilePositionInLocalSpace(TileX, TileY, LayerIndex) + RecenterOffset;
+		}
+	}
 }
 
 FBoxSphereBounds UPaperTileMap::GetRenderBounds() const
@@ -376,7 +462,16 @@ FBoxSphereBounds UPaperTileMap::GetRenderBounds() const
 			 const FBox Box(BottomLeft, BottomLeft + Dimensions);
 			 return FBoxSphereBounds(Box);
 		}
-//@TODO: verify bounds for IsometricStaggered and HexagonalStaggered
+		case ETileMapProjectionMode::HexagonalStaggered:
+		case ETileMapProjectionMode::IsometricStaggered:
+		{
+			const int32 RoundedHalfHeight = (MapHeight + 1) / 2;
+			const FVector BottomLeft((-0.5f) * TileWidthInUU, -HalfThickness - Depth, -(RoundedHalfHeight) * TileHeightInUU);
+			const FVector Dimensions((MapWidth + 0.5f) * TileWidthInUU, Depth + 2 * HalfThickness, (RoundedHalfHeight + 1.0f) * TileHeightInUU);
+
+			const FBox Box(BottomLeft, BottomLeft + Dimensions);
+			return FBoxSphereBounds(Box);
+		}
 	}
 }
 
diff --git a/Engine/Plugins/2D/Paper2D/Source/Paper2DEditor/Private/TileMapEditing/EdModeTileMap.cpp b/Engine/Plugins/2D/Paper2D/Source/Paper2DEditor/Private/TileMapEditing/EdModeTileMap.cpp
--- a/Engine/Plugins/2D/Paper2D/Source/Paper2DEditor/Private/TileMapEditing/EdModeTileMap.cpp
+++ b/Engine/Plugins/2D/Paper2D/Source/Paper2DEditor/Private/TileMapEditing/EdModeTileMap.cpp
@@ -165,6 +165,9 @@ void FEdModeTileMap::Enter()
 	CursorPreviewComponent = NewObject<UPaperTileMapComponent>();
 	CursorPreviewComponent->TileMap->InitializeNewEmptyTileMap();
 	CursorPreviewComponent->TranslucencySortPriority = 99999;
+	CursorPreviewComponent->bShowPerTileGridWhenSelected = false;
+	CursorPreviewComponent->bShowPerLayerGridWhenSelected = false;
+	CursorPreviewComponent->bShowOutlineWhenUnselected = false;
 	CursorPreviewComponent->UpdateBounds();
 	CursorPreviewComponent->AddToRoot();
 	CursorPreviewComponent->RegisterComponentWithWorld(World);
@@ -410,75 +413,25 @@ void FEdModeTileMap::Render(const FSceneView* View, FViewport* Viewport, FPrimit
 				CursorRange = LastEyeDropperBounds;
 			}
 
-			switch (TileMap->ProjectionMode)
-			{
-			case ETileMapProjectionMode::Orthogonal:
-			case ETileMapProjectionMode::IsometricDiamond:
-			default:
-				{
-					const FVector TL(ComponentToWorld.TransformPosition(TileMap->GetTilePositionInLocalSpace(CursorRange.Min.X, CursorRange.Min.Y, LastCursorTileZ)));
-					const FVector TR(ComponentToWorld.TransformPosition(TileMap->GetTilePositionInLocalSpace(CursorRange.Max.X, CursorRange.Min.Y, LastCursorTileZ)));
-					const FVector BL(ComponentToWorld.TransformPosition(TileMap->GetTilePositionInLocalSpace(CursorRange.Min.X, CursorRange.Max.Y, LastCursorTileZ)));
-					const FVector BR(ComponentToWorld.TransformPosition(TileMap->GetTilePositionInLocalSpace(CursorRange.Max.X, CursorRange.Max.Y, LastCursorTileZ)));
-
-					PDI->DrawLine(TL, TR, CursorWireColor, SDPG_Foreground, 0.0f, DepthBias);
-					PDI->DrawLine(TR, BR, CursorWireColor, SDPG_Foreground, 0.0f, DepthBias);
-					PDI->DrawLine(BR, BL, CursorWireColor, SDPG_Foreground, 0.0f, DepthBias);
-					PDI->DrawLine(BL, TL, CursorWireColor, SDPG_Foreground, 0.0f, DepthBias);
-				}
-				break;
-
-			case ETileMapProjectionMode::IsometricStaggered:
-				{
-					const float UnrealUnitsPerPixel = TileMap->GetUnrealUnitsPerPixel();
-					const float TileWidthInUU = TileMap->TileWidth * UnrealUnitsPerPixel;
-					const float TileHeightInUU = TileMap->TileHeight * UnrealUnitsPerPixel;
-
-					const FVector RecenterOffset = PaperAxisX*TileWidthInUU*0.5f;
-					const FVector LSTM = TileMap->GetTilePositionInLocalSpace(CursorRange.Min.X, CursorRange.Min.Y, LastCursorTileZ) + RecenterOffset;
-
-					const FVector TL(ComponentToWorld.TransformPosition(LSTM));
-					const FVector TR(ComponentToWorld.TransformPosition(LSTM + PaperAxisX*TileWidthInUU*0.5f - PaperAxisY*TileHeightInUU*0.5f));
-					const FVector BL(ComponentToWorld.TransformPosition(LSTM - PaperAxisX*TileWidthInUU*0.5f - PaperAxisY*TileHeightInUU*0.5f));
-					const FVector BR(ComponentToWorld.TransformPosition(LSTM - PaperAxisY*TileHeightInUU*1.0f));
-
-					PDI->DrawLine(TL, TR, CursorWireColor, SDPG_Foreground, 0.0f, DepthBias);
-					PDI->DrawLine(TR, BR, CursorWireColor, SDPG_Foreground, 0.0f, DepthBias);
-					PDI->DrawLine(BR, BL, CursorWireColor, SDPG_Foreground, 0.0f, DepthBias);
-					PDI->DrawLine(BL, TL, CursorWireColor, SDPG_Foreground, 0.0f, DepthBias);
-				}
-				break;
 
-			case ETileMapProjectionMode::HexagonalStaggered:
+			TArray<FVector> TilePolygon;
+			TilePolygon.Empty(6);
+			for (int32 CY = CursorRange.Min.Y; CY < CursorRange.Max.Y; ++CY)
+			{
+				for (int32 CX = CursorRange.Min.X; CX < CursorRange.Max.X; ++CX)
 				{
-					const float UnrealUnitsPerPixel = TileMap->GetUnrealUnitsPerPixel();
-					const float TileWidthInUU = TileMap->TileWidth * UnrealUnitsPerPixel;
-					const float TileHeightInUU = TileMap->TileHeight * UnrealUnitsPerPixel;
-
-					const FVector RecenterOffset = PaperAxisX*TileWidthInUU*0.5f;
-					const FVector LSTM = TileMap->GetTilePositionInLocalSpace(CursorRange.Min.X, CursorRange.Min.Y, LastCursorTileZ) + RecenterOffset;
+					TilePolygon.Reset();
+					TileMap->GetTilePolygon(CX, CY, LastCursorTileZ, /*out*/ TilePolygon);
 
-					const float HexSideLengthInUU = TileMap->HexSideLength * UnrealUnitsPerPixel;
-
-					const FVector Top(ComponentToWorld.TransformPosition(LSTM - PaperAxisY*HexSideLengthInUU));
-
-					const FVector RightTop(ComponentToWorld.TransformPosition(LSTM + PaperAxisX*TileWidthInUU*0.5f - PaperAxisY*TileHeightInUU*0.5f));
-					const FVector LeftTop(ComponentToWorld.TransformPosition(LSTM - PaperAxisX*TileWidthInUU*0.5f - PaperAxisY*TileHeightInUU*0.5f));
-
-					const FVector RightBottom(ComponentToWorld.TransformPosition(LSTM + PaperAxisX*TileWidthInUU*0.5f - PaperAxisY*(TileHeightInUU*0.5f + HexSideLengthInUU)));
-					const FVector LeftBottom(ComponentToWorld.TransformPosition(LSTM - PaperAxisX*TileWidthInUU*0.5f - PaperAxisY*(TileHeightInUU*0.5f + HexSideLengthInUU)));
-
-					const FVector Bottom(ComponentToWorld.TransformPosition(LSTM - PaperAxisY*TileHeightInUU*1.0f));
-
-					PDI->DrawLine(Top, RightTop, CursorWireColor, SDPG_Foreground, 0.0f, DepthBias);
-					PDI->DrawLine(RightTop, RightBottom, CursorWireColor, SDPG_Foreground, 0.0f, DepthBias);
-					PDI->DrawLine(RightBottom, Bottom, CursorWireColor, SDPG_Foreground, 0.0f, DepthBias);
-					PDI->DrawLine(Bottom, LeftBottom, CursorWireColor, SDPG_Foreground, 0.0f, DepthBias);
-					PDI->DrawLine(LeftBottom, LeftTop, CursorWireColor, SDPG_Foreground, 0.0f, DepthBias);
-					PDI->DrawLine(LeftTop, Top, CursorWireColor, SDPG_Foreground, 0.0f, DepthBias);
+					FVector LastPositionWS = ComponentToWorld.TransformPosition(TilePolygon[TilePolygon.Num()-1]);
+					for (int32 VertexIndex = 0; VertexIndex < TilePolygon.Num(); ++VertexIndex)
+					{
+						const FVector ThisPositionWS = ComponentToWorld.TransformPosition(TilePolygon[VertexIndex]);
+						PDI->DrawLine(LastPositionWS, ThisPositionWS, CursorWireColor, SDPG_Foreground, 0.0f, DepthBias);
+						LastPositionWS = ThisPositionWS;
+					}
 				}
-				break;
-			};
+			}
 		}
 	}
 }
