Paper2D: Removed an obsolete TODO

diff --git a/Engine/Plugins/2D/Paper2D/Source/Paper2DEditor/Private/PaperEditorShared/SpriteGeometryEditing.cpp b/Engine/Plugins/2D/Paper2D/Source/Paper2DEditor/Private/PaperEditorShared/SpriteGeometryEditing.cpp
--- a/Engine/Plugins/2D/Paper2D/Source/Paper2DEditor/Private/PaperEditorShared/SpriteGeometryEditing.cpp
+++ b/Engine/Plugins/2D/Paper2D/Source/Paper2DEditor/Private/PaperEditorShared/SpriteGeometryEditing.cpp
@@ -414,7 +414,6 @@ void FSpriteGeometryEditingHelper::DrawGeometry_CanvasPass(FViewport& InViewport
 	const FVector MousePositionWorldSpace = View.PixelToWorld(InViewport.GetMouseX(), InViewport.GetMouseY(), 0);
 	const FVector2D MousePositionTextureSpace = EditorContext->WorldSpaceToTextureSpace(MousePositionWorldSpace);
 
-	//@TODO: Display 'action required' text in a bright color for Add Polygon mode
 	//@TODO: Move all of the line drawing to the PDI pass
 	FSpriteGeometryCollection& Geometry = GetGeometryChecked();
 
