Editor: Changed SEditorViewport to bind the default IsVisible implemenetation if a subclass did not provide one
Paper2D: Removed dummy IsVisible implementations to allow the default to take effect, improving performance when several editors are open but not all visible at once [UE-14759]

diff --git a/Engine/Plugins/2D/Paper2D/Source/Paper2DEditor/Private/FlipbookEditor/FlipbookEditor.cpp b/Engine/Plugins/2D/Paper2D/Source/Paper2DEditor/Private/FlipbookEditor/FlipbookEditor.cpp
--- a/Engine/Plugins/2D/Paper2D/Source/Paper2DEditor/Private/FlipbookEditor/FlipbookEditor.cpp
+++ b/Engine/Plugins/2D/Paper2D/Source/Paper2DEditor/Private/FlipbookEditor/FlipbookEditor.cpp
@@ -78,10 +78,6 @@ private:
 
 	// Viewport client
 	TSharedPtr<FFlipbookEditorViewportClient> EditorViewportClient;
-
-private:
-	// Returns true if the viewport is visible
-	bool IsVisible() const override;
 };
 
 void SFlipbookEditorViewport::Construct(const FArguments& InArgs)
@@ -134,16 +130,9 @@ TSharedRef<FEditorViewportClient> SFlipbookEditorViewport::MakeEditorViewportCli
 {
 	EditorViewportClient = MakeShareable(new FFlipbookEditorViewportClient(FlipbookBeingEdited));
 
-	EditorViewportClient->VisibilityDelegate.BindSP(this, &SFlipbookEditorViewport::IsVisible);
-
 	return EditorViewportClient.ToSharedRef();
 }
 
-bool SFlipbookEditorViewport::IsVisible() const
-{
-	return true;//@TODO: Determine this better so viewport ticking optimizations can take place
-}
-
 TSharedPtr<SWidget> SFlipbookEditorViewport::MakeViewportToolbar()
 {
 	return SNew(SFlipbookEditorViewportToolbar, SharedThis(this));
diff --git a/Engine/Plugins/2D/Paper2D/Source/Paper2DEditor/Private/SpriteEditor/SpriteEditor.cpp b/Engine/Plugins/2D/Paper2D/Source/Paper2DEditor/Private/SpriteEditor/SpriteEditor.cpp
--- a/Engine/Plugins/2D/Paper2D/Source/Paper2DEditor/Private/SpriteEditor/SpriteEditor.cpp
+++ b/Engine/Plugins/2D/Paper2D/Source/Paper2DEditor/Private/SpriteEditor/SpriteEditor.cpp
@@ -86,10 +86,6 @@ private:
 
 	// Viewport client
 	TSharedPtr<FSpriteEditorViewportClient> EditorViewportClient;
-
-private:
-	// Returns true if the viewport is visible
-	bool IsVisible() const override;
 };
 
 void SSpriteEditorViewport::Construct(const FArguments& InArgs, TSharedPtr<FSpriteEditor> InSpriteEditor)
@@ -177,16 +173,9 @@ TSharedRef<FEditorViewportClient> SSpriteEditorViewport::MakeEditorViewportClien
 {
 	EditorViewportClient = MakeShareable(new FSpriteEditorViewportClient(SpriteEditorPtr, SharedThis(this)));
 
-	EditorViewportClient->VisibilityDelegate.BindSP(this, &SSpriteEditorViewport::IsVisible);
-
 	return EditorViewportClient.ToSharedRef();
 }
 
-bool SSpriteEditorViewport::IsVisible() const
-{
-	return true;//@TODO: Determine this better so viewport ticking optimizations can take place
-}
-
 TSharedPtr<SWidget> SSpriteEditorViewport::MakeViewportToolbar()
 {
 	return SNew(SSpriteEditorViewportToolbar, SharedThis(this));
diff --git a/Engine/Plugins/2D/Paper2D/Source/Paper2DEditor/Private/TileMapEditing/TileMapEditor.cpp b/Engine/Plugins/2D/Paper2D/Source/Paper2DEditor/Private/TileMapEditing/TileMapEditor.cpp
--- a/Engine/Plugins/2D/Paper2D/Source/Paper2DEditor/Private/TileMapEditing/TileMapEditor.cpp
+++ b/Engine/Plugins/2D/Paper2D/Source/Paper2DEditor/Private/TileMapEditing/TileMapEditor.cpp
@@ -79,10 +79,6 @@ private:
 
 	// Viewport client
 	TSharedPtr<FTileMapEditorViewportClient> EditorViewportClient;
-
-private:
-	// Returns true if the viewport is visible
-	bool IsVisible() const override;
 };
 
 void STileMapEditorViewport::Construct(const FArguments& InArgs, TSharedPtr<FTileMapEditor> InTileMapEditor)
@@ -136,16 +132,9 @@ TSharedRef<FEditorViewportClient> STileMapEditorViewport::MakeEditorViewportClie
 {
 	EditorViewportClient = MakeShareable(new FTileMapEditorViewportClient(TileMapEditorPtr, SharedThis(this)));
 
-	EditorViewportClient->VisibilityDelegate.BindSP(this, &STileMapEditorViewport::IsVisible);
-
 	return EditorViewportClient.ToSharedRef();
 }
 
-bool STileMapEditorViewport::IsVisible() const
-{
-	return true;//@TODO: Determine this better so viewport ticking optimizations can take place
-}
-
 TSharedPtr<SWidget> STileMapEditorViewport::MakeViewportToolbar()
 {
 	return SNew(STileMapEditorViewportToolbar, SharedThis(this));
diff --git a/Engine/Plugins/2D/Paper2D/Source/Paper2DEditor/Private/TileSetEditor/SingleTileEditorViewport.cpp b/Engine/Plugins/2D/Paper2D/Source/Paper2DEditor/Private/TileSetEditor/SingleTileEditorViewport.cpp
--- a/Engine/Plugins/2D/Paper2D/Source/Paper2DEditor/Private/TileSetEditor/SingleTileEditorViewport.cpp
+++ b/Engine/Plugins/2D/Paper2D/Source/Paper2DEditor/Private/TileSetEditor/SingleTileEditorViewport.cpp
@@ -106,8 +106,6 @@ TSharedPtr<SWidget> SSingleTileEditorViewport::MakeViewportToolbar()
 
 TSharedRef<FEditorViewportClient> SSingleTileEditorViewport::MakeEditorViewportClient()
 {
-	TypedViewportClient->VisibilityDelegate.BindSP(this, &SSingleTileEditorViewport::IsVisible);
-
 	return TypedViewportClient.ToSharedRef();
 }
 
@@ -175,11 +173,6 @@ FText SSingleTileEditorViewport::GetTitleText() const
 	}
 }
 
-bool SSingleTileEditorViewport::IsVisible() const
-{
-	return true;//@TODO: Determine this better so viewport ticking optimizations can take place
-}
-
 //////////////////////////////////////////////////////////////////////////
 
 #undef LOCTEXT_NAMESPACE
\ No newline at end of file
diff --git a/Engine/Plugins/2D/Paper2D/Source/Paper2DEditor/Private/TileSetEditor/SingleTileEditorViewport.h b/Engine/Plugins/2D/Paper2D/Source/Paper2DEditor/Private/TileSetEditor/SingleTileEditorViewport.h
--- a/Engine/Plugins/2D/Paper2D/Source/Paper2DEditor/Private/TileSetEditor/SingleTileEditorViewport.h
+++ b/Engine/Plugins/2D/Paper2D/Source/Paper2DEditor/Private/TileSetEditor/SingleTileEditorViewport.h
@@ -38,8 +38,6 @@ public:
 protected:
 	FText GetTitleText() const;
 
-	bool IsVisible() const;
-
 private:
 	TSharedPtr<FSingleTileEditorViewportClient> TypedViewportClient;
 	FEdModeTileMap* TileMapEditor;
diff --git a/Engine/Source/Editor/UnrealEd/Private/SEditorViewport.cpp b/Engine/Source/Editor/UnrealEd/Private/SEditorViewport.cpp
--- a/Engine/Source/Editor/UnrealEd/Private/SEditorViewport.cpp
+++ b/Engine/Source/Editor/UnrealEd/Private/SEditorViewport.cpp
@@ -58,6 +58,12 @@ void SEditorViewport::Construct( const FArguments& InArgs )
 	];
 
 	TSharedRef<FEditorViewportClient> ViewportClient = MakeEditorViewportClient();
+
+	if (!ViewportClient->VisibilityDelegate.IsBound())
+	{
+		ViewportClient->VisibilityDelegate.BindSP(this, &SEditorViewport::IsVisible);
+	}
+
 	SceneViewport = MakeShareable( new FSceneViewport( &ViewportClient.Get(), ViewportWidget ) );
 	ViewportClient->Viewport = SceneViewport.Get();
 	ViewportWidget->SetViewportInterface(SceneViewport.ToSharedRef());
