Paper2D: Removed editor-only gating on registering the Paper2D custom version and performed as much back-compat loading code as is possible in cooked builds (UPaperSprite assets still cannot rebuild collision or render data and will warn on load if they are stale in a cooked build due to incremental cooking) [UE-14776]

diff --git a/Engine/Plugins/2D/Paper2D/Source/Paper2D/Classes/PaperFlipbookComponent.h b/Engine/Plugins/2D/Paper2D/Source/Paper2D/Classes/PaperFlipbookComponent.h
--- a/Engine/Plugins/2D/Paper2D/Source/Paper2D/Classes/PaperFlipbookComponent.h
+++ b/Engine/Plugins/2D/Paper2D/Source/Paper2D/Classes/PaperFlipbookComponent.h
@@ -165,10 +165,8 @@ protected:
 
 public:
 	// UObject interface
-#if WITH_EDITORONLY_DATA
 	virtual void Serialize(FArchive& Ar) override;
 	virtual void PostLoad() override;
-#endif
 	// End of UObject interface
 
 	// UActorComponent interface
diff --git a/Engine/Plugins/2D/Paper2D/Source/Paper2D/Classes/PaperSprite.h b/Engine/Plugins/2D/Paper2D/Source/Paper2D/Classes/PaperSprite.h
--- a/Engine/Plugins/2D/Paper2D/Source/Paper2D/Classes/PaperSprite.h
+++ b/Engine/Plugins/2D/Paper2D/Source/Paper2D/Classes/PaperSprite.h
@@ -162,14 +162,16 @@ public:
 
 public:
 
-#if WITH_EDITOR
 	// UObject interface
-	virtual void PostEditChangeProperty(FPropertyChangedEvent& PropertyChangedEvent) override;
 	virtual void Serialize(FArchive& Ar) override;
 	virtual void PostLoad() override;
+#if WITH_EDITOR
 	virtual void GetAssetRegistryTags(TArray<FAssetRegistryTag>& OutTags) const override;
+	virtual void PostEditChangeProperty(FPropertyChangedEvent& PropertyChangedEvent) override;
+#endif
 	// End of UObject interface
 
+#if WITH_EDITOR
 	FVector2D ConvertTextureSpaceToPivotSpace(FVector2D Input) const;
 	FVector2D ConvertPivotSpaceToTextureSpace(FVector2D Input) const;
 	FVector ConvertTextureSpaceToPivotSpace(FVector Input) const;
diff --git a/Engine/Plugins/2D/Paper2D/Source/Paper2D/Classes/PaperSpriteComponent.h b/Engine/Plugins/2D/Paper2D/Source/Paper2D/Classes/PaperSpriteComponent.h
--- a/Engine/Plugins/2D/Paper2D/Source/Paper2D/Classes/PaperSpriteComponent.h
+++ b/Engine/Plugins/2D/Paper2D/Source/Paper2D/Classes/PaperSpriteComponent.h
@@ -57,10 +57,8 @@ public:
 #if WITH_EDITOR
 	virtual void PostEditChangeProperty(FPropertyChangedEvent& PropertyChangedEvent) override;
 #endif
-#if WITH_EDITORONLY_DATA
 	virtual void Serialize(FArchive& Ar) override;
 	virtual void PostLoad() override;
-#endif
 	// End of UObject interface
 
 	// UActorComponent interface
diff --git a/Engine/Plugins/2D/Paper2D/Source/Paper2D/Classes/PaperTileMapComponent.h b/Engine/Plugins/2D/Paper2D/Source/Paper2D/Classes/PaperTileMapComponent.h
--- a/Engine/Plugins/2D/Paper2D/Source/Paper2D/Classes/PaperTileMapComponent.h
+++ b/Engine/Plugins/2D/Paper2D/Source/Paper2D/Classes/PaperTileMapComponent.h
@@ -91,10 +91,8 @@ protected:
 public:
 	// UObject interface
 	virtual void PostInitProperties() override;
-#if WITH_EDITORONLY_DATA
 	virtual void Serialize(FArchive& Ar) override;
 	virtual void PostLoad() override;
-#endif
 	// End of UObject interface
 
 	// UActorComponent interface
diff --git a/Engine/Plugins/2D/Paper2D/Source/Paper2D/Private/PaperFlipbookComponent.cpp b/Engine/Plugins/2D/Paper2D/Source/Paper2D/Private/PaperFlipbookComponent.cpp
--- a/Engine/Plugins/2D/Paper2D/Source/Paper2D/Private/PaperFlipbookComponent.cpp
+++ b/Engine/Plugins/2D/Paper2D/Source/Paper2D/Private/PaperFlipbookComponent.cpp
@@ -56,7 +56,6 @@ UPaperSprite* UPaperFlipbookComponent::GetSpriteAtCachedIndex() const
 }
 
 
-#if WITH_EDITORONLY_DATA
 void UPaperFlipbookComponent::Serialize(FArchive& Ar)
 {
 	Super::Serialize(Ar);
@@ -78,7 +77,6 @@ void UPaperFlipbookComponent::PostLoad()
 		}
 	}
 }
-#endif
 
 FPrimitiveSceneProxy* UPaperFlipbookComponent::CreateSceneProxy()
 {
diff --git a/Engine/Plugins/2D/Paper2D/Source/Paper2D/Private/PaperSprite.cpp b/Engine/Plugins/2D/Paper2D/Source/Paper2D/Private/PaperSprite.cpp
--- a/Engine/Plugins/2D/Paper2D/Source/Paper2D/Private/PaperSprite.cpp
+++ b/Engine/Plugins/2D/Paper2D/Source/Paper2D/Private/PaperSprite.cpp
@@ -15,6 +15,42 @@
 #include "UnrealEd.h"
 #endif
 
+#if WITH_EDITOR
+static void UpdateGeometryToBeBoxPositionRelative(FSpriteGeometryCollection& Geometry)
+{
+	// Make sure the per-shape GeometryType fields are up to date (introduced in this version)
+	const bool bWasBoundingBox = (Geometry.GeometryType == ESpritePolygonMode::SourceBoundingBox) || (Geometry.GeometryType == ESpritePolygonMode::TightBoundingBox);
+
+	if (bWasBoundingBox)
+	{
+		for (FSpriteGeometryShape& Shape : Geometry.Shapes)
+		{
+			Shape.ShapeType = ESpriteShapeType::Box;
+
+			// Recenter the bounding box (BoxPosition is now defined as the center)
+			const FVector2D AmountToSubtract = Shape.BoxPosition + Shape.BoxSize * 0.5f;
+			Shape.BoxPosition += Shape.BoxSize * 0.5f;
+			for (FVector2D& Vertex : Shape.Vertices)
+			{
+				Vertex -= AmountToSubtract;
+			}
+		}
+	}
+	else
+	{
+		for (FSpriteGeometryShape& Shape : Geometry.Shapes)
+		{
+			Shape.ShapeType = ESpriteShapeType::Polygon;
+
+			// Make sure BoxPosition is zeroed since polygon points are relative to it now, but it was being ignored
+			//@TODO: Consider computing the center and recentering verts to keep the numbers small/relative
+			Shape.BoxPosition = FVector2D::ZeroVector;
+			Shape.BoxSize = FVector2D::ZeroVector;
+		}
+	}
+}
+#endif
+
 #if WITH_EDITOR
 
 #include "PaperSpriteAtlas.h"
@@ -1705,7 +1741,6 @@ void UPaperSprite::RemoveSocket(FName SocketNameToDelete)
 }
 #endif
 
-#if WITH_EDITOR
 void UPaperSprite::Serialize(FArchive& Ar)
 {
 	Super::Serialize(Ar);
@@ -1713,52 +1748,22 @@ void UPaperSprite::Serialize(FArchive& Ar)
 	Ar.UsingCustomVersion(FPaperCustomVersion::GUID);
 }
 
-static void UpdateGeometryToBeBoxPositionRelative(FSpriteGeometryCollection& Geometry)
-{
-	// Make sure the per-shape GeometryType fields are up to date (introduced in this version)
-	const bool bWasBoundingBox = (Geometry.GeometryType == ESpritePolygonMode::SourceBoundingBox) || (Geometry.GeometryType == ESpritePolygonMode::TightBoundingBox);
-
-	if (bWasBoundingBox)
-	{
-		for (FSpriteGeometryShape& Shape : Geometry.Shapes)
-		{
-			Shape.ShapeType = ESpriteShapeType::Box;
-
-			// Recenter the bounding box (BoxPosition is now defined as the center)
-			const FVector2D AmountToSubtract = Shape.BoxPosition + Shape.BoxSize * 0.5f;
-			Shape.BoxPosition += Shape.BoxSize * 0.5f;
-			for (FVector2D& Vertex : Shape.Vertices)
-			{
-				Vertex -= AmountToSubtract;
-			}
-		}
-	}
-	else
-	{
-		for (FSpriteGeometryShape& Shape : Geometry.Shapes)
-		{
-			Shape.ShapeType = ESpriteShapeType::Polygon;
-
-			// Make sure BoxPosition is zeroed since polygon points are relative to it now, but it was being ignored
-			//@TODO: Consider computing the center and recentering verts to keep the numbers small/relative
-			Shape.BoxPosition = FVector2D::ZeroVector;
-			Shape.BoxSize = FVector2D::ZeroVector;
-		}
-	}
-}
-
-
 void UPaperSprite::PostLoad()
 {
 	Super::PostLoad();
+	
+	const int32 PaperVer = GetLinkerCustomVersion(FPaperCustomVersion::GUID);
 
-#if WITH_EDITORONLY_DATA
+#if !WITH_EDITORONLY_DATA
+	if (PaperVer < FPaperCustomVersion::LatestVersion)
+	{
+		UE_LOG(LogPaper2D, Warning, TEXT("Stale UPaperSprite asset '%s' with version %d detected in a cooked build (latest version is %d).  Please perform a full recook."), *GetPathName(), PaperVer, (int32)FPaperCustomVersion::LatestVersion);
+	}
+#else
 	if (UTexture2D* EffectiveTexture = GetBakedTexture())
 	{
 		EffectiveTexture->ConditionalPostLoad();
 	}
-	
-	const int32 PaperVer = GetLinkerCustomVersion(FPaperCustomVersion::GUID);
 
 	bool bRebuildCollision = false;
 	bool bRebuildRenderData = false;
@@ -1813,7 +1818,6 @@ void UPaperSprite::PostLoad()
 	}
 #endif
 }
-#endif
 
 UTexture2D* UPaperSprite::GetBakedTexture() const
 {
diff --git a/Engine/Plugins/2D/Paper2D/Source/Paper2D/Private/PaperSpriteComponent.cpp b/Engine/Plugins/2D/Paper2D/Source/Paper2D/Private/PaperSpriteComponent.cpp
--- a/Engine/Plugins/2D/Paper2D/Source/Paper2D/Private/PaperSpriteComponent.cpp
+++ b/Engine/Plugins/2D/Paper2D/Source/Paper2D/Private/PaperSpriteComponent.cpp
@@ -38,7 +38,6 @@ void UPaperSpriteComponent::PostEditChangeProperty(FPropertyChangedEvent& Proper
 }
 #endif
 
-#if WITH_EDITORONLY_DATA
 void UPaperSpriteComponent::Serialize(FArchive& Ar)
 {
 	Super::Serialize(Ar);
@@ -60,7 +59,6 @@ void UPaperSpriteComponent::PostLoad()
 		}
 	}
 }
-#endif
 
 FPrimitiveSceneProxy* UPaperSpriteComponent::CreateSceneProxy()
 {
diff --git a/Engine/Plugins/2D/Paper2D/Source/Paper2D/Private/PaperTileMapComponent.cpp b/Engine/Plugins/2D/Paper2D/Source/Paper2D/Private/PaperTileMapComponent.cpp
--- a/Engine/Plugins/2D/Paper2D/Source/Paper2D/Private/PaperTileMapComponent.cpp
+++ b/Engine/Plugins/2D/Paper2D/Source/Paper2D/Private/PaperTileMapComponent.cpp
@@ -87,7 +87,6 @@ FBoxSphereBounds UPaperTileMapComponent::CalcBounds(const FTransform& LocalToWor
 	}
 }
 
-#if WITH_EDITORONLY_DATA
 void UPaperTileMapComponent::Serialize(FArchive& Ar)
 {
 	Super::Serialize(Ar);
@@ -99,7 +98,6 @@ void UPaperTileMapComponent::PostLoad()
 {
 	Super::PostLoad();
 
-#if WITH_EDITORONLY_DATA
 	if (GetLinkerCustomVersion(FPaperCustomVersion::GUID) < FPaperCustomVersion::MovedTileMapDataToSeparateClass)
 	{
 		// Create a tile map object and move our old properties over to it
@@ -126,9 +124,7 @@ void UPaperTileMapComponent::PostLoad()
 		Material_DEPRECATED = nullptr;
 		TileLayers_DEPRECATED.Empty();
 	}
-#endif
 }
-#endif
 
 UBodySetup* UPaperTileMapComponent::GetBodySetup()
 {
