Tile Maps: Improved the grid line and cursor preview displays for all projection modes [UE-11488]

diff --git a/Engine/Plugins/2D/Paper2D/Source/Paper2D/Classes/PaperTileMap.h b/Engine/Plugins/2D/Paper2D/Source/Paper2D/Classes/PaperTileMap.h
--- a/Engine/Plugins/2D/Paper2D/Source/Paper2D/Classes/PaperTileMap.h
+++ b/Engine/Plugins/2D/Paper2D/Source/Paper2D/Classes/PaperTileMap.h
@@ -142,7 +142,7 @@ public:
 	FVector GetTileCenterInLocalSpace(float TileX, float TileY, int32 LayerIndex = 0) const;
 
 	// Returns the polygon for the specified tile (will be 4 or 6 vertices as a rectangle, diamond, or hexagon)
-	void GetTilePolygon(int32 TileX, int32 TileY, int32 LayerIndex, TArray<FVector>& LocalSpacePoints);
+	void GetTilePolygon(int32 TileX, int32 TileY, int32 LayerIndex, TArray<FVector>& LocalSpacePoints) const;
 
 	void GetTileToLocalParameters(FVector& OutCornerPosition, FVector& OutStepX, FVector& OutStepY, FVector& OutOffsetYFactor) const;
 	void GetLocalToTileParameters(FVector& OutCornerPosition, FVector& OutStepX, FVector& OutStepY, FVector& OutOffsetYFactor) const;
diff --git a/Engine/Plugins/2D/Paper2D/Source/Paper2D/Private/PaperTileMap.cpp b/Engine/Plugins/2D/Paper2D/Source/Paper2D/Private/PaperTileMap.cpp
--- a/Engine/Plugins/2D/Paper2D/Source/Paper2D/Private/PaperTileMap.cpp
+++ b/Engine/Plugins/2D/Paper2D/Source/Paper2D/Private/PaperTileMap.cpp
@@ -343,7 +343,7 @@ FVector UPaperTileMap::GetTilePositionInLocalSpace(float TileX, float TileY, int
 	return LocalPos;
 }
 
-void UPaperTileMap::GetTilePolygon(int32 TileX, int32 TileY, int32 LayerIndex, TArray<FVector>& LocalSpacePoints)
+void UPaperTileMap::GetTilePolygon(int32 TileX, int32 TileY, int32 LayerIndex, TArray<FVector>& LocalSpacePoints) const
 {
 	switch (ProjectionMode)
 	{
@@ -352,8 +352,8 @@ void UPaperTileMap::GetTilePolygon(int32 TileX, int32 TileY, int32 LayerIndex, T
 	default:
 		LocalSpacePoints.Add(GetTilePositionInLocalSpace(TileX, TileY, LayerIndex));
 		LocalSpacePoints.Add(GetTilePositionInLocalSpace(TileX + 1, TileY, LayerIndex));
-		LocalSpacePoints.Add(GetTilePositionInLocalSpace(TileX, TileY + 1, LayerIndex));
 		LocalSpacePoints.Add(GetTilePositionInLocalSpace(TileX + 1, TileY + 1, LayerIndex));
+		LocalSpacePoints.Add(GetTilePositionInLocalSpace(TileX, TileY + 1, LayerIndex));
 		break;
 
 	case ETileMapProjectionMode::IsometricStaggered:
@@ -378,20 +378,22 @@ void UPaperTileMap::GetTilePolygon(int32 TileX, int32 TileY, int32 LayerIndex, T
 			const float TileWidthInUU = TileWidth * UnrealUnitsPerPixel;
 			const float TileHeightInUU = TileHeight * UnrealUnitsPerPixel;
 
-			const FVector RecenterOffset = PaperAxisX*TileWidthInUU*0.5f;
-			const FVector LSTM = GetTilePositionInLocalSpace(TileX, TileY, LayerIndex) + RecenterOffset;
+			const FVector HalfWidth = PaperAxisX*TileWidthInUU*0.5f;
+			const FVector LSTM = GetTilePositionInLocalSpace(TileX, TileY, LayerIndex) + HalfWidth;
 
 			const float HexSideLengthInUU = HexSideLength * UnrealUnitsPerPixel;
+			const float HalfHexLength = HexSideLengthInUU*0.5f;
+			const FVector Top(LSTM - PaperAxisY*HalfHexLength);
 
-			const FVector Top(LSTM - PaperAxisY*HexSideLengthInUU);
-
-			const FVector RightTop(LSTM + PaperAxisX*TileWidthInUU*0.5f - PaperAxisY*TileHeightInUU*0.5f);
-			const FVector LeftTop(LSTM - PaperAxisX*TileWidthInUU*0.5f - PaperAxisY*TileHeightInUU*0.5f);
+			const FVector StepTopSides = PaperAxisY*(TileHeightInUU*0.5f - HalfHexLength);
+			const FVector RightTop(LSTM + HalfWidth - StepTopSides);
+			const FVector LeftTop(LSTM - HalfWidth - StepTopSides);
 
-			const FVector RightBottom(LSTM + PaperAxisX*TileWidthInUU*0.5f - PaperAxisY*(TileHeightInUU*0.5f + HexSideLengthInUU));
-			const FVector LeftBottom(LSTM - PaperAxisX*TileWidthInUU*0.5f - PaperAxisY*(TileHeightInUU*0.5f + HexSideLengthInUU));
+			const FVector StepBottomSides = PaperAxisY*(TileHeightInUU*0.5f + HalfHexLength);
+			const FVector RightBottom(LSTM + HalfWidth - StepBottomSides);
+			const FVector LeftBottom(LSTM - HalfWidth - StepBottomSides);
 
-			const FVector Bottom(LSTM - PaperAxisY*TileHeightInUU*1.0f);
+			const FVector Bottom(LSTM - PaperAxisY*(TileHeightInUU - HalfHexLength));
 
 			LocalSpacePoints.Add(Top);
 			LocalSpacePoints.Add(RightTop);
@@ -411,6 +413,10 @@ FVector UPaperTileMap::GetTileCenterInLocalSpace(float TileX, float TileY, int32
 	case ETileMapProjectionMode::Orthogonal:
 	default:
 		return GetTilePositionInLocalSpace(TileX + 0.5f, TileY + 0.5f, LayerIndex);
+
+	case ETileMapProjectionMode::IsometricDiamond:
+		return GetTilePositionInLocalSpace(TileX + 0.5f, TileY + 0.5f, LayerIndex);
+
 	case ETileMapProjectionMode::IsometricStaggered:
 		{
 			const float UnrealUnitsPerPixel = GetUnrealUnitsPerPixel();
@@ -428,7 +434,7 @@ FVector UPaperTileMap::GetTileCenterInLocalSpace(float TileX, float TileY, int32
 			const float TileHeightInUU = TileHeight * UnrealUnitsPerPixel;
 			const float HexSideLengthInUU = HexSideLength * UnrealUnitsPerPixel;
 
-			const FVector RecenterOffset = PaperAxisX*TileWidthInUU*0.5f - PaperAxisY*(HexSideLengthInUU + TileHeightInUU)*0.5f;
+			const FVector RecenterOffset = PaperAxisX*TileWidthInUU*0.5f - PaperAxisY*(TileHeightInUU)*0.5f;
 			return GetTilePositionInLocalSpace(TileX, TileY, LayerIndex) + RecenterOffset;
 		}
 	}
diff --git a/Engine/Plugins/2D/Paper2D/Source/Paper2D/Private/PaperTileMapComponent.cpp b/Engine/Plugins/2D/Paper2D/Source/Paper2D/Private/PaperTileMapComponent.cpp
--- a/Engine/Plugins/2D/Paper2D/Source/Paper2D/Private/PaperTileMapComponent.cpp
+++ b/Engine/Plugins/2D/Paper2D/Source/Paper2D/Private/PaperTileMapComponent.cpp
@@ -258,7 +258,7 @@ void UPaperTileMapComponent::RebuildRenderData(FPaperTileMapRenderSceneProxy* Pr
 				EffectiveTopLeftCorner = CornerOffset;
 				break;
 			case ETileMapProjectionMode::IsometricDiamond:
-				EffectiveTopLeftCorner = CornerOffset - StepPerTileX;
+				EffectiveTopLeftCorner = CornerOffset - 0.5f * StepPerTileX + 0.5f * StepPerTileY;
 				break;
 			case ETileMapProjectionMode::IsometricStaggered:
 			case ETileMapProjectionMode::HexagonalStaggered:
diff --git a/Engine/Plugins/2D/Paper2D/Source/Paper2D/Private/PaperTileMapRenderSceneProxy.cpp b/Engine/Plugins/2D/Paper2D/Source/Paper2D/Private/PaperTileMapRenderSceneProxy.cpp
--- a/Engine/Plugins/2D/Paper2D/Source/Paper2D/Private/PaperTileMapRenderSceneProxy.cpp
+++ b/Engine/Plugins/2D/Paper2D/Source/Paper2D/Private/PaperTileMapRenderSceneProxy.cpp
@@ -12,6 +12,7 @@
 #include "PaperTileMapComponent.h"
 
 DECLARE_CYCLE_STAT(TEXT("Tile Map Proxy"), STAT_TileMap_GetDynamicMeshElements, STATGROUP_Paper2D);
+DECLARE_CYCLE_STAT(TEXT("Tile Map Editor Grid"), STAT_TileMap_EditorWireDrawing, STATGROUP_Paper2D);
 
 //////////////////////////////////////////////////////////////////////////
 // FPaperTileMapRenderSceneProxy
@@ -24,6 +25,7 @@ FPaperTileMapRenderSceneProxy::FPaperTileMapRenderSceneProxy(const UPaperTileMap
 	, bShowOutlineWhenUnselected(false)
 #endif
 	, TileMap(nullptr)
+	, WireDepthBias(0.0001f)
 {
 	check(InComponent);
 
@@ -41,19 +43,160 @@ FPaperTileMapRenderSceneProxy::FPaperTileMapRenderSceneProxy(const UPaperTileMap
 
 void FPaperTileMapRenderSceneProxy::DrawBoundsForLayer(FPrimitiveDrawInterface* PDI, const FLinearColor& Color, int32 LayerIndex) const
 {
-	// Slight depth bias so that the wireframe grid overlay doesn't z-fight with the tiles themselves
-	const float DepthBias = 0.0001f;
-
 	const FMatrix& LocalToWorld = GetLocalToWorld();
 	const FVector TL(LocalToWorld.TransformPosition(TileMap->GetTilePositionInLocalSpace(0, 0, LayerIndex)));
 	const FVector TR(LocalToWorld.TransformPosition(TileMap->GetTilePositionInLocalSpace(TileMap->MapWidth, 0, LayerIndex)));
 	const FVector BL(LocalToWorld.TransformPosition(TileMap->GetTilePositionInLocalSpace(0, TileMap->MapHeight, LayerIndex)));
 	const FVector BR(LocalToWorld.TransformPosition(TileMap->GetTilePositionInLocalSpace(TileMap->MapWidth, TileMap->MapHeight, LayerIndex)));
 
-	PDI->DrawLine(TL, TR, Color, SDPG_Foreground, 0.0f, DepthBias);
-	PDI->DrawLine(TR, BR, Color, SDPG_Foreground, 0.0f, DepthBias);
-	PDI->DrawLine(BR, BL, Color, SDPG_Foreground, 0.0f, DepthBias);
-	PDI->DrawLine(BL, TL, Color, SDPG_Foreground, 0.0f, DepthBias);
+	PDI->DrawLine(TL, TR, Color, SDPG_Foreground, 0.0f, WireDepthBias);
+	PDI->DrawLine(TR, BR, Color, SDPG_Foreground, 0.0f, WireDepthBias);
+	PDI->DrawLine(BR, BL, Color, SDPG_Foreground, 0.0f, WireDepthBias);
+	PDI->DrawLine(BL, TL, Color, SDPG_Foreground, 0.0f, WireDepthBias);
+}
+
+void FPaperTileMapRenderSceneProxy::DrawNormalGridLines(FPrimitiveDrawInterface* PDI, const FLinearColor& Color, int32 LayerIndex) const
+{
+	const FMatrix& LocalToWorld = GetLocalToWorld();
+	const uint8 DPG = SDPG_Foreground;//GetDepthPriorityGroup(View);
+
+	// Draw horizontal lines on the selection
+	for (int32 Y = 0; Y <= TileMap->MapHeight; ++Y)
+	{
+		int32 X = 0;
+		const FVector Start(TileMap->GetTilePositionInLocalSpace(X, Y, LayerIndex));
+
+		X = TileMap->MapWidth;
+		const FVector End(TileMap->GetTilePositionInLocalSpace(X, Y, LayerIndex));
+
+		PDI->DrawLine(LocalToWorld.TransformPosition(Start), LocalToWorld.TransformPosition(End), Color, DPG, 0.0f, WireDepthBias);
+	}
+
+	// Draw vertical lines
+	for (int32 X = 0; X <= TileMap->MapWidth; ++X)
+	{
+		int32 Y = 0;
+		const FVector Start(TileMap->GetTilePositionInLocalSpace(X, Y, LayerIndex));
+
+		Y = TileMap->MapHeight;
+		const FVector End(TileMap->GetTilePositionInLocalSpace(X, Y, LayerIndex));
+
+		PDI->DrawLine(LocalToWorld.TransformPosition(Start), LocalToWorld.TransformPosition(End), Color, DPG, 0.0f, WireDepthBias);
+	}
+}
+
+void FPaperTileMapRenderSceneProxy::DrawStaggeredGridLines(FPrimitiveDrawInterface* PDI, const FLinearColor& Color, int32 LayerIndex) const
+{
+	TArray<FVector> Poly;
+	Poly.Empty(4);
+
+	const FMatrix& LocalToWorld = GetLocalToWorld();
+	const uint8 DPG = SDPG_Foreground;//GetDepthPriorityGroup(View);
+
+	FVector CornerPosition;
+	FVector OffsetYFactor;
+	FVector StepX;
+	FVector StepY;
+
+	TileMap->GetTileToLocalParameters(/*out*/ CornerPosition, /*out*/ StepX, /*out*/ StepY, /*out*/ OffsetYFactor);
+
+	const FVector PartialZ = (TileMap->SeparationPerLayer * LayerIndex) * PaperAxisZ;
+	const FVector TotalOffset = CornerPosition + PartialZ;
+
+	const bool bStaggerEven = false;
+
+	const FVector TopCenterStart = TotalOffset + StepX * 0.5f + (StepX * -0.5f) + StepY;
+	for (int32 X = 0-((TileMap->MapHeight+1)/2); X < TileMap->MapWidth; ++X)
+	{
+		int32 XTop = FMath::Max(X, 0);
+		int32 YTop = FMath::Max(-2 * X, 0);
+
+		if (X < 0)
+		{
+			XTop--;
+			YTop--;
+		}
+
+		// A is top of center top row cell
+		Poly.Reset();
+		TileMap->GetTilePolygon(XTop, YTop, LayerIndex, Poly);
+		const FVector LSA = Poly[0];
+
+		// Determine the bottom row cell
+		int32 YBottom = TileMap->MapHeight - 1;
+		int32 XBottom = X + (TileMap->MapHeight +  1) / 2;
+		const int32 XExcess = FMath::Max(XBottom - TileMap->MapWidth, 0);
+		XBottom -= XExcess;
+		YBottom -= XExcess * 2;
+
+ 		if (XBottom == TileMap->MapWidth)
+ 		{
+			YBottom -= ((TileMap->MapHeight & 1) != 0) ? 0 : 1;
+ 		}
+
+		// Bottom center
+		Poly.Reset();
+		TileMap->GetTilePolygon(XBottom, YBottom, LayerIndex, Poly);
+		const FVector LSB = Poly[2];
+
+		PDI->DrawLine(LocalToWorld.TransformPosition(LSA), LocalToWorld.TransformPosition(LSB), Color, DPG, 0.0f, WireDepthBias);
+	}
+
+	for (int32 X = 0; X < TileMap->MapWidth + ((TileMap->MapHeight + 1) / 2) + 1; ++X)
+	{
+		const int32 XTop = FMath::Min(X, TileMap->MapWidth);
+		const int32 YTop = FMath::Max(2 * (X - TileMap->MapWidth), 0);
+
+		// A is top center of top row cell
+		Poly.Reset();
+		TileMap->GetTilePolygon(XTop, YTop, LayerIndex, Poly);
+		const FVector LSA = Poly[0];
+
+		// Determine the bottom row cell
+		int32 YBottom = TileMap->MapHeight;
+		int32 XBottom = X - ((TileMap->MapHeight+1) / 2);
+		const int32 XExcess = FMath::Max(-XBottom, 0);
+		XBottom += XExcess;
+		YBottom -= XExcess * 2;
+
+		if (XExcess > 0)
+		{
+			YBottom += (TileMap->MapHeight & 1);
+		}
+
+		// Bottom left
+		Poly.Reset();
+		TileMap->GetTilePolygon(XBottom, YBottom, LayerIndex, Poly);
+		const FVector LSB = Poly[3];
+
+		PDI->DrawLine(LocalToWorld.TransformPosition(LSA), LocalToWorld.TransformPosition(LSB), Color, DPG, 0.0f, WireDepthBias);
+	}
+}
+
+void FPaperTileMapRenderSceneProxy::DrawHexagonalGridLines(FPrimitiveDrawInterface* PDI, const FLinearColor& Color, int32 LayerIndex) const
+{
+	//@TODO: This isn't very efficient
+	const FMatrix& LocalToWorld = GetLocalToWorld();
+	const uint8 DPG = SDPG_Foreground;//GetDepthPriorityGroup(View);
+
+	TArray<FVector> Poly;
+	Poly.Empty(6);
+	for (int32 Y = 0; Y < TileMap->MapHeight; ++Y)
+	{
+		for (int32 X = 0; X < TileMap->MapWidth; ++X)
+		{
+			Poly.Reset();
+			TileMap->GetTilePolygon(X, Y, LayerIndex, Poly);
+
+			FVector LastVertexWS = LocalToWorld.TransformPosition(Poly[5]);
+			for (int32 VI = 0; VI < Poly.Num(); ++VI)
+			{
+				FVector ThisVertexWS = LocalToWorld.TransformPosition(Poly[VI]);
+				PDI->DrawLine(LastVertexWS, ThisVertexWS, Color, DPG, 0.0f, WireDepthBias);
+				LastVertexWS = ThisVertexWS;
+			}
+		}
+	}
 }
 
 void FPaperTileMapRenderSceneProxy::GetDynamicMeshElements(const TArray<const FSceneView*>& Views, const FSceneViewFamily& ViewFamily, uint32 VisibilityMap, FMeshElementCollector& Collector) const
@@ -61,11 +204,10 @@ void FPaperTileMapRenderSceneProxy::GetDynamicMeshElements(const TArray<const FS
 	SCOPE_CYCLE_COUNTER(STAT_TileMap_GetDynamicMeshElements);
 	checkSlow(IsInRenderingThread());
 
-	// Slight depth bias so that the wireframe grid overlay doesn't z-fight with the tiles themselves
-	const float DepthBias = 0.0001f;
-
 	for (int32 ViewIndex = 0; ViewIndex < Views.Num(); ViewIndex++)
 	{
+		SCOPE_CYCLE_COUNTER(STAT_TileMap_EditorWireDrawing);
+
 		if (VisibilityMap & (1 << ViewIndex))
 		{
 			const FSceneView* View = Views[ViewIndex];
@@ -155,28 +297,19 @@ void FPaperTileMapRenderSceneProxy::GetDynamicMeshElements(const TArray<const FS
 
 					if (bShowPerTileGrid && (SelectedLayerIndex != INDEX_NONE))
 					{
-						// Draw horizontal lines on the selection
-						for (int32 Y = 0; Y <= TileMap->MapHeight; ++Y)
+						switch (TileMap->ProjectionMode)
 						{
-							int32 X = 0;
-							const FVector Start(TileMap->GetTilePositionInLocalSpace(X, Y, SelectedLayerIndex));
-
-							X = TileMap->MapWidth;
-							const FVector End(TileMap->GetTilePositionInLocalSpace(X, Y, SelectedLayerIndex));
-
-							PDI->DrawLine(LocalToWorld.TransformPosition(Start), LocalToWorld.TransformPosition(End), OverrideColor, DPG, 0.0f, DepthBias);
-						}
-
-						// Draw vertical lines
-						for (int32 X = 0; X <= TileMap->MapWidth; ++X)
-						{
-							int32 Y = 0;
-							const FVector Start(TileMap->GetTilePositionInLocalSpace(X, Y, SelectedLayerIndex));
-
-							Y = TileMap->MapHeight;
-							const FVector End(TileMap->GetTilePositionInLocalSpace(X, Y, SelectedLayerIndex));
-
-							PDI->DrawLine(LocalToWorld.TransformPosition(Start), LocalToWorld.TransformPosition(End), OverrideColor, DPG, 0.0f, DepthBias);
+						default:
+						case ETileMapProjectionMode::Orthogonal:
+						case ETileMapProjectionMode::IsometricDiamond:
+							DrawNormalGridLines(PDI, OverrideColor, SelectedLayerIndex);
+							break;
+						case ETileMapProjectionMode::IsometricStaggered:
+							DrawStaggeredGridLines(PDI, OverrideColor, SelectedLayerIndex);
+							break;
+						case ETileMapProjectionMode::HexagonalStaggered:
+							DrawHexagonalGridLines(PDI, OverrideColor, SelectedLayerIndex);
+							break;
 						}
 					}
 				}
diff --git a/Engine/Plugins/2D/Paper2D/Source/Paper2D/Private/PaperTileMapRenderSceneProxy.h b/Engine/Plugins/2D/Paper2D/Source/Paper2D/Private/PaperTileMapRenderSceneProxy.h
--- a/Engine/Plugins/2D/Paper2D/Source/Paper2D/Private/PaperTileMapRenderSceneProxy.h
+++ b/Engine/Plugins/2D/Paper2D/Source/Paper2D/Private/PaperTileMapRenderSceneProxy.h
@@ -23,6 +23,9 @@ public:
 
 protected:
 	void DrawBoundsForLayer(FPrimitiveDrawInterface* PDI, const FLinearColor& Color, int32 LayerIndex) const;
+	void DrawNormalGridLines(FPrimitiveDrawInterface* PDI, const FLinearColor& Color, int32 LayerIndex) const;
+	void DrawStaggeredGridLines(FPrimitiveDrawInterface* PDI, const FLinearColor& Color, int32 LayerIndex) const;
+	void DrawHexagonalGridLines(FPrimitiveDrawInterface* PDI, const FLinearColor& Color, int32 LayerIndex) const;
 
 protected:
 
@@ -34,4 +37,7 @@ protected:
 
 	//@TODO: Not thread safe
 	const class UPaperTileMap* TileMap;
+
+	// Slight depth bias so that the wireframe grid overlay doesn't z-fight with the tiles themselves
+	const float WireDepthBias;
 };
