Tile Maps: Added analysis of the default tile sheet texture to determine if a masked or translucent material should be used [UE-14420]
- Changed FSpriteAssetInitParams.Offset and Dimension to be FIntPoint instead of FVector2D, along with various other sprite texture region code
- Added const to a number of FAlphaBitmap methods
- Split texture analysis and material selection code out into separate methods in UPaperImporterSettings

diff --git a/Engine/Plugins/2D/Paper2D/Source/Paper2D/Classes/PaperSprite.h b/Engine/Plugins/2D/Paper2D/Source/Paper2D/Classes/PaperSprite.h
--- a/Engine/Plugins/2D/Paper2D/Source/Paper2D/Classes/PaperSprite.h
+++ b/Engine/Plugins/2D/Paper2D/Source/Paper2D/Classes/PaperSprite.h
@@ -41,11 +41,11 @@ class PAPER2D_API UPaperSprite : public UObject, public IInterface_CollisionData
 protected:
 
 #if WITH_EDITORONLY_DATA
-	// Position with SourceTexture
+	// Position with SourceTexture (in pixels)
 	UPROPERTY(Category=Sprite, EditAnywhere, AssetRegistrySearchable)
 	FVector2D SourceUV;
 
-	// Dimensions within SourceTexture
+	// Dimensions within SourceTexture (in pixels)
 	UPROPERTY(Category=Sprite, EditAnywhere, AssetRegistrySearchable)
 	FVector2D SourceDimension;
 
diff --git a/Engine/Plugins/2D/Paper2D/Source/Paper2D/Classes/PaperTileMap.h b/Engine/Plugins/2D/Paper2D/Source/Paper2D/Classes/PaperTileMap.h
--- a/Engine/Plugins/2D/Paper2D/Source/Paper2D/Classes/PaperTileMap.h
+++ b/Engine/Plugins/2D/Paper2D/Source/Paper2D/Classes/PaperTileMap.h
@@ -178,7 +178,7 @@ public:
 
 	// Called when a fresh tile map has been created (by factory or otherwise)
 	// Adds a default layer and pulls the PixelsPerUnrealUnit from the project settings
-	void InitializeNewEmptyTileMap();
+	void InitializeNewEmptyTileMap(UPaperTileSet* DefaultTileSetAsset = nullptr);
 
 	// Creates a clone of this tile map in the specified outer
 	UPaperTileMap* CloneTileMap(UObject* OuterForClone);
diff --git a/Engine/Plugins/2D/Paper2D/Source/Paper2D/Classes/SpriteEditorOnlyTypes.h b/Engine/Plugins/2D/Paper2D/Source/Paper2D/Classes/SpriteEditorOnlyTypes.h
--- a/Engine/Plugins/2D/Paper2D/Source/Paper2D/Classes/SpriteEditorOnlyTypes.h
+++ b/Engine/Plugins/2D/Paper2D/Source/Paper2D/Classes/SpriteEditorOnlyTypes.h
@@ -256,8 +256,8 @@ struct FSpriteAssetInitParameters
 
 	FSpriteAssetInitParameters()
 		: Texture(nullptr)
-		, Offset(FVector2D::ZeroVector)
-		, Dimension(FVector2D::ZeroVector)
+		, Offset(ForceInitToZero)
+		, Dimension(ForceInitToZero)
 		, bOverridePixelsPerUnrealUnit(false)
 		, PixelsPerUnrealUnit(1.0f)
 		, DefaultMaterialOverride(nullptr)
@@ -271,13 +271,13 @@ struct FSpriteAssetInitParameters
 		Texture = InTexture;
 		if (Texture != nullptr)
 		{
-			Dimension = FVector2D(Texture->GetImportedSize());
-			Offset = FVector2D::ZeroVector;
+			Dimension = Texture->GetImportedSize();
+			Offset = FIntPoint::ZeroValue;
 		}
 		else
 		{
-			Dimension = FVector2D::ZeroVector;
-			Offset = FVector2D::ZeroVector;
+			Dimension = FIntPoint::ZeroValue;
+			Offset = FIntPoint::ZeroValue;
 		}
 	}
 
@@ -294,10 +294,10 @@ public:
 	TArray<UTexture*> AdditionalTextures;
 
 	// The offset within the texture (in pixels)
-	FVector2D Offset;
+	FIntPoint Offset;
 
 	// The dimension of the subregion within the texture (in pixels)
-	FVector2D Dimension;
+	FIntPoint Dimension;
 
 	// Should we apply the PixelsPerUnrealUnit, or leave the value as it was
 	bool bOverridePixelsPerUnrealUnit;
diff --git a/Engine/Plugins/2D/Paper2D/Source/Paper2D/Private/PaperTileMap.cpp b/Engine/Plugins/2D/Paper2D/Source/Paper2D/Private/PaperTileMap.cpp
--- a/Engine/Plugins/2D/Paper2D/Source/Paper2D/Private/PaperTileMap.cpp
+++ b/Engine/Plugins/2D/Paper2D/Source/Paper2D/Private/PaperTileMap.cpp
@@ -542,9 +542,15 @@ void UPaperTileMap::ResizeMap(int32 NewWidth, int32 NewHeight, bool bForceResize
 	}
 }
 
-void UPaperTileMap::InitializeNewEmptyTileMap()
+void UPaperTileMap::InitializeNewEmptyTileMap(UPaperTileSet* InitialTileSet)
 {
-	//@TODO: Consider removing this function / making it a 'prepare' method just like the sprite one on the importer settings
+	if (InitialTileSet != nullptr)
+	{
+		TileWidth = InitialTileSet->TileWidth;
+		TileHeight = InitialTileSet->TileHeight;
+		SelectedTileSet = InitialTileSet;
+	}
+
 	AddNewLayer();
 }
 
diff --git a/Engine/Plugins/2D/Paper2D/Source/Paper2D/Private/PaperTileMapComponent.cpp b/Engine/Plugins/2D/Paper2D/Source/Paper2D/Private/PaperTileMapComponent.cpp
--- a/Engine/Plugins/2D/Paper2D/Source/Paper2D/Private/PaperTileMapComponent.cpp
+++ b/Engine/Plugins/2D/Paper2D/Source/Paper2D/Private/PaperTileMapComponent.cpp
@@ -390,11 +390,6 @@ void UPaperTileMapComponent::CreateNewOwnedTileMap()
 	NewTileMap->SetFlags(RF_Transactional);
 	NewTileMap->InitializeNewEmptyTileMap();
 
-	if (TileMap != nullptr)
-	{
-		NewTileMap->SelectedTileSet = TileMap->SelectedTileSet;
-	}
-
 	SetTileMap(NewTileMap);
 }
 
diff --git a/Engine/Plugins/2D/Paper2D/Source/Paper2D/Public/AlphaBitmap.h b/Engine/Plugins/2D/Paper2D/Source/Paper2D/Public/AlphaBitmap.h
--- a/Engine/Plugins/2D/Paper2D/Source/Paper2D/Public/AlphaBitmap.h
+++ b/Engine/Plugins/2D/Paper2D/Source/Paper2D/Public/AlphaBitmap.h
@@ -7,7 +7,7 @@
 
 struct PAPER2D_API FAlphaBitmap
 {
-	FAlphaBitmap(UTexture2D* SourceTexture, uint8 InDefaultValue = 0)
+	FAlphaBitmap(UTexture* SourceTexture, uint8 InDefaultValue = 0)
 		: Width(0)
 		, Height(0)
 		, DefaultValue(InDefaultValue)
@@ -15,7 +15,7 @@ struct PAPER2D_API FAlphaBitmap
 		ExtractFromTexture(SourceTexture);
 	}
 
-	void ExtractFromTexture(UTexture2D* SourceTexture)
+	void ExtractFromTexture(UTexture* SourceTexture)
 	{
 		// use the source art if it exists
 		FTextureSource* TextureSource = nullptr;
@@ -135,7 +135,7 @@ struct PAPER2D_API FAlphaBitmap
 	}
 
 	// Is the rectangle empty (X0..X1, Y0..Y1 inclusive)
-	bool IsRegionEqual(int32 X0, int32 Y0, int32 X1, int32 Y1, int32 Target)
+	bool IsRegionEqual(int32 X0, int32 Y0, int32 X1, int32 Y1, int32 Target) const
 	{
 		bool bEqual = true;
 		for (int32 Y = Y0; (Y < Y1) && bEqual; ++Y)
@@ -150,7 +150,7 @@ struct PAPER2D_API FAlphaBitmap
 		return IsColumnEqual(X, Y0, Y1, 0);
 	}
 
-	bool IsRowEmpty(int X0, int X1, int Y)
+	bool IsRowEmpty(int X0, int X1, int Y) const
 	{
 		return IsRowEqual(X0, X1, Y, 0);
 	}
@@ -162,7 +162,7 @@ struct PAPER2D_API FAlphaBitmap
 	}
 
 	// Returns the tight bounding box around pixels that are not 0
-	void GetTightBounds(FIntPoint& OutOrigin, FIntPoint& OutDimension)
+	void GetTightBounds(FIntPoint& OutOrigin, FIntPoint& OutDimension) const
 	{
 		OutOrigin.X = 0;
 		OutOrigin.Y = 0;
@@ -172,7 +172,7 @@ struct PAPER2D_API FAlphaBitmap
 	}
 
 	// Returns the tight bounding box around pixels that are not 0
-	void TightenBounds(FIntPoint& InOutOrigin, FIntPoint& InOutDimension)
+	void TightenBounds(FIntPoint& InOutOrigin, FIntPoint& InOutDimension) const
 	{
 		int Top = InOutOrigin.Y;
 		int Bottom = InOutOrigin.Y + InOutDimension.Y - 1;
@@ -237,7 +237,7 @@ struct PAPER2D_API FAlphaBitmap
 	}
 
 	// Wind through the bitmap from StartX, StartY to find the closest hit point
-	bool FoundClosestValidPoint(int StartX, int StartY, const int MaxAllowedSearchDistance, FIntPoint& OutHitPoint)
+	bool FoundClosestValidPoint(int StartX, int StartY, const int MaxAllowedSearchDistance, FIntPoint& OutHitPoint) const
 	{
 		// Constrain the point within the image bounds
 		StartX = FMath::Clamp(StartX, 0, Width - 1);
@@ -323,7 +323,7 @@ struct PAPER2D_API FAlphaBitmap
 	}
 
 	// Checks the image to determine if it is suitable for opaque, masked or translucent rendering
-	void AnalyzeImage(int32 StartX, int32 StartY, int32 InWidth, int32 InHeight, bool& bOutHasZeros, bool& bOutHasIntermediateValues)
+	void AnalyzeImage(int32 StartX, int32 StartY, int32 InWidth, int32 InHeight, bool& bOutHasZeros, bool& bOutHasIntermediateValues) const
 	{
 		bOutHasZeros = false;
 		bOutHasIntermediateValues = false;
diff --git a/Engine/Plugins/2D/Paper2D/Source/Paper2DEditor/Classes/PaperImporterSettings.h b/Engine/Plugins/2D/Paper2D/Source/Paper2DEditor/Classes/PaperImporterSettings.h
--- a/Engine/Plugins/2D/Paper2D/Source/Paper2DEditor/Classes/PaperImporterSettings.h
+++ b/Engine/Plugins/2D/Paper2D/Source/Paper2DEditor/Classes/PaperImporterSettings.h
@@ -4,6 +4,9 @@
 
 #include "PaperImporterSettings.generated.h"
 
+class UPaperTileSet;
+class UPaperTileMap;
+
 enum class ESpriteInitMaterialLightingMode : uint8
 {
 	// Use the default in the importer settings
@@ -44,15 +47,19 @@ class PAPER2DEDITOR_API UPaperImporterSettings : public UObject
 
 protected:
 	// Should the source texture be scanned when creating new sprites to determine the appropriate material? (if false, the Default Masked Material is always used)
-	UPROPERTY(config, EditAnywhere, Category=Settings)
-	bool bPickBestMaterialWhenCreatingSprite;
+	UPROPERTY(config, EditAnywhere, Category=NewAssetSettings)
+	bool bPickBestMaterialWhenCreatingSprites;
+
+	// Should the source texture be scanned when creating new tile maps (from a tile set or via importing) to determine the appropriate material? (if false, the Default Masked Material is always used)
+	UPROPERTY(config, EditAnywhere, Category=NewAssetSettings)
+	bool bPickBestMaterialWhenCreatingTileMaps;
 
 	// Can opaque materials be applied as part of the 'best material' analysis?
-	UPROPERTY(config, EditAnywhere, Category=Settings, meta=(EditCondition=bPickBestMaterialWhenCreatingSprite))
+	UPROPERTY(config, EditAnywhere, Category=NewAssetSettings)
 	bool bAnalysisCanUseOpaque;
 
 	// The default scaling factor between pixels and Unreal units (cm) to use for newly created sprite assets (e.g., 0.64 would make a 64 pixel wide sprite take up 100 cm)
-	UPROPERTY(config, EditAnywhere, Category=Settings)
+	UPROPERTY(config, EditAnywhere, Category=NewAssetSettings)
 	float DefaultPixelsPerUnrealUnit;
 
 	// A list of default suffixes to use when looking for associated normal maps while importing sprites or creating sprites from textures
@@ -103,6 +110,9 @@ protected:
 public:
 	UPaperImporterSettings();
 
+	// Should the source texture be scanned when creating new tile maps (from a tile set or via importing) to determine the appropriate material? (if false, the Default Masked Material is always used)
+	bool ShouldPickBestMaterialWhenCreatingTileMaps() const { return bPickBestMaterialWhenCreatingTileMaps; }
+
 	// Removes the suffix from the specified name if it matches something in BaseMapTextureSuffixes
 	FString RemoveSuffixFromBaseMapName(const FString& InName) const;
 
@@ -112,10 +122,18 @@ public:
 	// Applies the compression settings to the specified texture
 	void ApplyTextureSettings(UTexture2D* Texture) const;
 
-	// Applies the material settings to the 
+	// Fills out the sprite init parameters with the default settings given the desired material type and lighting mode (which can both be automatic)
 	// Note: This should be called after the texture has been set, as that will be analyzed if the lighting/material type flags are set to automatic
 	void ApplySettingsForSpriteInit(FSpriteAssetInitParameters& InitParams, ESpriteInitMaterialLightingMode LightingMode = ESpriteInitMaterialLightingMode::Automatic, ESpriteInitMaterialType MaterialTypeMode = ESpriteInitMaterialType::Automatic) const;
 
+	// Fills out the tile map with the default settings given the desired material type and lighting mode (which can both be automatic)
+	void ApplySettingsForTileMapInit(UPaperTileMap* TileMap, UPaperTileSet* DefaultTileSet = nullptr, ESpriteInitMaterialLightingMode LightingMode = ESpriteInitMaterialLightingMode::Automatic, ESpriteInitMaterialType MaterialTypeMode = ESpriteInitMaterialType::Automatic, bool bCreateEmptyLayer = true) const;
+
+	// Analyzes the specified texture in the Offset..Offset+Dimensions region and returns the best kind of material
+	// to represent the alpha content in the texture (typically masked or translucent, but can return opaque if bAnalysisCanUseOpaque is true
+	// Note: Will return Automatic if the texture is null.
+	ESpriteInitMaterialType AnalyzeTextureForDesiredMaterialType(UTexture* Texture, const FIntPoint& Offset, const FIntPoint& Dimensions) const;
+
 	// Returns the default translucent material
 	UMaterialInterface* GetDefaultTranslucentMaterial(bool bLit) const;
 
@@ -125,6 +143,10 @@ public:
 	// Returns the default masked material
 	UMaterialInterface* GetDefaultMaskedMaterial(bool bLit) const;
 
+	// Returns the default material for the specified material type.
+	// Input should be Masked, Opaque, or Translucent.  Automatic and LeaveAsIs will be treated like masked!
+	UMaterialInterface* GetDefaultMaterial(ESpriteInitMaterialType MaterialType, bool bUseLitMaterial) const;
+
 	// Returns the default pixels/uu setting
 	float GetDefaultPixelsPerUnrealUnit() const { return DefaultPixelsPerUnrealUnit; }
 };
diff --git a/Engine/Plugins/2D/Paper2D/Source/Paper2DEditor/Classes/PaperSpriteFactory.h b/Engine/Plugins/2D/Paper2D/Source/Paper2DEditor/Classes/PaperSpriteFactory.h
--- a/Engine/Plugins/2D/Paper2D/Source/Paper2DEditor/Classes/PaperSpriteFactory.h
+++ b/Engine/Plugins/2D/Paper2D/Source/Paper2DEditor/Classes/PaperSpriteFactory.h
@@ -18,8 +18,8 @@ class UPaperSpriteFactory : public UFactory
 
 	// Set bUseSourceRegion to get it to use/set initial sourceUV and dimensions
 	bool bUseSourceRegion;
-	FVector2D InitialSourceUV;
-	FVector2D InitialSourceDimension;
+	FIntPoint InitialSourceUV;
+	FIntPoint InitialSourceDimension;
 
 	// UFactory interface
 	virtual bool ConfigureProperties() override;
diff --git a/Engine/Plugins/2D/Paper2D/Source/Paper2DEditor/Private/ExtractSprites/SPaperExtractSpritesDialog.cpp b/Engine/Plugins/2D/Paper2D/Source/Paper2DEditor/Private/ExtractSprites/SPaperExtractSpritesDialog.cpp
--- a/Engine/Plugins/2D/Paper2D/Source/Paper2DEditor/Private/ExtractSprites/SPaperExtractSpritesDialog.cpp
+++ b/Engine/Plugins/2D/Paper2D/Source/Paper2DEditor/Private/ExtractSprites/SPaperExtractSpritesDialog.cpp
@@ -444,8 +444,8 @@ void SPaperExtractSpritesDialog::CreateExtractedSprites()
 
 		SpriteFactory->bUseSourceRegion = true;
 		const FIntRect& ExtractedRect = ExtractedSprite.Rect;
-		SpriteFactory->InitialSourceUV = FVector2D(ExtractedRect.Min.X, ExtractedRect.Min.Y);
-		SpriteFactory->InitialSourceDimension = FVector2D(ExtractedRect.Width(), ExtractedRect.Height());
+		SpriteFactory->InitialSourceUV = ExtractedRect.Min;
+		SpriteFactory->InitialSourceDimension = FIntPoint(ExtractedRect.Width(), ExtractedRect.Height());
 
 		// Get a unique name for the sprite
 		// Extracted sprite name is a name, we insert a _ as we're appending this to the texture name
diff --git a/Engine/Plugins/2D/Paper2D/Source/Paper2DEditor/Private/PaperImporterSettings.cpp b/Engine/Plugins/2D/Paper2D/Source/Paper2DEditor/Private/PaperImporterSettings.cpp
--- a/Engine/Plugins/2D/Paper2D/Source/Paper2DEditor/Private/PaperImporterSettings.cpp
+++ b/Engine/Plugins/2D/Paper2D/Source/Paper2DEditor/Private/PaperImporterSettings.cpp
@@ -3,12 +3,14 @@
 #include "Paper2DEditorPrivatePCH.h"
 #include "PaperImporterSettings.h"
 #include "AlphaBitmap.h"
+#include "TileMapEditing/TileMapEditorSettings.h"
 
 //////////////////////////////////////////////////////////////////////////
 // UPaperImporterSettings
 
 UPaperImporterSettings::UPaperImporterSettings()
-	: bPickBestMaterialWhenCreatingSprite(true)
+	: bPickBestMaterialWhenCreatingSprites(true)
+	, bPickBestMaterialWhenCreatingTileMaps(true)
 	, DefaultPixelsPerUnrealUnit(1.0f)
 	, DefaultSpriteTextureGroup(TEXTUREGROUP_Pixels2D)
 	, bOverrideTextureCompression(true)
@@ -76,6 +78,29 @@ void UPaperImporterSettings::ApplyTextureSettings(UTexture2D* Texture) const
 	}
 }
 
+ESpriteInitMaterialType UPaperImporterSettings::AnalyzeTextureForDesiredMaterialType(UTexture* Texture, const FIntPoint& Offset, const FIntPoint& Dimensions) const
+{
+	// Analyze the texture if desired (to see if it's got greyscale alpha or just binary alpha, picking either a translucent or masked material)
+	if (Texture != nullptr)
+	{
+		FAlphaBitmap AlphaBitmap(Texture);
+		bool bHasIntermediateValues;
+		bool bHasZeros;
+		AlphaBitmap.AnalyzeImage(Offset.X, Offset.Y, Dimensions.X, Dimensions.Y, /*out*/ bHasZeros, /*out*/ bHasIntermediateValues);
+
+		if (bAnalysisCanUseOpaque && !bHasIntermediateValues && !bHasZeros)
+		{
+			return ESpriteInitMaterialType::Opaque;
+		}
+		else
+		{
+			return bHasIntermediateValues ? ESpriteInitMaterialType::Translucent : ESpriteInitMaterialType::Masked;
+		}
+	}
+
+	return ESpriteInitMaterialType::Automatic;
+}
+
 void UPaperImporterSettings::ApplySettingsForSpriteInit(FSpriteAssetInitParameters& InitParams, ESpriteInitMaterialLightingMode LightingMode, ESpriteInitMaterialType MaterialTypeMode) const
 {
 	InitParams.SetPixelsPerUnrealUnit(DefaultPixelsPerUnrealUnit);
@@ -84,21 +109,9 @@ void UPaperImporterSettings::ApplySettingsForSpriteInit(FSpriteAssetInitParamete
 	if (DesiredMaterialType == ESpriteInitMaterialType::Automatic)
 	{
 		// Analyze the texture if desired (to see if it's got greyscale alpha or just binary alpha, picking either a translucent or masked material)
-		if (bPickBestMaterialWhenCreatingSprite && (InitParams.Texture != nullptr))
+		if (bPickBestMaterialWhenCreatingSprites)
 		{
-			FAlphaBitmap AlphaBitmap(InitParams.Texture);
-			bool bHasIntermediateValues;
-			bool bHasZeros;
-			AlphaBitmap.AnalyzeImage((int32)InitParams.Offset.X, (int32)InitParams.Offset.Y, (int32)InitParams.Dimension.X, (int32)InitParams.Dimension.Y, /*out*/ bHasZeros, /*out*/ bHasIntermediateValues);
-
-			if (bAnalysisCanUseOpaque && !bHasIntermediateValues && !bHasZeros)
-			{
-				DesiredMaterialType = ESpriteInitMaterialType::Opaque;
-			}
-			else
-			{
-				DesiredMaterialType = bHasIntermediateValues ? ESpriteInitMaterialType::Translucent : ESpriteInitMaterialType::Masked;
-			}
+			DesiredMaterialType = AnalyzeTextureForDesiredMaterialType(InitParams.Texture, InitParams.Offset, InitParams.Dimension);
 		}
 	}
 
@@ -114,24 +127,52 @@ void UPaperImporterSettings::ApplySettingsForSpriteInit(FSpriteAssetInitParamete
 		const bool bUseLitMaterial = LightingMode == ESpriteInitMaterialLightingMode::ForceLit;
 
 		// Apply the materials
-		switch (DesiredMaterialType)
+		InitParams.DefaultMaterialOverride = GetDefaultMaterial(DesiredMaterialType, bUseLitMaterial);
+		InitParams.AlternateMaterialOverride = GetDefaultOpaqueMaterial(bUseLitMaterial);
+	}
+}
+
+void UPaperImporterSettings::ApplySettingsForTileMapInit(UPaperTileMap* TileMap, UPaperTileSet* DefaultTileSet, ESpriteInitMaterialLightingMode LightingMode, ESpriteInitMaterialType MaterialTypeMode, bool bCreateEmptyLayer) const
+{
+	if (DefaultTileSet != nullptr)
+	{
+		TileMap->TileWidth = DefaultTileSet->TileWidth;
+		TileMap->TileHeight = DefaultTileSet->TileHeight;
+		TileMap->SelectedTileSet = DefaultTileSet;
+	}
+
+	TileMap->PixelsPerUnrealUnit = DefaultPixelsPerUnrealUnit;
+	TileMap->BackgroundColor = GetDefault<UTileMapEditorSettings>()->DefaultBackgroundColor;
+
+	ESpriteInitMaterialType DesiredMaterialType = MaterialTypeMode;
+	if (DesiredMaterialType == ESpriteInitMaterialType::Automatic)
+	{
+		// Analyze the texture if desired (to see if it's got greyscale alpha or just binary alpha, picking either a translucent or masked material)
+		if (bPickBestMaterialWhenCreatingTileMaps && (DefaultTileSet != nullptr) && (DefaultTileSet->TileSheet != nullptr))
 		{
-		default:
-		case ESpriteInitMaterialType::LeaveAsIs:
-		case ESpriteInitMaterialType::Automatic:
-			check(false);
-		case ESpriteInitMaterialType::Masked:
-			InitParams.DefaultMaterialOverride = GetDefaultMaskedMaterial(bUseLitMaterial);
-			break;
-		case ESpriteInitMaterialType::Translucent:
-			InitParams.DefaultMaterialOverride = GetDefaultTranslucentMaterial(bUseLitMaterial);
-			break;
-		case ESpriteInitMaterialType::Opaque:
-			InitParams.DefaultMaterialOverride = GetDefaultOpaqueMaterial(bUseLitMaterial);
-			break;
+			DesiredMaterialType = AnalyzeTextureForDesiredMaterialType(DefaultTileSet->TileSheet, FIntPoint::ZeroValue, DefaultTileSet->TileSheet->GetImportedSize());
 		}
+	}
 
-		InitParams.AlternateMaterialOverride = GetDefaultOpaqueMaterial(bUseLitMaterial);
+	if (DesiredMaterialType == ESpriteInitMaterialType::Automatic)
+	{
+		// Fall back to masked if we wanted automatic and couldn't analyze things
+		DesiredMaterialType = ESpriteInitMaterialType::Masked;
+	}
+
+	if (DesiredMaterialType != ESpriteInitMaterialType::LeaveAsIs)
+	{
+		// Determine whether to use lit or unlit materials
+		const bool bUseLitMaterial = LightingMode == ESpriteInitMaterialLightingMode::ForceLit;
+
+		// Apply the material
+		TileMap->Material = GetDefaultMaterial(DesiredMaterialType, bUseLitMaterial);
+	}
+
+	if (bCreateEmptyLayer)
+	{
+		// Add a new empty layer
+		TileMap->AddNewLayer();
 	}
 }
 
@@ -149,3 +190,22 @@ UMaterialInterface* UPaperImporterSettings::GetDefaultMaskedMaterial(bool bLit)
 {
 	return Cast<UMaterialInterface>((bLit ? LitDefaultMaskedMaterialName : UnlitDefaultMaskedMaterialName).TryLoad());
 }
+
+UMaterialInterface* UPaperImporterSettings::GetDefaultMaterial(ESpriteInitMaterialType MaterialType, bool bUseLitMaterial) const
+{
+	// Apply the materials
+	switch (MaterialType)
+	{
+	default:
+		ensureMsgf(false, TEXT("Unexpected material type in UPaperImporterSettings::GetDefaultMaterial"));
+		// Fall thru
+	case ESpriteInitMaterialType::LeaveAsIs:
+	case ESpriteInitMaterialType::Automatic:
+	case ESpriteInitMaterialType::Masked:
+		return GetDefaultMaskedMaterial(bUseLitMaterial);
+	case ESpriteInitMaterialType::Translucent:
+		return GetDefaultTranslucentMaterial(bUseLitMaterial);
+	case ESpriteInitMaterialType::Opaque:
+		return GetDefaultOpaqueMaterial(bUseLitMaterial);
+	}
+}
\ No newline at end of file
diff --git a/Engine/Plugins/2D/Paper2D/Source/Paper2DEditor/Private/PaperJSONHelpers.cpp b/Engine/Plugins/2D/Paper2D/Source/Paper2DEditor/Private/PaperJSONHelpers.cpp
--- a/Engine/Plugins/2D/Paper2D/Source/Paper2DEditor/Private/PaperJSONHelpers.cpp
+++ b/Engine/Plugins/2D/Paper2D/Source/Paper2DEditor/Private/PaperJSONHelpers.cpp
@@ -82,22 +82,22 @@ bool FPaperJSONHelpers::ReadIntegerNoDefault(const TSharedPtr<class FJsonObject>
 	}
 }
 
-bool FPaperJSONHelpers::ReadRectangle(const TSharedPtr<class FJsonObject> Item, const FString& Key, FVector2D& Out_XY, FVector2D& Out_WH)
+bool FPaperJSONHelpers::ReadRectangle(const TSharedPtr<class FJsonObject> Item, const FString& Key, FIntPoint& Out_XY, FIntPoint& Out_WH)
 {
 	const TSharedPtr<FJsonObject> Struct = ReadObject(Item, Key);
 	if (Struct.IsValid())
 	{
-		if (ReadFloatNoDefault(Struct, TEXT("x"), /*out*/ Out_XY.X) &&
-			ReadFloatNoDefault(Struct, TEXT("y"), /*out*/ Out_XY.Y) &&
-			ReadFloatNoDefault(Struct, TEXT("w"), /*out*/ Out_WH.X) &&
-			ReadFloatNoDefault(Struct, TEXT("h"), /*out*/ Out_WH.Y))
+		if (ReadIntegerNoDefault(Struct, TEXT("x"), /*out*/ Out_XY.X) &&
+			ReadIntegerNoDefault(Struct, TEXT("y"), /*out*/ Out_XY.Y) &&
+			ReadIntegerNoDefault(Struct, TEXT("w"), /*out*/ Out_WH.X) &&
+			ReadIntegerNoDefault(Struct, TEXT("h"), /*out*/ Out_WH.Y))
 		{
 			return true;
 		}
 		else
 		{
-			Out_XY = FVector2D::ZeroVector;
-			Out_WH = FVector2D::ZeroVector;
+			Out_XY = FIntPoint::ZeroValue;
+			Out_WH = FIntPoint::ZeroValue;
 			return false;
 		}
 	}
diff --git a/Engine/Plugins/2D/Paper2D/Source/Paper2DEditor/Private/PaperTileMapFactory.cpp b/Engine/Plugins/2D/Paper2D/Source/Paper2DEditor/Private/PaperTileMapFactory.cpp
--- a/Engine/Plugins/2D/Paper2D/Source/Paper2DEditor/Private/PaperTileMapFactory.cpp
+++ b/Engine/Plugins/2D/Paper2D/Source/Paper2DEditor/Private/PaperTileMapFactory.cpp
@@ -2,7 +2,7 @@
 
 #include "Paper2DEditorPrivatePCH.h"
 #include "PaperTileMapFactory.h"
-#include "TileMapEditing/TileMapEditorSettings.h"
+#include "PaperImporterSettings.h"
 
 #define LOCTEXT_NAMESPACE "Paper2D"
 
@@ -21,15 +21,7 @@ UObject* UPaperTileMapFactory::FactoryCreateNew(UClass* Class, UObject* InParent
 {
 	UPaperTileMap* NewTileMap = NewObject<UPaperTileMap>(InParent, Class, Name, Flags | RF_Transactional);
 
-	if (InitialTileSet != nullptr)
-	{
-		NewTileMap->TileWidth = InitialTileSet->TileWidth;
-		NewTileMap->TileHeight = InitialTileSet->TileHeight;
-		NewTileMap->SelectedTileSet = InitialTileSet;
-	}
-
-	NewTileMap->BackgroundColor = GetDefault<UTileMapEditorSettings>()->DefaultBackgroundColor;
-	NewTileMap->InitializeNewEmptyTileMap();
+	GetDefault<UPaperImporterSettings>()->ApplySettingsForTileMapInit(NewTileMap, InitialTileSet);
 
 	return NewTileMap;
 }
diff --git a/Engine/Plugins/2D/Paper2D/Source/Paper2DEditor/Private/SpriteEditor/SpriteEditorViewportClient.cpp b/Engine/Plugins/2D/Paper2D/Source/Paper2DEditor/Private/SpriteEditor/SpriteEditorViewportClient.cpp
--- a/Engine/Plugins/2D/Paper2D/Source/Paper2DEditor/Private/SpriteEditor/SpriteEditorViewportClient.cpp
+++ b/Engine/Plugins/2D/Paper2D/Source/Paper2DEditor/Private/SpriteEditor/SpriteEditorViewportClient.cpp
@@ -698,8 +698,10 @@ void FSpriteEditorViewportClient::ProcessClick(FSceneView& View, HHitProxy* HitP
 			FVector2D TexturePoint = SourceTextureViewComponent->GetSprite()->ConvertWorldSpaceToTextureSpace(WorldPoint);
 			if (bIsCtrlKeyDown)
 			{
-				UPaperSprite* NewSprite = CreateNewSprite(Sprite->GetSourceUV(), Sprite->GetSourceSize());
-				if (NewSprite != nullptr)
+				const FVector2D StartingUV = Sprite->GetSourceUV();
+				const FVector2D StartingSize = Sprite->GetSourceSize();
+
+				if (UPaperSprite* NewSprite = CreateNewSprite(FIntPoint((int32)StartingUV.X, (int32)StartingUV.Y), FIntPoint((int32)StartingSize.X, (int32)StartingSize.Y)))
 				{
 					NewSprite->ExtractSourceRegionFromTexturePoint(TexturePoint);
 					bHandled = true;
@@ -746,7 +748,7 @@ void FSpriteEditorViewportClient::ProcessClick(FSceneView& View, HHitProxy* HitP
 
 // Create a new sprite and return this sprite. The sprite editor will now be editing this new sprite
 // Returns nullptr if failed
-UPaperSprite* FSpriteEditorViewportClient::CreateNewSprite(FVector2D TopLeft, FVector2D Dimensions)
+UPaperSprite* FSpriteEditorViewportClient::CreateNewSprite(const FIntPoint& TopLeft, const FIntPoint& Dimensions)
 {
 	FAssetToolsModule& AssetToolsModule = FModuleManager::Get().LoadModuleChecked<FAssetToolsModule>("AssetTools");
 	FContentBrowserModule& ContentBrowserModule = FModuleManager::LoadModuleChecked<FContentBrowserModule>("ContentBrowser");
@@ -798,8 +800,8 @@ bool FSpriteEditorViewportClient::InputKey(FViewport* Viewport, int32 Controller
 		const bool bMarqueeStartModifier = InputState.IsCtrlButtonPressed();
 		if (GeometryEditMode->ProcessMarquee(Viewport, Key, Event, bMarqueeStartModifier))
 		{
-			FVector2D TextureSpaceStartPos;
-			FVector2D TextureSpaceDimensions;
+			FIntPoint TextureSpaceStartPos;
+			FIntPoint TextureSpaceDimensions;
 			if (ConvertMarqueeToSourceTextureSpace(/*out*/TextureSpaceStartPos, /*out*/TextureSpaceDimensions))
 			{
 				//@TODO: Warn if overlapping with another sprite
@@ -985,7 +987,7 @@ void FSpriteEditorViewportClient::InternalActivateNewMode(ESpriteEditorMode::Typ
 	}
 }
 
-bool FSpriteEditorViewportClient::ConvertMarqueeToSourceTextureSpace(/*out*/FVector2D& OutStartPos, /*out*/FVector2D& OutDimension)
+bool FSpriteEditorViewportClient::ConvertMarqueeToSourceTextureSpace(/*out*/ FIntPoint& OutStartPos, /*out*/ FIntPoint& OutDimension)
 {
 	FSpriteGeometryEditMode* GeometryEditMode = ModeTools->GetActiveModeTyped<FSpriteGeometryEditMode>(FSpriteGeometryEditMode::EM_SpriteGeometry);
 	check(GeometryEditMode);
@@ -1019,15 +1021,19 @@ bool FSpriteEditorViewportClient::ConvertMarqueeToSourceTextureSpace(/*out*/FVec
 		const FIntPoint SourceTextureSize(SpriteSourceTexture->GetImportedSize());
 		const int32 SourceTextureWidth = SourceTextureSize.X;
 		const int32 SourceTextureHeight = SourceTextureSize.Y;
-		TextureSpaceStartPos.X = FMath::Clamp((int)TextureSpaceStartPos.X, 0, SourceTextureWidth - 1);
-		TextureSpaceStartPos.Y = FMath::Clamp((int)TextureSpaceStartPos.Y, 0, SourceTextureHeight - 1);
-		TextureSpaceEndPos.X = FMath::Clamp((int)TextureSpaceEndPos.X, 0, SourceTextureWidth - 1);
-		TextureSpaceEndPos.Y = FMath::Clamp((int)TextureSpaceEndPos.Y, 0, SourceTextureHeight - 1);
+		
+		FIntPoint TSStartPos;
+		TSStartPos.X = FMath::Clamp<int32>((int32)TextureSpaceStartPos.X, 0, SourceTextureWidth - 1);
+		TSStartPos.Y = FMath::Clamp<int32>((int32)TextureSpaceStartPos.Y, 0, SourceTextureHeight - 1);
+
+		FIntPoint TSEndPos;
+		TSEndPos.X = FMath::Clamp<int32>((int32)TextureSpaceEndPos.X, 0, SourceTextureWidth - 1);
+		TSEndPos.Y = FMath::Clamp<int32>((int32)TextureSpaceEndPos.Y, 0, SourceTextureHeight - 1);
 
-		const FVector2D TextureSpaceDimensions = TextureSpaceEndPos - TextureSpaceStartPos;
-		if (TextureSpaceDimensions.X > 0 || TextureSpaceDimensions.Y > 0)
+		const FIntPoint TextureSpaceDimensions = TSEndPos - TSStartPos;
+		if ((TextureSpaceDimensions.X > 0) || (TextureSpaceDimensions.Y > 0))
 		{
-			OutStartPos = TextureSpaceStartPos;
+			OutStartPos = TSStartPos;
 			OutDimension = TextureSpaceDimensions;
 			bSuccessful = true;
 		}
diff --git a/Engine/Plugins/2D/Paper2D/Source/Paper2DEditor/Private/SpriteEditor/SpriteEditorViewportClient.h b/Engine/Plugins/2D/Paper2D/Source/Paper2DEditor/Private/SpriteEditor/SpriteEditorViewportClient.h
--- a/Engine/Plugins/2D/Paper2D/Source/Paper2DEditor/Private/SpriteEditor/SpriteEditorViewportClient.h
+++ b/Engine/Plugins/2D/Paper2D/Source/Paper2DEditor/Private/SpriteEditor/SpriteEditorViewportClient.h
@@ -90,7 +90,7 @@ public:
 
 	void UpdateRelatedSpritesList();
 
-	UPaperSprite* CreateNewSprite(FVector2D TopLeft, FVector2D Dimensions);
+	UPaperSprite* CreateNewSprite(const FIntPoint& TopLeft, const FIntPoint& Dimensions);
 
 	ESpriteEditorMode::Type GetCurrentMode() const
 	{
@@ -168,7 +168,7 @@ private:
 
 	void UpdateSourceTextureSpriteFromSprite(UPaperSprite* SourceSprite);
 	
-	bool ConvertMarqueeToSourceTextureSpace(/*out*/FVector2D& OutStartPos, /*out*/FVector2D& OutDimension);
+	bool ConvertMarqueeToSourceTextureSpace(/*out*/ FIntPoint& OutStartPos, /*out*/ FIntPoint& OutDimension);
 
 	// Activates a new mode, clearing selection set, etc...
 	void InternalActivateNewMode(ESpriteEditorMode::Type NewMode);
diff --git a/Engine/Plugins/2D/Paper2D/Source/Paper2DEditor/Private/TileMapEditing/PaperTileMapDetailsCustomization.cpp b/Engine/Plugins/2D/Paper2D/Source/Paper2DEditor/Private/TileMapEditing/PaperTileMapDetailsCustomization.cpp
--- a/Engine/Plugins/2D/Paper2D/Source/Paper2DEditor/Private/TileMapEditing/PaperTileMapDetailsCustomization.cpp
+++ b/Engine/Plugins/2D/Paper2D/Source/Paper2DEditor/Private/TileMapEditing/PaperTileMapDetailsCustomization.cpp
@@ -8,6 +8,7 @@
 #include "PropertyEditing.h"
 #include "AssetToolsModule.h"
 #include "PaperTileMapPromotionFactory.h"
+#include "PaperImporterSettings.h"
 
 #include "STileLayerList.h"
 #include "ScopedTransaction.h"
@@ -262,10 +263,15 @@ FReply FPaperTileMapDetailsCustomization::OnNewButtonClicked()
 {
 	if (UPaperTileMapComponent* TileMapComponent = TileMapComponentPtr.Get())
 	{
+		UPaperTileSet* OldTileSet = (TileMapComponent->TileMap != nullptr) ? TileMapComponent->TileMap->SelectedTileSet.Get() : nullptr;
+
 		const FScopedTransaction Transaction(LOCTEXT("CreateNewTileMap", "New Tile Map"));
 		TileMapComponent->Modify();
 		TileMapComponent->CreateNewOwnedTileMap();
 
+		// Add a layer and set things up
+		GetDefault<UPaperImporterSettings>()->ApplySettingsForTileMapInit(TileMapComponent->TileMap, OldTileSet);
+
 		MyDetailLayout->ForceRefreshDetails();
 	}
 
diff --git a/Engine/Plugins/2D/Paper2D/Source/Paper2DEditor/Private/TileMapEditing/TileMapActorFactory.cpp b/Engine/Plugins/2D/Paper2D/Source/Paper2DEditor/Private/TileMapEditing/TileMapActorFactory.cpp
--- a/Engine/Plugins/2D/Paper2D/Source/Paper2DEditor/Private/TileMapEditing/TileMapActorFactory.cpp
+++ b/Engine/Plugins/2D/Paper2D/Source/Paper2DEditor/Private/TileMapEditing/TileMapActorFactory.cpp
@@ -23,29 +23,22 @@ void UTileMapActorFactory::PostSpawnActor(UObject* Asset, AActor* NewActor)
 	UPaperTileMapComponent* RenderComponent = TypedActor->GetRenderComponent();
 	check(RenderComponent);
 
-	if (UPaperTileMap* TileMap = Cast<UPaperTileMap>(Asset))
+	if (UPaperTileMap* TileMapAsset = Cast<UPaperTileMap>(Asset))
 	{
 		RenderComponent->UnregisterComponent();
-		RenderComponent->TileMap = TileMap;
+		RenderComponent->TileMap = TileMapAsset;
 		RenderComponent->RegisterComponent();
 	}
-	else if (UPaperTileSet* TileSet = Cast<UPaperTileSet>(Asset))
+	else if (RenderComponent->OwnsTileMap())
 	{
-		if (RenderComponent->TileMap != nullptr)
-		{
-			RenderComponent->UnregisterComponent();
-			RenderComponent->TileMap->TileWidth = TileSet->TileWidth;
-			RenderComponent->TileMap->TileHeight = TileSet->TileHeight;
-			RenderComponent->TileMap->SelectedTileSet = TileSet;
-			//@TODO: Initialize the tile map material here too, based on analysis of the TileSheet texture
-			RenderComponent->RegisterComponent();
-		}
-	}
+		RenderComponent->UnregisterComponent();
 
-	if (RenderComponent->OwnsTileMap())
-	{
-		RenderComponent->TileMap->InitializeNewEmptyTileMap();
-		RenderComponent->TileMap->PixelsPerUnrealUnit = GetDefault<UPaperImporterSettings>()->GetDefaultPixelsPerUnrealUnit();
+		UPaperTileMap* OwnedTileMap = RenderComponent->TileMap;
+		check(OwnedTileMap);
+
+		GetDefault<UPaperImporterSettings>()->ApplySettingsForTileMapInit(OwnedTileMap, Cast<UPaperTileSet>(Asset));
+
+		RenderComponent->RegisterComponent();
 	}
 }
 
@@ -60,21 +53,12 @@ void UTileMapActorFactory::PostCreateBlueprint(UObject* Asset, AActor* CDO)
 		{
 			RenderComponent->TileMap = TileMap;
 		}
-		else if (UPaperTileSet* TileSet = Cast<UPaperTileSet>(Asset))
+		else if (RenderComponent->OwnsTileMap())
 		{
-			if (RenderComponent->TileMap != nullptr)
-			{
-				RenderComponent->TileMap->TileWidth = TileSet->TileWidth;
-				RenderComponent->TileMap->TileHeight = TileSet->TileHeight;
-				RenderComponent->TileMap->SelectedTileSet = TileSet;
-				//@TODO: Initialize the tile map material here too, based on analysis of the TileSheet texture
-			}
-		}
+			UPaperTileMap* OwnedTileMap = RenderComponent->TileMap;
+			check(OwnedTileMap);
 
-		if (RenderComponent->OwnsTileMap())
-		{
-			RenderComponent->TileMap->InitializeNewEmptyTileMap();
-			RenderComponent->TileMap->PixelsPerUnrealUnit = GetDefault<UPaperImporterSettings>()->GetDefaultPixelsPerUnrealUnit();
+			GetDefault<UPaperImporterSettings>()->ApplySettingsForTileMapInit(OwnedTileMap, Cast<UPaperTileSet>(Asset));
 		}
 	}
 }
diff --git a/Engine/Plugins/2D/Paper2D/Source/Paper2DEditor/Private/TileSetEditor/SingleTileEditorViewportClient.cpp b/Engine/Plugins/2D/Paper2D/Source/Paper2DEditor/Private/TileSetEditor/SingleTileEditorViewportClient.cpp
--- a/Engine/Plugins/2D/Paper2D/Source/Paper2DEditor/Private/TileSetEditor/SingleTileEditorViewportClient.cpp
+++ b/Engine/Plugins/2D/Paper2D/Source/Paper2DEditor/Private/TileSetEditor/SingleTileEditorViewportClient.cpp
@@ -239,9 +239,10 @@ void FSingleTileEditorViewportClient::SetTileIndex(int32 InTileIndex)
 		SpriteReinitParams.Texture = TileSet->TileSheet;
 
 		//@TODO: Should analyze the texture (*at a higher level, not per tile click!*) to pick the correct material
-
-		TileSet->GetTileUV(TileBeingEditedIndex, /*out*/ SpriteReinitParams.Offset);
-		SpriteReinitParams.Dimension = FVector2D(TileSet->TileWidth, TileSet->TileHeight);
+		FVector2D UV;
+		TileSet->GetTileUV(TileBeingEditedIndex, /*out*/ UV);
+		SpriteReinitParams.Offset = FIntPoint((int32)UV.X, (int32)(UV.Y));
+		SpriteReinitParams.Dimension = FIntPoint(TileSet->TileWidth, TileSet->TileHeight);
 		SpriteReinitParams.SetPixelsPerUnrealUnit(1.0f);
 		DummySprite->InitializeSprite(SpriteReinitParams);
 	}
diff --git a/Engine/Plugins/2D/Paper2D/Source/Paper2DEditor/Public/PaperJSONHelpers.h b/Engine/Plugins/2D/Paper2D/Source/Paper2DEditor/Public/PaperJSONHelpers.h
--- a/Engine/Plugins/2D/Paper2D/Source/Paper2DEditor/Public/PaperJSONHelpers.h
+++ b/Engine/Plugins/2D/Paper2D/Source/Paper2DEditor/Public/PaperJSONHelpers.h
@@ -26,8 +26,8 @@ public:
 	// Returns true if the field named Key is a Number (truncating it to an integer), populating Out_Value, or false if it is missing or the wrong type (Out_Value is set to 0 on failure)
 	static bool ReadIntegerNoDefault(const TSharedPtr<class FJsonObject> Item, const FString& Key, int32& Out_Value);
 
-	// Returns true if the object named Key is a struct containing four floats (x,y,w,h), populating XY and WH with the values)
-	static bool ReadRectangle(const TSharedPtr<class FJsonObject> Item, const FString& Key, FVector2D& Out_XY, FVector2D& Out_WH);
+	// Returns true if the object named Key is a struct containing four integers (x,y,w,h), populating XY and WH with the values)
+	static bool ReadRectangle(const TSharedPtr<class FJsonObject> Item, const FString& Key, FIntPoint& Out_XY, FIntPoint& Out_WH);
 
 	// Returns true if the object named Key is a struct containing two floats (w,h), populating WH with the values)
 	static bool ReadSize(const TSharedPtr<class FJsonObject> Item, const FString& Key, FVector2D& Out_WH);
diff --git a/Engine/Plugins/2D/Paper2D/Source/PaperSpriteSheetImporter/Private/PaperJsonSpriteSheetImporter.cpp b/Engine/Plugins/2D/Paper2D/Source/PaperSpriteSheetImporter/Private/PaperJsonSpriteSheetImporter.cpp
--- a/Engine/Plugins/2D/Paper2D/Source/PaperSpriteSheetImporter/Private/PaperJsonSpriteSheetImporter.cpp
+++ b/Engine/Plugins/2D/Paper2D/Source/PaperSpriteSheetImporter/Private/PaperJsonSpriteSheetImporter.cpp
@@ -129,7 +129,7 @@ static bool ParseFrame(TSharedPtr<FJsonObject> &FrameData, FSpriteFrame &OutFram
 	}
 	else
 	{
-		OutFrame.SpriteSourcePos = FVector2D::ZeroVector;
+		OutFrame.SpriteSourcePos = FIntPoint::ZeroValue;
 		OutFrame.SpriteSourceSize = OutFrame.SpriteSizeInSheet;
 		OutFrame.ImageSourceSize = OutFrame.SpriteSizeInSheet;
 	}
diff --git a/Engine/Plugins/2D/Paper2D/Source/PaperSpriteSheetImporter/Private/PaperJsonSpriteSheetImporter.h b/Engine/Plugins/2D/Paper2D/Source/PaperSpriteSheetImporter/Private/PaperJsonSpriteSheetImporter.h
--- a/Engine/Plugins/2D/Paper2D/Source/PaperSpriteSheetImporter/Private/PaperJsonSpriteSheetImporter.h
+++ b/Engine/Plugins/2D/Paper2D/Source/PaperSpriteSheetImporter/Private/PaperJsonSpriteSheetImporter.h
@@ -15,11 +15,11 @@ struct FSpriteFrame
 {
 	FName FrameName;
 
-	FVector2D SpritePosInSheet;
-	FVector2D SpriteSizeInSheet;
+	FIntPoint SpritePosInSheet;
+	FIntPoint SpriteSizeInSheet;
 
-	FVector2D SpriteSourcePos;
-	FVector2D SpriteSourceSize;
+	FIntPoint SpriteSourcePos;
+	FIntPoint SpriteSourceSize;
 
 	FVector2D ImageSourceSize;
 
diff --git a/Engine/Plugins/2D/Paper2D/Source/PaperTiledImporter/Classes/PaperTiledImporterFactory.h b/Engine/Plugins/2D/Paper2D/Source/PaperTiledImporter/Classes/PaperTiledImporterFactory.h
--- a/Engine/Plugins/2D/Paper2D/Source/PaperTiledImporter/Classes/PaperTiledImporterFactory.h
+++ b/Engine/Plugins/2D/Paper2D/Source/PaperTiledImporter/Classes/PaperTiledImporterFactory.h
@@ -83,6 +83,10 @@ protected:
 
 	static UObject* CreateNewAsset(UClass* AssetClass, const FString& TargetPath, const FString& DesiredName, EObjectFlags Flags);
 	static UTexture2D* ImportTexture(const FString& SourceFilename, const FString& TargetSubPath);
+
+	void FinalizeTileMap(FTileMapFromTiled& GlobalInfo, UPaperTileMap* TileMap);
+
+	bool ConvertTileSets(FTileMapFromTiled& GlobalInfo, const FString& CurrentSourcePath, const FString& LongPackagePath, EObjectFlags Flags);
 };
 
 //////////////////////////////////////////////////////////////////////////
diff --git a/Engine/Plugins/2D/Paper2D/Source/PaperTiledImporter/Private/PaperTiledImporterFactory.cpp b/Engine/Plugins/2D/Paper2D/Source/PaperTiledImporter/Private/PaperTiledImporterFactory.cpp
--- a/Engine/Plugins/2D/Paper2D/Source/PaperTiledImporter/Private/PaperTiledImporterFactory.cpp
+++ b/Engine/Plugins/2D/Paper2D/Source/PaperTiledImporter/Private/PaperTiledImporterFactory.cpp
@@ -209,7 +209,6 @@ UObject* UPaperTiledImporterFactory::FactoryCreateText(UClass* InClass, UObject*
 		Result->SeparationPerTileY = 0.0f;
 		Result->SeparationPerLayer = 1.0f;
 		Result->ProjectionMode = GlobalInfo.GetOrientationType();
-		Result->PixelsPerUnrealUnit = GetDefault<UPaperImporterSettings>()->GetDefaultPixelsPerUnrealUnit();
 		Result->BackgroundColor = GlobalInfo.BackgroundColor;
 		Result->HexSideLength = GlobalInfo.HexSideLength;
 
@@ -219,100 +218,8 @@ UObject* UPaperTiledImporterFactory::FactoryCreateText(UClass* InClass, UObject*
 		}
 
 		// Create the tile sets
-		for (const FTileSetFromTiled& TileSetData : GlobalInfo.TileSets)
-		{
-			if (TileSetData.IsValid())
-			{
-				const FString TargetTileSetPath = LongPackagePath;
-				const FString TargetTexturePath = LongPackagePath + TEXT("/Textures");
-
-				UPaperTileSet* TileSetAsset = CastChecked<UPaperTileSet>(CreateNewAsset(UPaperTileSet::StaticClass(), TargetTileSetPath, TileSetData.Name, Flags));
-				TileSetAsset->Modify();
-
-				TileSetAsset->TileWidth = TileSetData.TileWidth;
-				TileSetAsset->TileHeight = TileSetData.TileHeight;
-				TileSetAsset->Margin = TileSetData.Margin;
-				TileSetAsset->Spacing = TileSetData.Spacing;
-				TileSetAsset->DrawingOffset = FIntPoint(TileSetData.TileOffsetX, TileSetData.TileOffsetY);
-
-				// Import the texture
-				const FString SourceImageFilename = FPaths::Combine(*CurrentSourcePath, *TileSetData.ImagePath);
-				TileSetAsset->TileSheet = ImportTexture(SourceImageFilename, TargetTexturePath);
-
-				if (TileSetAsset->TileSheet == nullptr)
-				{
-					UE_LOG(LogPaperTiledImporter, Warning, TEXT("Failed to import tile set image '%s' referenced from tile set '%s'."), *TileSetData.ImagePath, *TileSetData.Name);
-					bLoadedSuccessfully = false;
-				}
-
-				// Make the tile set allocate space for the per-tile data
-				FPropertyChangedEvent InteractiveRebuildTileSet(nullptr, EPropertyChangeType::Interactive);
-				TileSetAsset->PostEditChangeProperty(InteractiveRebuildTileSet);
-
-				// Copy across terrain information
-				const int32 MaxTerrainTypes = 0xFE;
-				const uint8 NoTerrainMembershipIndex = 0xFF;
-				if (TileSetData.TerrainTypes.Num() > MaxTerrainTypes)
-				{
-					UE_LOG(LogPaperTiledImporter, Warning, TEXT("Tile set '%s' contains more than %d terrain types, ones above this will be ignored."), *TileSetData.Name, MaxTerrainTypes);
-				}
-				const int32 NumTerrainsToCopy = FMath::Min<int32>(TileSetData.TerrainTypes.Num(), MaxTerrainTypes);
-				for (int32 TerrainIndex = 0; TerrainIndex < NumTerrainsToCopy; ++TerrainIndex)
-				{
-					const FTiledTerrain& SourceTerrain = TileSetData.TerrainTypes[TerrainIndex];
-
-					FPaperTileSetTerrain DestTerrain;
-					DestTerrain.TerrainName = SourceTerrain.TerrainName;
-					DestTerrain.CenterTileIndex = SourceTerrain.SolidTileLocalIndex;
-
-					TileSetAsset->AddTerrainDescription(DestTerrain);
-				}
-
-				// Copy across per-tile metadata
-				const int32 NumTilesCreated = TileSetAsset->GetTileCount();
-				for (const auto& KV : TileSetData.PerTileData)
-				{
-					const int32 TileIndex = KV.Key;
-					const FTiledTileInfo& SourceTileData = KV.Value;
-
-					if (FPaperTileMetadata* TargetTileData = TileSetAsset->GetMutableTileMetadata(TileIndex))
-					{
-						// Convert collision geometry
-						FTiledObject::AddToSpriteGeometryCollection(FVector2D::ZeroVector, SourceTileData.Objects, TargetTileData->CollisionData);
-						
-						// Convert terrain memberhsip
-						for (int32 Index = 0; Index < 4; ++Index)
-						{
-							const int32 SourceTerrainIndex = SourceTileData.TerrainIndices[Index];
-
-							const uint8 DestTerrainIndex = ((SourceTerrainIndex >= 0) && (SourceTerrainIndex < NumTerrainsToCopy)) ? (uint8)SourceTerrainIndex : NoTerrainMembershipIndex;
-							TargetTileData->TerrainMembership[Index] = DestTerrainIndex;
-						}
-
-						//@TODO: Convert metadata?
-					}
-				}
-				
-				// Update anyone who might be using the tile set (in case we're reimporting)
-				FPropertyChangedEvent FinalRebuildTileSet(nullptr, EPropertyChangeType::ValueSet);
-				TileSetAsset->PostEditChangeProperty(FinalRebuildTileSet);
-
-				// Save off that we created the asset
-				GlobalInfo.CreatedTileSetAssets.Add(TileSetAsset);
-			}
-			else
-			{
-				GlobalInfo.CreatedTileSetAssets.Add(nullptr);
-			}
-		}
-
-		if (GlobalInfo.CreatedTileSetAssets.Num() > 0)
-		{
-			// Bind our selected tile set to the first tile set that was imported so something is already picked
-			Result->SelectedTileSet = GlobalInfo.CreatedTileSetAssets[0];
-		
-			//@TODO: Should analyze the material from the used tile sets and set it appropriately on the tile map
-		}
+		const bool bConvertedTileSetsSuccessfully = ConvertTileSets(GlobalInfo, CurrentSourcePath, LongPackagePath, Flags);
+		bLoadedSuccessfully = bLoadedSuccessfully && bConvertedTileSetsSuccessfully;
 
 		// Create the layers
 		for (int32 LayerIndex = GlobalInfo.Layers.Num() - 1; LayerIndex >= 0; --LayerIndex)
@@ -349,6 +256,9 @@ UObject* UPaperTiledImporterFactory::FactoryCreateText(UClass* InClass, UObject*
 			}
 		}
 
+		// Finalize the tile map, including analyzing the tile set textures to determine a good material
+		FinalizeTileMap(GlobalInfo, Result);
+
 		Result->PostEditChange();
 	}
  	else
@@ -1236,6 +1146,149 @@ void FTiledObject::AddToSpriteGeometryCollection(const FVector2D& Offset, const
 	InOutShapes.ConditionGeometry();
 }
 
+void UPaperTiledImporterFactory::FinalizeTileMap(FTileMapFromTiled& GlobalInfo, UPaperTileMap* TileMap)
+{
+	const UPaperImporterSettings* ImporterSettings = GetDefault<UPaperImporterSettings>();
+
+	// Bind our selected tile set to the first tile set that was imported so something is already picked
+	UPaperTileSet* DefaultTileSet = (GlobalInfo.CreatedTileSetAssets.Num() > 0) ? GlobalInfo.CreatedTileSetAssets[0] : nullptr;
+	TileMap->SelectedTileSet = DefaultTileSet;
+
+	// Initialize the scale
+	TileMap->PixelsPerUnrealUnit = ImporterSettings->GetDefaultPixelsPerUnrealUnit();
+	
+	// Analyze the tile set textures (anything with translucent wins; failing that use masked)
+	ESpriteInitMaterialType BestMaterial = ESpriteInitMaterialType::Masked;
+	if (ImporterSettings->ShouldPickBestMaterialWhenCreatingTileMaps())
+	{
+		BestMaterial = ESpriteInitMaterialType::Automatic;
+		for (UPaperTileSet* TileSet : GlobalInfo.CreatedTileSetAssets)
+		{
+			if ((TileSet != nullptr) && (TileSet->TileSheet != nullptr))
+			{
+				ESpriteInitMaterialType TileSheetMaterial = ImporterSettings->AnalyzeTextureForDesiredMaterialType(TileSet->TileSheet, FIntPoint::ZeroValue, TileSet->TileSheet->GetImportedSize());
+				
+				switch (TileSheetMaterial)
+				{
+				case ESpriteInitMaterialType::Opaque:
+				case ESpriteInitMaterialType::Masked:
+					BestMaterial = ((BestMaterial == ESpriteInitMaterialType::Automatic) || (BestMaterial == ESpriteInitMaterialType::Opaque)) ? TileSheetMaterial : BestMaterial;
+					break;
+				case ESpriteInitMaterialType::Translucent:
+					BestMaterial = TileSheetMaterial;
+					break;
+				}
+			}
+		}
+	}
+
+	if (BestMaterial == ESpriteInitMaterialType::Automatic)
+	{
+		// Fall back to masked if we wanted automatic and couldn't analyze things
+		BestMaterial = ESpriteInitMaterialType::Masked;
+	}
+
+	if (BestMaterial != ESpriteInitMaterialType::LeaveAsIs)
+	{
+		const bool bUseLitMaterial = false;
+		TileMap->Material = ImporterSettings->GetDefaultMaterial(BestMaterial, bUseLitMaterial);
+	}
+}
+
+bool UPaperTiledImporterFactory::ConvertTileSets(FTileMapFromTiled& GlobalInfo, const FString& CurrentSourcePath, const FString& LongPackagePath, EObjectFlags Flags)
+{
+	bool bLoadedSuccessfully = true;
+
+	for (const FTileSetFromTiled& TileSetData : GlobalInfo.TileSets)
+	{
+		if (TileSetData.IsValid())
+		{
+			const FString TargetTileSetPath = LongPackagePath;
+			const FString TargetTexturePath = LongPackagePath / TEXT("Textures");
+
+			UPaperTileSet* TileSetAsset = CastChecked<UPaperTileSet>(CreateNewAsset(UPaperTileSet::StaticClass(), TargetTileSetPath, TileSetData.Name, Flags));
+			TileSetAsset->Modify();
+
+			TileSetAsset->TileWidth = TileSetData.TileWidth;
+			TileSetAsset->TileHeight = TileSetData.TileHeight;
+			TileSetAsset->Margin = TileSetData.Margin;
+			TileSetAsset->Spacing = TileSetData.Spacing;
+			TileSetAsset->DrawingOffset = FIntPoint(TileSetData.TileOffsetX, TileSetData.TileOffsetY);
+
+			// Import the texture
+			const FString SourceImageFilename = FPaths::Combine(*CurrentSourcePath, *TileSetData.ImagePath);
+			TileSetAsset->TileSheet = ImportTexture(SourceImageFilename, TargetTexturePath);
+
+			if (TileSetAsset->TileSheet == nullptr)
+			{
+				UE_LOG(LogPaperTiledImporter, Warning, TEXT("Failed to import tile set image '%s' referenced from tile set '%s'."), *TileSetData.ImagePath, *TileSetData.Name);
+				bLoadedSuccessfully = false;
+			}
+
+			// Make the tile set allocate space for the per-tile data
+			FPropertyChangedEvent InteractiveRebuildTileSet(nullptr, EPropertyChangeType::Interactive);
+			TileSetAsset->PostEditChangeProperty(InteractiveRebuildTileSet);
+
+			// Copy across terrain information
+			const int32 MaxTerrainTypes = 0xFE;
+			const uint8 NoTerrainMembershipIndex = 0xFF;
+			if (TileSetData.TerrainTypes.Num() > MaxTerrainTypes)
+			{
+				UE_LOG(LogPaperTiledImporter, Warning, TEXT("Tile set '%s' contains more than %d terrain types, ones above this will be ignored."), *TileSetData.Name, MaxTerrainTypes);
+			}
+			const int32 NumTerrainsToCopy = FMath::Min<int32>(TileSetData.TerrainTypes.Num(), MaxTerrainTypes);
+			for (int32 TerrainIndex = 0; TerrainIndex < NumTerrainsToCopy; ++TerrainIndex)
+			{
+				const FTiledTerrain& SourceTerrain = TileSetData.TerrainTypes[TerrainIndex];
+
+				FPaperTileSetTerrain DestTerrain;
+				DestTerrain.TerrainName = SourceTerrain.TerrainName;
+				DestTerrain.CenterTileIndex = SourceTerrain.SolidTileLocalIndex;
+
+				TileSetAsset->AddTerrainDescription(DestTerrain);
+			}
+
+			// Copy across per-tile metadata
+			const int32 NumTilesCreated = TileSetAsset->GetTileCount();
+			for (const auto& KV : TileSetData.PerTileData)
+			{
+				const int32 TileIndex = KV.Key;
+				const FTiledTileInfo& SourceTileData = KV.Value;
+
+				if (FPaperTileMetadata* TargetTileData = TileSetAsset->GetMutableTileMetadata(TileIndex))
+				{
+					// Convert collision geometry
+					FTiledObject::AddToSpriteGeometryCollection(FVector2D::ZeroVector, SourceTileData.Objects, TargetTileData->CollisionData);
+
+					// Convert terrain memberhsip
+					for (int32 Index = 0; Index < 4; ++Index)
+					{
+						const int32 SourceTerrainIndex = SourceTileData.TerrainIndices[Index];
+
+						const uint8 DestTerrainIndex = ((SourceTerrainIndex >= 0) && (SourceTerrainIndex < NumTerrainsToCopy)) ? (uint8)SourceTerrainIndex : NoTerrainMembershipIndex;
+						TargetTileData->TerrainMembership[Index] = DestTerrainIndex;
+					}
+
+					//@TODO: Convert metadata?
+				}
+			}
+
+			// Update anyone who might be using the tile set (in case we're reimporting)
+			FPropertyChangedEvent FinalRebuildTileSet(nullptr, EPropertyChangeType::ValueSet);
+			TileSetAsset->PostEditChangeProperty(FinalRebuildTileSet);
+
+			// Save off that we created the asset
+			GlobalInfo.CreatedTileSetAssets.Add(TileSetAsset);
+		}
+		else
+		{
+			GlobalInfo.CreatedTileSetAssets.Add(nullptr);
+		}
+	}
+
+	return bLoadedSuccessfully;
+}
+
 //////////////////////////////////////////////////////////////////////////
 
 #undef TILED_IMPORT_ERROR
