Sprite Editor: Synchronized the selection in the sprite list when the sprite being edited changes [UE-13811]

diff --git a/Engine/Plugins/2D/Paper2D/Source/Paper2DEditor/Private/SpriteEditor/SSpriteList.cpp b/Engine/Plugins/2D/Paper2D/Source/Paper2DEditor/Private/SpriteEditor/SSpriteList.cpp
--- a/Engine/Plugins/2D/Paper2D/Source/Paper2DEditor/Private/SpriteEditor/SSpriteList.cpp
+++ b/Engine/Plugins/2D/Paper2D/Source/Paper2DEditor/Private/SpriteEditor/SSpriteList.cpp
@@ -31,6 +31,7 @@ void SSpriteList::RebuildWidget(UTexture2D* NewTextureFilter)
 
 	Config.ThumbnailScale = 0.0f;
 	Config.InitialAssetViewType = EAssetViewType::Tile;
+	Config.SyncToAssetsDelegates.Add(&SyncToAssetsDelegate);
 
 	if (NewTextureFilter != nullptr)
 	{
@@ -94,3 +95,16 @@ bool SSpriteList::CanShowColumnForAssetRegistryTag(FName AssetType, FName TagNam
 {
 	return !AssetRegistryTagsToIgnore.Contains(TagName);
 }
+
+void SSpriteList::SelectAsset(UObject* Asset)
+{
+	FAssetData AssetData(Asset);
+
+	if (AssetData.IsValid())
+	{
+		TArray<FAssetData> AssetsToSelect;
+		AssetsToSelect.Add(AssetData);
+
+		SyncToAssetsDelegate.Execute(AssetsToSelect);
+	}
+}
diff --git a/Engine/Plugins/2D/Paper2D/Source/Paper2DEditor/Private/SpriteEditor/SSpriteList.h b/Engine/Plugins/2D/Paper2D/Source/Paper2DEditor/Private/SpriteEditor/SSpriteList.h
--- a/Engine/Plugins/2D/Paper2D/Source/Paper2DEditor/Private/SpriteEditor/SSpriteList.h
+++ b/Engine/Plugins/2D/Paper2D/Source/Paper2DEditor/Private/SpriteEditor/SSpriteList.h
@@ -2,6 +2,8 @@
 
 #pragma once
 
+#include "ContentBrowserDelegates.h"
+
 //////////////////////////////////////////////////////////////////////////
 // SSpriteList
 
@@ -17,6 +19,8 @@ public:
 	virtual void Tick( const FGeometry& AllottedGeometry, const double InCurrentTime, const float InDeltaTime ) override;
 	// End of SWidget interface
 
+	void SelectAsset(UObject* Asset);
+
 protected:
 	void RebuildWidget(UTexture2D* NewTextureFilter);
 
@@ -33,4 +37,7 @@ protected:
 
 	// Set of tags to prevent creating details view columns for (infrequently used)
 	TSet<FName> AssetRegistryTagsToIgnore;
+
+	// Delegate to sync the asset picker to selected assets
+	FSyncToAssetsDelegate SyncToAssetsDelegate;
 };
\ No newline at end of file
diff --git a/Engine/Plugins/2D/Paper2D/Source/Paper2DEditor/Private/SpriteEditor/SpriteEditor.cpp b/Engine/Plugins/2D/Paper2D/Source/Paper2DEditor/Private/SpriteEditor/SpriteEditor.cpp
--- a/Engine/Plugins/2D/Paper2D/Source/Paper2DEditor/Private/SpriteEditor/SpriteEditor.cpp
+++ b/Engine/Plugins/2D/Paper2D/Source/Paper2DEditor/Private/SpriteEditor/SpriteEditor.cpp
@@ -293,13 +293,11 @@ TSharedRef<SDockTab> FSpriteEditor::SpawnTab_Details(const FSpawnTabArgs& Args)
 
 TSharedRef<SDockTab> FSpriteEditor::SpawnTab_SpriteList(const FSpawnTabArgs& Args)
 {
-	TSharedPtr<FSpriteEditor> SpriteEditorPtr = SharedThis(this);
-
 	// Spawn the tab
 	return SNew(SDockTab)
 		.Label(LOCTEXT("SpriteListTab_Title", "Sprite List"))
 		[
-			SNew(SSpriteList, SpriteEditorPtr)
+			SpriteListPtr.ToSharedRef()
 		];
 }
 
@@ -344,8 +342,10 @@ void FSpriteEditor::InitSpriteEditor(const EToolkitMode::Type Mode, const TShare
 
 	BindCommands();
 
-	ViewportPtr = SNew(SSpriteEditorViewport, SharedThis(this));
-	
+	TSharedPtr<FSpriteEditor> SpriteEditorPtr = SharedThis(this);
+	ViewportPtr = SNew(SSpriteEditorViewport, SpriteEditorPtr);
+	SpriteListPtr = SNew(SSpriteList, SpriteEditorPtr);
+
 	// Default layout
 	const TSharedRef<FTabManager::FLayout> StandaloneDefaultLayout = FTabManager::NewLayout("Standalone_SpriteEditor_Layout_v6")
 		->AddArea
@@ -536,6 +536,9 @@ void FSpriteEditor::SetSpriteBeingEdited(UPaperSprite* NewSprite)
 		// Let the editor know that are editing something different
 		RemoveEditingObject(OldSprite);
 		AddEditingObject(NewSprite);
+
+		// Update the asset picker to select the new active sprite
+		SpriteListPtr->SelectAsset(NewSprite);
 	}
 }
 
diff --git a/Engine/Plugins/2D/Paper2D/Source/Paper2DEditor/Private/SpriteEditor/SpriteEditor.h b/Engine/Plugins/2D/Paper2D/Source/Paper2DEditor/Private/SpriteEditor/SpriteEditor.h
--- a/Engine/Plugins/2D/Paper2D/Source/Paper2DEditor/Private/SpriteEditor/SpriteEditor.h
+++ b/Engine/Plugins/2D/Paper2D/Source/Paper2DEditor/Private/SpriteEditor/SpriteEditor.h
@@ -5,6 +5,9 @@
 #include "Toolkits/AssetEditorToolkit.h"
 #include "Toolkits/AssetEditorManager.h"
 
+class SSpriteEditorViewport;
+class SSpriteList;
+
 //////////////////////////////////////////////////////////////////////////
 // 
 
@@ -57,7 +60,8 @@ public:
 
 protected:
 	UPaperSprite* SpriteBeingEdited;
-	TSharedPtr<class SSpriteEditorViewport> ViewportPtr;
+	TSharedPtr<SSpriteEditorViewport> ViewportPtr;
+	TSharedPtr<SSpriteList> SpriteListPtr;
 
 protected:
 	void BindCommands();
