Sprite Editor: Improved source region edit mode
- Now displays the name of each sprite in the center of the region
- Actions are gated by LMB now
- Corrected help text prompt about the difference between double-clicking and Ctrl double-clicking [UE-13817]

diff --git a/Engine/Plugins/2D/Paper2D/Source/Paper2DEditor/Private/SpriteEditor/SpriteEditorViewportClient.cpp b/Engine/Plugins/2D/Paper2D/Source/Paper2DEditor/Private/SpriteEditor/SpriteEditorViewportClient.cpp
--- a/Engine/Plugins/2D/Paper2D/Source/Paper2DEditor/Private/SpriteEditor/SpriteEditorViewportClient.cpp
+++ b/Engine/Plugins/2D/Paper2D/Source/Paper2DEditor/Private/SpriteEditor/SpriteEditorViewportClient.cpp
@@ -52,8 +52,7 @@ namespace SpriteEditingConstants
 	const float BakedCollisionVertexSize = 3.0f;
 
 	const FLinearColor SourceRegionBoundsColor(1.0f, 1.0f, 1.0f, 0.8f);
-	// Alpha is being ignored in FBatchedElements::AddLine :(
-	const FLinearColor SourceRegionRelatedBoundsColor(0.3f, 0.3f, 0.3f, 0.2f);
+	const FLinearColor SourceRegionRelatedBoundsColor(0.3f, 0.3f, 0.3f, 0.8f);
 
 	const FLinearColor CollisionShapeColor(0.0f, 0.7f, 1.0f, 1.0f);
 	const FLinearColor RenderShapeColor(1.0f, 0.2f, 0.0f, 1.0f);
@@ -86,6 +85,7 @@ FSpriteEditorViewportClient::FSpriteEditorViewportClient(TWeakPtr<FSpriteEditor>
 	bShowSockets = true;
 	bShowPivot = true;
 	bShowRelatedSprites = true;
+	bShowNamesForSprites = true;
 
 	DrawHelper.bDrawGrid = GetDefault<USpriteEditorSettings>()->bShowGridByDefault;
 
@@ -220,23 +220,44 @@ FVector FSpriteEditorViewportClient::SourceTextureSpaceToWorldSpace(const FVecto
 
 void FSpriteEditorViewportClient::DrawRelatedSprites(FViewport& InViewport, FSceneView& View, FCanvas& Canvas, const FLinearColor& LineColor)
 {
+	FLinearColor ShadowColor = FLinearColor::Black;
+	ShadowColor.A = LineColor.A;
+
 	for (int32 SpriteIndex = 0; SpriteIndex < RelatedSprites.Num(); ++SpriteIndex)
 	{
 		FRelatedSprite& RelatedSprite = RelatedSprites[SpriteIndex];
 		const FVector2D& SourceUV = RelatedSprite.SourceUV;
 		const FVector2D& SourceDimension = RelatedSprite.SourceDimension;
-		FVector2D BoundsVertices[4];
-		BoundsVertices[0] = SourceTextureSpaceToScreenSpace(View, SourceUV);
-		BoundsVertices[1] = SourceTextureSpaceToScreenSpace(View, SourceUV + FVector2D(SourceDimension.X, 0));
-		BoundsVertices[2] = SourceTextureSpaceToScreenSpace(View, SourceUV + FVector2D(SourceDimension.X, SourceDimension.Y));
-		BoundsVertices[3] = SourceTextureSpaceToScreenSpace(View, SourceUV + FVector2D(0, SourceDimension.Y));
-		for (int32 VertexIndex = 0; VertexIndex < 4; ++VertexIndex)
+
+		if (bShowNamesForSprites)
 		{
-			const int32 NextVertexIndex = (VertexIndex + 1) % 4;
+			const FVector2D TextPos = SourceTextureSpaceToScreenSpace(View, SourceUV + SourceDimension*0.5f);
+
+			const FText AssetNameText = FText::AsCultureInvariant(RelatedSprite.AssetData.AssetName.ToString());
+			FCanvasTextItem TextItem(TextPos, AssetNameText, GEngine->GetSmallFont(), LineColor);
+			TextItem.EnableShadow(ShadowColor);
+			TextItem.bCentreX = true;
+			TextItem.bCentreY = true;
 
-			FCanvasLineItem LineItem(BoundsVertices[VertexIndex], BoundsVertices[NextVertexIndex]);
-			LineItem.SetColor(LineColor);
-			Canvas.DrawItem(LineItem);
+			TextItem.Draw(&Canvas);
+		}
+
+		if (bShowRelatedSprites)
+		{
+			FVector2D BoundsVertices[4];
+			BoundsVertices[0] = SourceTextureSpaceToScreenSpace(View, SourceUV);
+			BoundsVertices[1] = SourceTextureSpaceToScreenSpace(View, SourceUV + FVector2D(SourceDimension.X, 0));
+			BoundsVertices[2] = SourceTextureSpaceToScreenSpace(View, SourceUV + FVector2D(SourceDimension.X, SourceDimension.Y));
+			BoundsVertices[3] = SourceTextureSpaceToScreenSpace(View, SourceUV + FVector2D(0, SourceDimension.Y));
+
+			for (int32 VertexIndex = 0; VertexIndex < 4; ++VertexIndex)
+			{
+				const int32 NextVertexIndex = (VertexIndex + 1) % 4;
+
+				FCanvasLineItem LineItem(BoundsVertices[VertexIndex], BoundsVertices[NextVertexIndex]);
+				LineItem.SetColor(LineColor);
+				Canvas.DrawItem(LineItem);
+			}
 		}
 	}
 }
@@ -259,6 +280,20 @@ void FSpriteEditorViewportClient::DrawSourceRegion(FViewport& InViewport, FScene
 	BoundsVertices[1] = SourceTextureSpaceToScreenSpace(View, Sprite->SourceUV + FVector2D(Sprite->SourceDimension.X, 0));
 	BoundsVertices[2] = SourceTextureSpaceToScreenSpace(View, Sprite->SourceUV + FVector2D(Sprite->SourceDimension.X, Sprite->SourceDimension.Y));
 	BoundsVertices[3] = SourceTextureSpaceToScreenSpace(View, Sprite->SourceUV + FVector2D(0, Sprite->SourceDimension.Y));
+
+	if (bShowNamesForSprites)
+	{
+		const FVector2D TextPos = SourceTextureSpaceToScreenSpace(View, Sprite->SourceUV + FVector2D(Sprite->SourceDimension.X*0.5f, Sprite->SourceDimension.Y*0.5f));
+
+		const FText AssetNameText = FText::AsCultureInvariant(Sprite->GetName());
+		FCanvasTextItem TextItem(TextPos, AssetNameText, GEngine->GetSmallFont(), FLinearColor::White);
+		TextItem.EnableShadow(FLinearColor::Black);
+		TextItem.bCentreX = true;
+		TextItem.bCentreY = true;
+
+		TextItem.Draw(&Canvas);
+	}
+
 	for (int32 VertexIndex = 0; VertexIndex < 4; ++VertexIndex)
 	{
 		const int32 NextVertexIndex = (VertexIndex + 1) % 4;
@@ -452,7 +487,7 @@ void FSpriteEditorViewportClient::DrawCanvas(FViewport& Viewport, FSceneView& Vi
 
 	int32 YPos = 42;
 
-	static const FText SourceRegionHelpStr = LOCTEXT("SourceRegionHelp", "Drag handles to adjust source region\nDouble-click on an image region to create a new sprite from all connected pixels\nHold down Ctrl and drag a rectangle to create a new sprite at that position\nClick on other sprite rectangles to change the active sprite");
+	static const FText SourceRegionHelpStr = LOCTEXT("SourceRegionHelp", "Drag handles to adjust source region\nDouble-click on an image region to select all connected pixels (Ctrl creates a new sprite)\nHold down Ctrl and drag a rectangle to create a new sprite at that position\nClick on other sprite rectangles to change the active sprite");
 
 	switch (CurrentMode)
 	{
@@ -521,7 +556,7 @@ void FSpriteEditorViewportClient::DrawCanvas(FViewport& Viewport, FSceneView& Vi
 				YPos += 18;
 			}
 
-			if (bShowRelatedSprites)
+			if (bShowRelatedSprites || bShowNamesForSprites)
 			{
 				DrawRelatedSprites(Viewport, View, Canvas, SpriteEditingConstants::SourceRegionRelatedBoundsColor);
 			}
@@ -691,7 +726,7 @@ void FSpriteEditorViewportClient::ProcessClick(FSceneView& View, HHitProxy* HitP
 
 	if (IsInSourceRegionEditMode())
 	{
-		if (Event == EInputEvent::IE_DoubleClick)
+		if ((Event == EInputEvent::IE_DoubleClick) && (Key == EKeys::LeftMouseButton))
 		{
 			FVector4 WorldPoint = View.PixelToWorld(HitX, HitY, 0);
 			UPaperSprite* Sprite = GetSpriteBeingEdited();
@@ -713,7 +748,7 @@ void FSpriteEditorViewportClient::ProcessClick(FSceneView& View, HHitProxy* HitP
 				bHandled = true;
 			}
 		}
-		else
+		else if ((Event == EInputEvent::IE_Released) && (Key == EKeys::LeftMouseButton))
 		{
 			FVector4 WorldPoint = View.PixelToWorld(HitX, HitY, 0);
 			FVector2D TexturePoint = SourceTextureViewComponent->GetSprite()->ConvertWorldSpaceToTextureSpace(WorldPoint);
diff --git a/Engine/Plugins/2D/Paper2D/Source/Paper2DEditor/Private/SpriteEditor/SpriteEditorViewportClient.h b/Engine/Plugins/2D/Paper2D/Source/Paper2DEditor/Private/SpriteEditor/SpriteEditorViewportClient.h
--- a/Engine/Plugins/2D/Paper2D/Source/Paper2DEditor/Private/SpriteEditor/SpriteEditorViewportClient.h
+++ b/Engine/Plugins/2D/Paper2D/Source/Paper2DEditor/Private/SpriteEditor/SpriteEditorViewportClient.h
@@ -147,6 +147,9 @@ private:
 	// Should we show related sprites in the source texture?
 	bool bShowRelatedSprites;
 
+	// Should we show names for sprites in the source region edit mode?
+	bool bShowNamesForSprites;
+
 	// Other sprites that share the same source texture
 	TArray<FRelatedSprite> RelatedSprites;
 
