Tile Set Editor: Stopped resetting the selection when the tile index being edited hadn't really changed [UE-13620]

diff --git a/Engine/Plugins/2D/Paper2D/Source/Paper2DEditor/Private/PaperEditorShared/SpriteGeometryEditMode.cpp b/Engine/Plugins/2D/Paper2D/Source/Paper2DEditor/Private/PaperEditorShared/SpriteGeometryEditMode.cpp
--- a/Engine/Plugins/2D/Paper2D/Source/Paper2DEditor/Private/PaperEditorShared/SpriteGeometryEditMode.cpp
+++ b/Engine/Plugins/2D/Paper2D/Source/Paper2DEditor/Private/PaperEditorShared/SpriteGeometryEditMode.cpp
@@ -316,7 +316,10 @@ void FSpriteGeometryEditMode::SetGeometryColors(const FLinearColor& NewVertexCol
 
 void FSpriteGeometryEditMode::SetGeometryBeingEdited(FSpriteGeometryCollection* NewGeometryBeingEdited, bool bInAllowCircles, bool bInAllowSubtractivePolygons)
 {
-	SpriteGeometryHelper.SetGeometryBeingEdited(NewGeometryBeingEdited, bInAllowCircles, bInAllowSubtractivePolygons);
+	if (SpriteGeometryHelper.GetGeometryBeingEdited() != NewGeometryBeingEdited)
+	{
+		SpriteGeometryHelper.SetGeometryBeingEdited(NewGeometryBeingEdited, bInAllowCircles, bInAllowSubtractivePolygons);
+	}
 	bIsMarqueeTracking = false;
 }
 
