Paper2D: Moved FSpriteGeometryEditCommands out to a separate file

diff --git a/Engine/Plugins/2D/Paper2D/Source/Paper2DEditor/Private/PaperEditorShared/SpriteGeometryEditCommands.cpp b/Engine/Plugins/2D/Paper2D/Source/Paper2DEditor/Private/PaperEditorShared/SpriteGeometryEditCommands.cpp
new file mode 100644
--- /dev/null
+++ b/Engine/Plugins/2D/Paper2D/Source/Paper2DEditor/Private/PaperEditorShared/SpriteGeometryEditCommands.cpp
@@ -0,0 +1,27 @@
+// Copyright 1998-2015 Epic Games, Inc. All Rights Reserved.
+
+#include "Paper2DEditorPrivatePCH.h"
+#include "SpriteGeometryEditCommands.h"
+
+#define LOCTEXT_NAMESPACE "SpriteEditor"
+
+//////////////////////////////////////////////////////////////////////////
+// FSpriteGeometryEditCommands
+
+void FSpriteGeometryEditCommands::RegisterCommands()
+{
+	// Show toggles
+	UI_COMMAND(SetShowNormals, "Normals", "Toggles display of vertex normals in the Preview Pane.", EUserInterfaceActionType::ToggleButton, FInputChord());
+
+	// Geometry editing commands
+	UI_COMMAND(DeleteSelection, "Delete", "Delete the selection.", EUserInterfaceActionType::Button, FInputChord(EKeys::Platform_Delete, EModifierKey::None));
+	UI_COMMAND(SplitEdge, "Split", "Split edge.", EUserInterfaceActionType::Button, FInputChord(EKeys::Insert, EModifierKey::None));
+	UI_COMMAND(AddBoxShape, "Add Box", "Adds a new box shape.", EUserInterfaceActionType::Button, FInputChord());
+	UI_COMMAND(ToggleAddPolygonMode, "Add Polygon", "Adds a new polygon shape.", EUserInterfaceActionType::ToggleButton, FInputChord());
+	UI_COMMAND(AddCircleShape, "Add Circle", "Adds a new circle shape.", EUserInterfaceActionType::Button, FInputChord());
+	UI_COMMAND(SnapAllVertices, "Snap to pixel grid", "Snaps all vertices to the pixel grid.", EUserInterfaceActionType::Button, FInputChord());
+}
+
+//////////////////////////////////////////////////////////////////////////
+
+#undef LOCTEXT_NAMESPACE
\ No newline at end of file
diff --git a/Engine/Plugins/2D/Paper2D/Source/Paper2DEditor/Private/PaperEditorShared/SpriteGeometryEditCommands.h b/Engine/Plugins/2D/Paper2D/Source/Paper2DEditor/Private/PaperEditorShared/SpriteGeometryEditCommands.h
new file mode 100644
--- /dev/null
+++ b/Engine/Plugins/2D/Paper2D/Source/Paper2DEditor/Private/PaperEditorShared/SpriteGeometryEditCommands.h
@@ -0,0 +1,35 @@
+// Copyright 1998-2015 Epic Games, Inc. All Rights Reserved.
+
+#pragma once
+
+#include "PaperStyle.h"
+
+class FSpriteGeometryEditCommands : public TCommands<FSpriteGeometryEditCommands>
+{
+public:
+	FSpriteGeometryEditCommands()
+		: TCommands<FSpriteGeometryEditCommands>(
+			TEXT("SpriteGeometryEditor"), // Context name for fast lookup
+			NSLOCTEXT("Contexts", "PaperEditor", "Sprite Geometry Editor"), // Localized context name for displaying
+			NAME_None, // Parent
+			FPaperStyle::Get()->GetStyleSetName() // Icon Style Set
+			)
+	{
+	}
+
+	// TCommand<> interface
+	virtual void RegisterCommands() override;
+	// End of TCommand<> interface
+
+public:
+	// Show toggles
+	TSharedPtr<FUICommandInfo> SetShowNormals;
+
+	// Geometry editing commands
+	TSharedPtr<FUICommandInfo> DeleteSelection;
+	TSharedPtr<FUICommandInfo> SplitEdge;
+	TSharedPtr<FUICommandInfo> AddBoxShape;
+	TSharedPtr<FUICommandInfo> AddCircleShape;
+	TSharedPtr<FUICommandInfo> ToggleAddPolygonMode;
+	TSharedPtr<FUICommandInfo> SnapAllVertices;
+};
diff --git a/Engine/Plugins/2D/Paper2D/Source/Paper2DEditor/Private/PaperEditorShared/SpriteGeometryEditMode.cpp b/Engine/Plugins/2D/Paper2D/Source/Paper2DEditor/Private/PaperEditorShared/SpriteGeometryEditMode.cpp
--- a/Engine/Plugins/2D/Paper2D/Source/Paper2DEditor/Private/PaperEditorShared/SpriteGeometryEditMode.cpp
+++ b/Engine/Plugins/2D/Paper2D/Source/Paper2DEditor/Private/PaperEditorShared/SpriteGeometryEditMode.cpp
@@ -2,7 +2,7 @@
 
 #include "Paper2DEditorPrivatePCH.h"
 #include "SpriteGeometryEditMode.h"
-#include "../SpriteEditor/SpriteEditorCommands.h"
+#include "SpriteGeometryEditCommands.h"
 #include "PhysicsEngine/BodySetup2D.h"
 #include "SocketEditing.h"
 
diff --git a/Engine/Plugins/2D/Paper2D/Source/Paper2DEditor/Private/SpriteEditor/SSpriteEditorViewportToolbar.cpp b/Engine/Plugins/2D/Paper2D/Source/Paper2DEditor/Private/SpriteEditor/SSpriteEditorViewportToolbar.cpp
--- a/Engine/Plugins/2D/Paper2D/Source/Paper2DEditor/Private/SpriteEditor/SSpriteEditorViewportToolbar.cpp
+++ b/Engine/Plugins/2D/Paper2D/Source/Paper2DEditor/Private/SpriteEditor/SSpriteEditorViewportToolbar.cpp
@@ -3,6 +3,7 @@
 #include "Paper2DEditorPrivatePCH.h"
 #include "SSpriteEditorViewportToolbar.h"
 #include "SpriteEditorCommands.h"
+#include "PaperEditorShared/SpriteGeometryEditCommands.h"
 #include "SEditorViewport.h"
 
 #define LOCTEXT_NAMESPACE "SSpriteEditorViewportToolbar"
diff --git a/Engine/Plugins/2D/Paper2D/Source/Paper2DEditor/Private/SpriteEditor/SpriteEditor.cpp b/Engine/Plugins/2D/Paper2D/Source/Paper2DEditor/Private/SpriteEditor/SpriteEditor.cpp
--- a/Engine/Plugins/2D/Paper2D/Source/Paper2DEditor/Private/SpriteEditor/SpriteEditor.cpp
+++ b/Engine/Plugins/2D/Paper2D/Source/Paper2DEditor/Private/SpriteEditor/SpriteEditor.cpp
@@ -9,6 +9,7 @@
 
 #include "SpriteEditorViewportClient.h"
 #include "SpriteEditorCommands.h"
+#include "PaperEditorShared/SpriteGeometryEditCommands.h"
 #include "SEditorViewport.h"
 #include "WorkspaceMenuStructureModule.h"
 #include "Paper2DEditorModule.h"
diff --git a/Engine/Plugins/2D/Paper2D/Source/Paper2DEditor/Private/SpriteEditor/SpriteEditorCommands.cpp b/Engine/Plugins/2D/Paper2D/Source/Paper2DEditor/Private/SpriteEditor/SpriteEditorCommands.cpp
--- a/Engine/Plugins/2D/Paper2D/Source/Paper2DEditor/Private/SpriteEditor/SpriteEditorCommands.cpp
+++ b/Engine/Plugins/2D/Paper2D/Source/Paper2DEditor/Private/SpriteEditor/SpriteEditorCommands.cpp
@@ -3,24 +3,7 @@
 #include "Paper2DEditorPrivatePCH.h"
 #include "SpriteEditorCommands.h"
 
-#define LOCTEXT_NAMESPACE ""
-
-//////////////////////////////////////////////////////////////////////////
-// FSpriteGeometryEditCommands
-
-void FSpriteGeometryEditCommands::RegisterCommands()
-{
-	// Show toggles
-	UI_COMMAND(SetShowNormals, "Normals", "Toggles display of vertex normals in the Preview Pane.", EUserInterfaceActionType::ToggleButton, FInputChord());
-
-	// Geometry editing commands
-	UI_COMMAND(DeleteSelection, "Delete", "Delete the selection.", EUserInterfaceActionType::Button, FInputChord(EKeys::Platform_Delete, EModifierKey::None));
-	UI_COMMAND(SplitEdge, "Split", "Split edge.", EUserInterfaceActionType::Button, FInputChord(EKeys::Insert, EModifierKey::None));
-	UI_COMMAND(AddBoxShape, "Add Box", "Adds a new box shape.", EUserInterfaceActionType::Button, FInputChord());
-	UI_COMMAND(ToggleAddPolygonMode, "Add Polygon", "Adds a new polygon shape.", EUserInterfaceActionType::ToggleButton, FInputChord());
-	UI_COMMAND(AddCircleShape, "Add Circle", "Adds a new circle shape.", EUserInterfaceActionType::Button, FInputChord());
-	UI_COMMAND(SnapAllVertices, "Snap to pixel grid", "Snaps all vertices to the pixel grid.", EUserInterfaceActionType::Button, FInputChord());
-}
+#define LOCTEXT_NAMESPACE "SpriteEditor"
 
 //////////////////////////////////////////////////////////////////////////
 // FSpriteEditorCommands
diff --git a/Engine/Plugins/2D/Paper2D/Source/Paper2DEditor/Private/SpriteEditor/SpriteEditorCommands.h b/Engine/Plugins/2D/Paper2D/Source/Paper2DEditor/Private/SpriteEditor/SpriteEditorCommands.h
--- a/Engine/Plugins/2D/Paper2D/Source/Paper2DEditor/Private/SpriteEditor/SpriteEditorCommands.h
+++ b/Engine/Plugins/2D/Paper2D/Source/Paper2DEditor/Private/SpriteEditor/SpriteEditorCommands.h
@@ -4,37 +4,6 @@
 
 #include "PaperStyle.h"
 
-class FSpriteGeometryEditCommands : public TCommands<FSpriteGeometryEditCommands>
-{
-public:
-	FSpriteGeometryEditCommands()
-		: TCommands<FSpriteGeometryEditCommands>(
-			TEXT("SpriteGeometryEditor"), // Context name for fast lookup
-			NSLOCTEXT("Contexts", "PaperEditor", "Sprite Geometry Editor"), // Localized context name for displaying
-			NAME_None, // Parent
-			FPaperStyle::Get()->GetStyleSetName() // Icon Style Set
-			)
-	{
-	}
-
-	// TCommand<> interface
-	virtual void RegisterCommands() override;
-	// End of TCommand<> interface
-
-public:
-	// Show toggles
-	TSharedPtr<FUICommandInfo> SetShowNormals;
-
-	// Geometry editing commands
-	TSharedPtr<FUICommandInfo> DeleteSelection;
-	TSharedPtr<FUICommandInfo> SplitEdge;
-	TSharedPtr<FUICommandInfo> AddBoxShape;
-	TSharedPtr<FUICommandInfo> AddCircleShape;
-	TSharedPtr<FUICommandInfo> ToggleAddPolygonMode;
-	TSharedPtr<FUICommandInfo> SnapAllVertices;
-};
-
-
 class FSpriteEditorCommands : public TCommands<FSpriteEditorCommands>
 {
 public:
diff --git a/Engine/Plugins/2D/Paper2D/Source/Paper2DEditor/Private/TileSetEditor.cpp b/Engine/Plugins/2D/Paper2D/Source/Paper2DEditor/Private/TileSetEditor.cpp
--- a/Engine/Plugins/2D/Paper2D/Source/Paper2DEditor/Private/TileSetEditor.cpp
+++ b/Engine/Plugins/2D/Paper2D/Source/Paper2DEditor/Private/TileSetEditor.cpp
@@ -17,8 +17,8 @@
 #include "TileSetEditor/SingleTileEditorViewport.h"
 #include "TileSetEditor/SingleTileEditorViewportClient.h"
 #include "TileSetEditor/TileSetEditorCommands.h"
-#include "SpriteEditor/SpriteEditorCommands.h"
 #include "TileSetEditor/TileSetDetailsCustomization.h"
+#include "PaperEditorShared/SpriteGeometryEditCommands.h"
 #include "Editor/PropertyEditor/Public/IDetailsView.h"
 
 #define LOCTEXT_NAMESPACE "TileSetEditor"
diff --git a/Engine/Plugins/2D/Paper2D/Source/Paper2DEditor/Private/TileSetEditor/SingleTileEditorViewport.cpp b/Engine/Plugins/2D/Paper2D/Source/Paper2DEditor/Private/TileSetEditor/SingleTileEditorViewport.cpp
--- a/Engine/Plugins/2D/Paper2D/Source/Paper2DEditor/Private/TileSetEditor/SingleTileEditorViewport.cpp
+++ b/Engine/Plugins/2D/Paper2D/Source/Paper2DEditor/Private/TileSetEditor/SingleTileEditorViewport.cpp
@@ -3,7 +3,7 @@
 #include "Paper2DEditorPrivatePCH.h"
 #include "PaperStyle.h"
 #include "SingleTileEditorViewport.h"
-#include "SpriteEditor/SpriteEditorCommands.h"
+#include "PaperEditorShared/SpriteGeometryEditCommands.h"
 #include "SingleTileEditorViewportClient.h"
 #include "TileSetEditorCommands.h"
 
diff --git a/Engine/Plugins/2D/Paper2D/Source/Paper2DEditor/Private/TileSetEditor/SingleTileEditorViewportClient.cpp b/Engine/Plugins/2D/Paper2D/Source/Paper2DEditor/Private/TileSetEditor/SingleTileEditorViewportClient.cpp
--- a/Engine/Plugins/2D/Paper2D/Source/Paper2DEditor/Private/TileSetEditor/SingleTileEditorViewportClient.cpp
+++ b/Engine/Plugins/2D/Paper2D/Source/Paper2DEditor/Private/TileSetEditor/SingleTileEditorViewportClient.cpp
@@ -2,7 +2,7 @@
 
 #include "Paper2DEditorPrivatePCH.h"
 #include "SingleTileEditorViewportClient.h"
-#include "SpriteEditor/SpriteEditorCommands.h"
+#include "PaperEditorShared/SpriteGeometryEditCommands.h"
 #include "PaperEditorShared/SpriteGeometryEditMode.h"
 #include "ScopedTransaction.h"
 #include "TileSetEditorSettings.h"
