Tile Map Editor: Fixed a crash when undoing creation of layers [UE-14812]

diff --git a/Engine/Plugins/2D/Paper2D/Source/Paper2DEditor/Private/TileMapEditing/STileLayerList.cpp b/Engine/Plugins/2D/Paper2D/Source/Paper2DEditor/Private/TileMapEditing/STileLayerList.cpp
--- a/Engine/Plugins/2D/Paper2D/Source/Paper2DEditor/Private/TileMapEditing/STileLayerList.cpp
+++ b/Engine/Plugins/2D/Paper2D/Source/Paper2DEditor/Private/TileMapEditing/STileLayerList.cpp
@@ -87,8 +87,6 @@ void STileLayerList::Construct(const FArguments& InArgs, UPaperTileMap* InTileMa
 
 	TSharedRef<SWidget> Toolbar = ToolbarBuilder.MakeWidget();
 
-	RefreshMirrorList();
-
 	ListViewWidget = SNew(SPaperLayerListView)
 		.SelectionMode(ESelectionMode::Single)
 		.ClearSelectionOnClick(false)
@@ -97,6 +95,8 @@ void STileLayerList::Construct(const FArguments& InArgs, UPaperTileMap* InTileMa
 		.OnGenerateRow(this, &STileLayerList::OnGenerateLayerListRow)
 		.OnContextMenuOpening(this, &STileLayerList::OnConstructContextMenu);
 
+	RefreshMirrorList();
+
 	// Restore the selection
 	InTileMap->ValidateSelectedLayerIndex();
 	if (InTileMap->TileLayers.IsValidIndex(InTileMap->SelectedLayerIndex))
@@ -122,6 +122,8 @@ void STileLayerList::Construct(const FArguments& InArgs, UPaperTileMap* InTileMa
 			Toolbar
 		]
 	];
+
+	GEditor->RegisterForUndo(this);
 }
 
 TSharedRef<ITableRow> STileLayerList::OnGenerateLayerListRow(FMirrorEntry Item, const TSharedRef<STableViewBase>& OwnerTable)
@@ -498,8 +500,6 @@ void STileLayerList::PostEditNotfications()
 {
 	RefreshMirrorList();
 
-	ListViewWidget->RequestListRefresh();
-
 	if (UPaperTileMap* TileMap = TileMapPtr.Get())
 	{
 		TileMap->PostEditChange();
@@ -536,6 +536,23 @@ void STileLayerList::RefreshMirrorList()
 	{
 		MirrorList.Empty();
 	}
+
+	ListViewWidget->RequestListRefresh();
+}
+
+void STileLayerList::PostUndo(bool bSuccess)
+{
+	RefreshMirrorList();
+}
+
+void STileLayerList::PostRedo(bool bSuccess)
+{
+	RefreshMirrorList();
+}
+
+STileLayerList::~STileLayerList()
+{
+	GEditor->UnregisterForUndo(this);
 }
 
 //////////////////////////////////////////////////////////////////////////
diff --git a/Engine/Plugins/2D/Paper2D/Source/Paper2DEditor/Private/TileMapEditing/STileLayerList.h b/Engine/Plugins/2D/Paper2D/Source/Paper2DEditor/Private/TileMapEditing/STileLayerList.h
--- a/Engine/Plugins/2D/Paper2D/Source/Paper2DEditor/Private/TileMapEditing/STileLayerList.h
+++ b/Engine/Plugins/2D/Paper2D/Source/Paper2DEditor/Private/TileMapEditing/STileLayerList.h
@@ -5,7 +5,7 @@
 //////////////////////////////////////////////////////////////////////////
 // STileLayerList
 
-class STileLayerList : public SCompoundWidget
+class STileLayerList : public SCompoundWidget, public FEditorUndoClient
 {
 public:
 	SLATE_BEGIN_ARGS(STileLayerList) {}
@@ -13,6 +13,12 @@ public:
 
 	void Construct(const FArguments& InArgs, UPaperTileMap* TileMap, FNotifyHook* InNotifyHook, TSharedPtr<class FUICommandList> InCommandList);
 
+	// FEditorUndoClient Interface
+	virtual void PostUndo(bool bSuccess) override;
+	virtual void PostRedo(bool bSuccess) override;
+	// End of FEditorUndoClient
+
+	~STileLayerList();
 protected:
 	typedef TSharedPtr<int32> FMirrorEntry;
 
