Tile Maps: Increased the brightness of the per-tile grid

diff --git a/Engine/Plugins/2D/Paper2D/Source/Paper2D/Private/PaperTileMapRenderSceneProxy.cpp b/Engine/Plugins/2D/Paper2D/Source/Paper2D/Private/PaperTileMapRenderSceneProxy.cpp
--- a/Engine/Plugins/2D/Paper2D/Source/Paper2D/Private/PaperTileMapRenderSceneProxy.cpp
+++ b/Engine/Plugins/2D/Paper2D/Source/Paper2D/Private/PaperTileMapRenderSceneProxy.cpp
@@ -274,7 +274,7 @@ void FPaperTileMapRenderSceneProxy::GetDynamicMeshElements(const TArray<const FS
 				const uint8 DPG = SDPG_Foreground;//GetDepthPriorityGroup(View);
 
 				// Draw separation wires if selected
-				const FLinearColor OverrideColor = GetSelectionColor(FLinearColor::White, bShowAsSelected, IsHovered());
+				const FLinearColor OverrideColor = GetSelectionColor(FLinearColor::White, bShowAsSelected, IsHovered(), /*bUseOverlayIntensity=*/ false);
 
 				FTransform LocalToWorld(GetLocalToWorld());
 
