Paper2D: Switched to querying for the content path instead of hard-coding it relative to the Engine plugins dir

diff --git a/Engine/Plugins/2D/Paper2D/Source/Paper2DEditor/Paper2DEditor.Build.cs b/Engine/Plugins/2D/Paper2D/Source/Paper2DEditor/Paper2DEditor.Build.cs
--- a/Engine/Plugins/2D/Paper2D/Source/Paper2DEditor/Paper2DEditor.Build.cs
+++ b/Engine/Plugins/2D/Paper2D/Source/Paper2DEditor/Paper2DEditor.Build.cs
@@ -27,7 +27,8 @@ public class Paper2DEditor : ModuleRules
 				"WorkspaceMenuStructure",
 				"EditorStyle",
 				"MeshPaint",
-				"EditorWidgets"
+				"EditorWidgets",
+				"Projects"
 			}
 			);
 
diff --git a/Engine/Plugins/2D/Paper2D/Source/Paper2DEditor/Private/PaperStyle.cpp b/Engine/Plugins/2D/Paper2D/Source/Paper2DEditor/Private/PaperStyle.cpp
--- a/Engine/Plugins/2D/Paper2D/Source/Paper2DEditor/Private/PaperStyle.cpp
+++ b/Engine/Plugins/2D/Paper2D/Source/Paper2DEditor/Private/PaperStyle.cpp
@@ -3,6 +3,7 @@
 #include "Paper2DEditorPrivatePCH.h"
 #include "PaperStyle.h"
 #include "SlateStyle.h"
+#include "IPluginManager.h"
 
 #define IMAGE_PLUGIN_BRUSH( RelativePath, ... ) FSlateImageBrush( FPaperStyle::InContent( RelativePath, ".png" ), __VA_ARGS__ )
 #define IMAGE_BRUSH(RelativePath, ...) FSlateImageBrush(StyleSet->RootToContentDir(RelativePath, TEXT(".png")), __VA_ARGS__)
@@ -12,7 +13,7 @@
 
 FString FPaperStyle::InContent(const FString& RelativePath, const ANSICHAR* Extension)
 {
-	static FString ContentDir = FPaths::EnginePluginsDir() / TEXT("2D/Paper2D/Content");
+	static FString ContentDir = IPluginManager::Get().FindPlugin(TEXT("Paper2D"))->GetContentDir();
 	return (ContentDir / RelativePath) + Extension;
 }
 
