Paper2D: Added automatic adjustment of the tile height when the hex side length changes

diff --git a/Engine/Plugins/2D/Paper2D/Source/Paper2D/Classes/PaperTileMap.h b/Engine/Plugins/2D/Paper2D/Source/Paper2D/Classes/PaperTileMap.h
--- a/Engine/Plugins/2D/Paper2D/Source/Paper2D/Classes/PaperTileMap.h
+++ b/Engine/Plugins/2D/Paper2D/Source/Paper2D/Classes/PaperTileMap.h
@@ -49,11 +49,6 @@ class PAPER2D_API UPaperTileMap : public UObject
 	UPROPERTY(Category=Setup, EditAnywhere, BlueprintReadOnly, meta=(UIMin=1, ClampMin=1))
 	int32 TileHeight;
 
-	// The vertical height of the sides of the hex cell for a tile.
-	// Note: This value should already be included as part of the TileHeight, and is purely cosmetic; it only affects how the tile cursor preview is drawn.
-	UPROPERTY(Category=Setup, EditAnywhere, meta=(UIMin=0, ClampMin=0))
-	int32 HexSideLength;
-
 	// The scaling factor between pixels and Unreal units (cm) (e.g., 0.64 would make a 64 pixel wide tile take up 100 cm)
 	UPROPERTY(Category = Setup, EditAnywhere)
 	float PixelsPerUnrealUnit;
@@ -96,6 +91,11 @@ public:
 	UPROPERTY(Category=Setup, EditAnywhere, BlueprintReadOnly)
 	TEnumAsByte<ETileMapProjectionMode::Type> ProjectionMode;
 
+	// The vertical height of the sides of the hex cell for a tile.
+	// Note: This value should already be included as part of the TileHeight, and is purely cosmetic; it only affects how the tile cursor preview is drawn.
+	UPROPERTY(Category=Setup, EditAnywhere, meta=(UIMin=0, ClampMin=0))
+	int32 HexSideLength;
+
 	// Baked physics data.
 	UPROPERTY()
 	class UBodySetup* BodySetup;
@@ -122,6 +122,7 @@ public:
 public:
 	// UObject interface
 #if WITH_EDITOR
+	virtual void PreEditChange(UProperty* PropertyAboutToChange) override;
 	virtual void PostEditChangeProperty(FPropertyChangedEvent& PropertyChangedEvent) override;
 	virtual bool CanEditChange(const UProperty* InProperty) const override;
 	virtual void PostLoad() override;
diff --git a/Engine/Plugins/2D/Paper2D/Source/Paper2D/Private/PaperTileMap.cpp b/Engine/Plugins/2D/Paper2D/Source/Paper2D/Private/PaperTileMap.cpp
--- a/Engine/Plugins/2D/Paper2D/Source/Paper2D/Private/PaperTileMap.cpp
+++ b/Engine/Plugins/2D/Paper2D/Source/Paper2D/Private/PaperTileMap.cpp
@@ -37,6 +37,17 @@ UPaperTileMap::UPaperTileMap(const FObjectInitializer& ObjectInitializer)
 	Material = DefaultMaterial.Object;
 }
 
+#if WITH_EDITOR
+void UPaperTileMap::PreEditChange(UProperty* PropertyAboutToChange)
+{
+	if ((PropertyAboutToChange != nullptr) && (PropertyAboutToChange->GetFName() == GET_MEMBER_NAME_CHECKED(UPaperTileMap, HexSideLength)))
+	{
+		// Subtract out the hex side length; we'll add it back (along with any changes) in PostEditChangeProperty
+		TileHeight -= HexSideLength;
+	}
+}
+#endif
+
 #if WITH_EDITOR
 
 #include "PaperTileMapComponent.h"
@@ -88,6 +99,19 @@ void UPaperTileMap::PostEditChangeProperty(FPropertyChangedEvent& PropertyChange
 
 	ValidateSelectedLayerIndex();
 
+	if (PropertyName == GET_MEMBER_NAME_CHECKED(UPaperTileMap, HexSideLength))
+	{
+		HexSideLength = FMath::Max<int32>(HexSideLength, 0);
+
+		// The side length needs to be included in the overall tile height
+		TileHeight += HexSideLength;
+	}
+
+	TileWidth = FMath::Max(TileWidth, 1);
+	TileHeight = FMath::Max(TileHeight, 1);
+	MapWidth = FMath::Max(MapWidth, 1);
+	MapHeight = FMath::Max(MapHeight, 1);
+
 	if (PixelsPerUnrealUnit <= 0.0f)
 	{
 		PixelsPerUnrealUnit = 1.0f;
