Tile Maps: Made the default material for the tile map use the more explicit name MaskedUnlitSpriteMaterial

diff --git a/Engine/Plugins/2D/Paper2D/Source/Paper2D/Private/PaperTileMap.cpp b/Engine/Plugins/2D/Paper2D/Source/Paper2D/Private/PaperTileMap.cpp
--- a/Engine/Plugins/2D/Paper2D/Source/Paper2D/Private/PaperTileMap.cpp
+++ b/Engine/Plugins/2D/Paper2D/Source/Paper2D/Private/PaperTileMap.cpp
@@ -33,7 +33,7 @@ UPaperTileMap::UPaperTileMap(const FObjectInitializer& ObjectInitializer)
 
 	LayerNameIndex = 0;
 
-	static ConstructorHelpers::FObjectFinder<UMaterial> DefaultMaterial(TEXT("/Paper2D/DefaultSpriteMaterial"));
+	static ConstructorHelpers::FObjectFinder<UMaterialInterface> DefaultMaterial(TEXT("/Paper2D/MaskedUnlitSpriteMaterial"));
 	Material = DefaultMaterial.Object;
 }
 
diff --git a/Engine/Plugins/2D/Paper2D/Source/Paper2D/Private/PaperTileMapComponent.cpp b/Engine/Plugins/2D/Paper2D/Source/Paper2D/Private/PaperTileMapComponent.cpp
--- a/Engine/Plugins/2D/Paper2D/Source/Paper2D/Private/PaperTileMapComponent.cpp
+++ b/Engine/Plugins/2D/Paper2D/Source/Paper2D/Private/PaperTileMapComponent.cpp
@@ -24,7 +24,7 @@ UPaperTileMapComponent::UPaperTileMapComponent(const FObjectInitializer& ObjectI
 	TileWidth_DEPRECATED = 32;
 	TileHeight_DEPRECATED = 32;
 
-	static ConstructorHelpers::FObjectFinder<UMaterial> DefaultMaterial(TEXT("/Paper2D/DefaultSpriteMaterial"));
+	static ConstructorHelpers::FObjectFinder<UMaterialInterface> DefaultMaterial(TEXT("/Paper2D/MaskedUnlitSpriteMaterial"));
 	Material_DEPRECATED = DefaultMaterial.Object;
 
 	CastShadow = false;
