Tile Maps: Added a conditioning command to prepare tile set textures for filtering (WIP) [UE-14453]

diff --git a/Engine/Plugins/2D/Paper2D/Source/Paper2DEditor/Private/Atlasing/PaperAtlasGenerator.cpp b/Engine/Plugins/2D/Paper2D/Source/Paper2DEditor/Private/Atlasing/PaperAtlasGenerator.cpp
--- a/Engine/Plugins/2D/Paper2D/Source/Paper2DEditor/Private/Atlasing/PaperAtlasGenerator.cpp
+++ b/Engine/Plugins/2D/Paper2D/Source/Paper2DEditor/Private/Atlasing/PaperAtlasGenerator.cpp
@@ -21,7 +21,7 @@ void FPaperAtlasGenerator::HandleAssetChangedEvent(UPaperSpriteAtlas* Atlas)
 	static_assert(MaxAtlasDimension < 16384, "PaperAtlasGenerator MaxAtlasDimension exceeds multiplier in sort key");
 	Atlas->MaxWidth = FMath::Clamp(Atlas->MaxWidth, 16, MaxAtlasDimension);
 	Atlas->MaxHeight = FMath::Clamp(Atlas->MaxHeight, 16, MaxAtlasDimension);
-	Atlas->MipCount = ClampMips(Atlas->MaxWidth, Atlas->MaxHeight, Atlas->MipCount);
+	Atlas->MipCount = FPaperAtlasTextureHelpers::ClampMips(Atlas->MaxWidth, Atlas->MaxHeight, Atlas->MipCount);
 
 	bool bTestForAtlasImprovement = true;
 
@@ -219,14 +219,14 @@ void FPaperAtlasGenerator::HandleAssetChangedEvent(UPaperSpriteAtlas* Atlas)
 			if (Slot.AtlasIndex == AtlasIndex && SpriteBeingBuilt != nullptr)
 			{
 				SlotsForAtlas.Add(Slot);
-				CopySpriteToAtlasTextureData(AtlasTextureData, AtlasWidth, AtlasHeight, BytesPerPixel, Atlas->PaddingType, Atlas->Padding, SpriteBeingBuilt, Slot);
+				FPaperAtlasTextureHelpers::CopySpriteToAtlasTextureData(AtlasTextureData, AtlasWidth, AtlasHeight, BytesPerPixel, Atlas->PaddingType, Atlas->Padding, SpriteBeingBuilt, Slot);
 			}
 		}
 
 		// Mipmaps
 		if (Atlas->MipCount > 1)
 		{
-			GenerateMipChainARGB(SlotsForAtlas, AtlasTextureData, Atlas->MipCount, AtlasWidth, AtlasHeight);
+			FPaperAtlasTextureHelpers::GenerateMipChainARGB(SlotsForAtlas, AtlasTextureData, Atlas->MipCount, AtlasWidth, AtlasHeight);
 		}
 
 
diff --git a/Engine/Plugins/2D/Paper2D/Source/Paper2DEditor/Private/Atlasing/PaperAtlasTextureHelpers.cpp b/Engine/Plugins/2D/Paper2D/Source/Paper2DEditor/Private/Atlasing/PaperAtlasTextureHelpers.cpp
new file mode 100644
--- /dev/null
+++ b/Engine/Plugins/2D/Paper2D/Source/Paper2DEditor/Private/Atlasing/PaperAtlasTextureHelpers.cpp
@@ -0,0 +1,235 @@
+// Copyright 1998-2015 Epic Games, Inc. All Rights Reserved.
+
+#include "Paper2DEditorPrivatePCH.h"
+#include "PaperAtlasTextureHelpers.h"
+
+//////////////////////////////////////////////////////////////////////////
+// 
+
+void FPaperAtlasTextureHelpers::CopyTextureData(const uint8* Source, uint8* Dest, uint32 SizeX, uint32 SizeY, uint32 BytesPerPixel, uint32 SourceStride, uint32 DestStride)
+{
+	const uint32 NumBytesPerRow = SizeX * BytesPerPixel;
+
+	for (uint32 Y = 0; Y < SizeY; ++Y)
+	{
+		FMemory::Memcpy(
+			Dest + (DestStride * Y),
+			Source + (SourceStride * Y),
+			NumBytesPerRow
+			);
+	}
+}
+
+bool FPaperAtlasTextureHelpers::ReadSpriteTexture(UTexture* SourceTexture, const FIntPoint& SourceXY, const FIntPoint& SourceSize, TArray<uint8>& TargetBuffer)
+{
+	{
+		const int32 BytesPerPixel = 4;
+		TargetBuffer.Empty();
+		TargetBuffer.AddZeroed(SourceSize.X * SourceSize.Y * BytesPerPixel);
+	}
+
+	check(SourceTexture);
+	FTextureSource& SourceData = SourceTexture->Source;
+	if (SourceData.GetFormat() == TSF_BGRA8)
+	{
+		uint32 BytesPerPixel = SourceData.GetBytesPerPixel();
+		uint8* OffsetSource = SourceData.LockMip(0) + (SourceXY.X + SourceXY.Y * SourceData.GetSizeX()) * BytesPerPixel;
+		uint8* OffsetDest = TargetBuffer.GetData();
+		CopyTextureData(OffsetSource, OffsetDest, SourceSize.X, SourceSize.Y, BytesPerPixel, SourceData.GetSizeX() * BytesPerPixel, SourceSize.X * BytesPerPixel);
+		SourceData.UnlockMip(0);
+	}
+	else
+	{
+		UE_LOG(LogPaper2DEditor, Error, TEXT("Sprite texture %s is not BGRA8, which isn't supported in atlases yet"), *SourceTexture->GetName());
+	}
+
+	return true;
+}
+
+void FPaperAtlasTextureHelpers::PadSprite(const FPaperSpriteAtlasSlot& Slot, EPaperSpriteAtlasPadding PaddingType, int32 Padding, const FIntPoint& SpriteSize, int32 AtlasWidth, int32 AtlasBytesPerPixel, TArray<uint8>& TextureData)
+{
+#define UE_PIXEL(PIXEL__X, PIXEL__Y) TextureData[((Slot.Y + (PIXEL__Y)) * AtlasWidth + (Slot.X + (PIXEL__X))) * AtlasBytesPerPixel + (PixelIndex)]
+	if (PaddingType == EPaperSpriteAtlasPadding::DilateBorder)
+	{
+		for (int32 X = 0; X < Padding; ++X)
+		{
+			for (int32 Y = 0; Y < SpriteSize.Y + Padding * 2; ++Y)
+			{
+				for (int32 PixelIndex = 0; PixelIndex < AtlasBytesPerPixel; ++PixelIndex)
+				{
+					const int32 ClampedY = Padding + FMath::Clamp(Y - Padding, 0, SpriteSize.Y - 1);
+					UE_PIXEL(X, Y) = UE_PIXEL(Padding, ClampedY);
+					UE_PIXEL(Padding + SpriteSize.X + X, Y) = UE_PIXEL(Padding + SpriteSize.X - 1, ClampedY);
+				}
+			}
+		}
+		for (int32 Y = 0; Y < Padding; ++Y)
+		{
+			for (int32 X = 0; X < SpriteSize.X + Padding * 2; ++X)
+			{
+				for (int32 PixelIndex = 0; PixelIndex < AtlasBytesPerPixel; ++PixelIndex)
+				{
+					const int32 ClampedX = Padding + FMath::Clamp(X - Padding, 0, SpriteSize.X - 1);
+					UE_PIXEL(X, Y) = UE_PIXEL(ClampedX, Padding);
+					UE_PIXEL(X, Padding + SpriteSize.Y + Y) = UE_PIXEL(ClampedX, Padding + SpriteSize.Y - 1);
+				}
+			}
+		}
+	}
+	else if (PaddingType == EPaperSpriteAtlasPadding::PadWithZero)
+	{
+		for (int32 X = 0; X < Padding; ++X)
+		{
+			for (int32 Y = 0; Y < SpriteSize.Y + Padding * 2; ++Y)
+			{
+				for (int32 PixelIndex = 0; PixelIndex < AtlasBytesPerPixel; ++PixelIndex)
+				{
+					UE_PIXEL(X, Y) = 0;
+					UE_PIXEL(Padding + SpriteSize.X + X, Y) = 0;
+				}
+			}
+		}
+		for (int32 Y = 0; Y < Padding; ++Y)
+		{
+			for (int32 X = 0; X < SpriteSize.X + Padding * 2; ++X)
+			{
+				for (int32 PixelIndex = 0; PixelIndex < AtlasBytesPerPixel; ++PixelIndex)
+				{
+					UE_PIXEL(X, Y) = 0;
+					UE_PIXEL(X, Padding + SpriteSize.Y + Y) = 0;
+				}
+			}
+		}
+	}
+#undef UE_PIXEL
+}
+
+void FPaperAtlasTextureHelpers::CopySpriteToAtlasTextureData(TArray<uint8>& TextureData, int32 AtlasWidth, int32 AtlasHeight, int32 AtlasBytesPerPixel, EPaperSpriteAtlasPadding PaddingType, int32 Padding, UPaperSprite* Sprite, const FPaperSpriteAtlasSlot& Slot)
+{
+	const FVector2D SpriteSizeFloat = Sprite->GetSourceSize();
+	const FIntPoint SpriteSize(FMath::TruncToInt(Sprite->GetSourceSize().X), FMath::TruncToInt(Sprite->GetSourceSize().Y));
+	const FIntPoint SpriteXY(FMath::TruncToInt(Sprite->GetSourceUV().X), FMath::TruncToInt(Sprite->GetSourceUV().Y));
+
+	TArray<uint8> DummyBuffer;
+	ReadSpriteTexture(Sprite->GetSourceTexture(), SpriteXY, SpriteSize, DummyBuffer);
+
+	CopyTextureRegionToAtlasTextureData(TextureData, AtlasWidth, AtlasHeight, AtlasBytesPerPixel, PaddingType, Padding, DummyBuffer, SpriteSize, Slot);
+}
+
+void FPaperAtlasTextureHelpers::CopyTextureRegionToAtlasTextureData(TArray<uint8>& TextureData, int32 AtlasWidth, int32 AtlasHeight, int32 AtlasBytesPerPixel, EPaperSpriteAtlasPadding PaddingType, int32 Padding, TArray<uint8>& SourceData, const FIntPoint& SourceSize, const FPaperSpriteAtlasSlot& Slot)
+{
+	// Copy the texture into the texture buffer
+	for (int32 Y = 0; Y < SourceSize.Y; ++Y)
+	{
+		for (int32 X = 0; X < SourceSize.X; ++X)
+		{
+			for (int32 PixelIndex = 0; PixelIndex < AtlasBytesPerPixel; ++PixelIndex)
+			{
+				TextureData[((Slot.Y + Y + Padding) * AtlasWidth + (Slot.X + X + Padding)) * AtlasBytesPerPixel + PixelIndex] = SourceData[(Y * SourceSize.X + X) * AtlasBytesPerPixel + PixelIndex];
+			}
+		}
+	}
+
+	// Padding
+	PadSprite(Slot, PaddingType, Padding, SourceSize, AtlasWidth, AtlasBytesPerPixel, TextureData);
+}
+
+int32 FPaperAtlasTextureHelpers::ClampMips(int Width, int Height, int MipCount)
+{
+	int32 NumMips = 1;
+	while (MipCount > 1 && (Width % 2) == 0 && (Height % 2) == 0)
+	{
+		Width /= 2;
+		Height /= 2;
+		MipCount--;
+		NumMips++;
+	}
+	return NumMips;
+}
+
+void FPaperAtlasTextureHelpers::GenerateMipChainARGB(const TArray<FPaperSpriteAtlasSlot>& Slots, TArray<uint8>& AtlasTextureData, int32 MipCount, int32 Width, int32 Height)
+{
+	int32 SourceMipOffset = 0;
+	int32 SourceMipWidth = Width;
+	int32 SourceMipHeight = Height;
+
+	// Mask bitmap stores all valid pixels in the image, i.e. of the rects that contain data
+	TArray<int8> MaskBitmap;
+	MaskBitmap.AddZeroed(SourceMipWidth * SourceMipHeight);
+	for (const FPaperSpriteAtlasSlot& Slot : Slots)
+	{
+		for (int32 Y = 0; Y < Slot.Height; ++Y)
+		{
+			for (int32 X = 0; X < Slot.Width; ++X)
+			{
+				MaskBitmap[(Slot.Y + Y) * SourceMipWidth + (Slot.X + X)] = 1;
+			}
+		}
+	}
+
+	int32 BytesPerPixel = 4;
+	check(BytesPerPixel == 4); // Only support 4 bytes per pixel
+
+	// Mipmap offset to mip level being written into
+	int32 TargetMipOffset = SourceMipWidth * SourceMipHeight * BytesPerPixel;
+	int32 TargetMipWidth = SourceMipWidth / 2;
+	int32 TargetMipHeight = SourceMipHeight / 2;
+
+	for (int32 MipIndex = 1; MipIndex < MipCount; ++MipIndex)
+	{
+		int32 MipLevelSize = TargetMipHeight * TargetMipWidth * BytesPerPixel;
+
+		for (int32 Y = 0; Y < TargetMipHeight; ++Y)
+		{
+			for (int32 X = 0; X < TargetMipWidth; ++X)
+			{
+				int32 TotalChannelValues[4] = { 0, 0, 0, 0 };
+				check(BytesPerPixel == 4); // refer to above
+
+				int32 ValidPixelCount = 0;
+				int32 Mask0 = MaskBitmap[(Y * 2 + 0) * SourceMipWidth + (X * 2 + 0)];
+				int32 Mask1 = MaskBitmap[(Y * 2 + 0) * SourceMipWidth + (X * 2 + 1)];
+				int32 Mask2 = MaskBitmap[(Y * 2 + 1) * SourceMipWidth + (X * 2 + 0)];
+				int32 Mask3 = MaskBitmap[(Y * 2 + 1) * SourceMipWidth + (X * 2 + 1)];
+				ValidPixelCount = Mask0 + Mask1 + Mask2 + Mask3;
+
+				for (int32 PixelIndex = 0; PixelIndex < BytesPerPixel; ++PixelIndex)
+				{
+					TotalChannelValues[PixelIndex] += AtlasTextureData[SourceMipOffset + ((Y * 2 + 0) * SourceMipWidth + (X * 2 + 0)) * BytesPerPixel + PixelIndex];
+					TotalChannelValues[PixelIndex] += AtlasTextureData[SourceMipOffset + ((Y * 2 + 0) * SourceMipWidth + (X * 2 + 1)) * BytesPerPixel + PixelIndex];
+					TotalChannelValues[PixelIndex] += AtlasTextureData[SourceMipOffset + ((Y * 2 + 1) * SourceMipWidth + (X * 2 + 0)) * BytesPerPixel + PixelIndex];
+					TotalChannelValues[PixelIndex] += AtlasTextureData[SourceMipOffset + ((Y * 2 + 1) * SourceMipWidth + (X * 2 + 1)) * BytesPerPixel + PixelIndex];
+				}
+
+				for (int32 PixelIndex = 0; PixelIndex < BytesPerPixel; ++PixelIndex)
+				{
+					uint8 TargetPixelValue = (ValidPixelCount > 0) ? (TotalChannelValues[PixelIndex] / ValidPixelCount) : 0;
+					AtlasTextureData[TargetMipOffset + (Y * TargetMipWidth + X) * BytesPerPixel + PixelIndex] = TargetPixelValue;
+				}
+			}
+		}
+
+		// Downsample mask in place, if any of the 4 mask pixels == 1, the target mask should be 1
+		for (int32 Y = 0; Y < TargetMipHeight; ++Y)
+		{
+			for (int32 X = 0; X < TargetMipWidth; ++X)
+			{
+				uint8 TargetPixelValue = 0;
+				TargetPixelValue |= MaskBitmap[(Y * 2 + 0) * SourceMipWidth + (X * 2 + 0)];
+				TargetPixelValue |= MaskBitmap[(Y * 2 + 0) * SourceMipWidth + (X * 2 + 1)];
+				TargetPixelValue |= MaskBitmap[(Y * 2 + 1) * SourceMipWidth + (X * 2 + 0)];
+				TargetPixelValue |= MaskBitmap[(Y * 2 + 1) * SourceMipWidth + (X * 2 + 1)];
+				MaskBitmap[Y * TargetMipWidth + X] = TargetPixelValue;
+			}
+		}
+
+		// Update for next mip
+		SourceMipOffset = TargetMipOffset;
+		SourceMipWidth = TargetMipWidth;
+		SourceMipHeight = TargetMipHeight;
+		TargetMipOffset += MipLevelSize;
+		TargetMipWidth /= 2;
+		TargetMipHeight /= 2;
+	}
+}
+
diff --git a/Engine/Plugins/2D/Paper2D/Source/Paper2DEditor/Private/Atlasing/PaperAtlasTextureHelpers.h b/Engine/Plugins/2D/Paper2D/Source/Paper2DEditor/Private/Atlasing/PaperAtlasTextureHelpers.h
--- a/Engine/Plugins/2D/Paper2D/Source/Paper2DEditor/Private/Atlasing/PaperAtlasTextureHelpers.h
+++ b/Engine/Plugins/2D/Paper2D/Source/Paper2DEditor/Private/Atlasing/PaperAtlasTextureHelpers.h
@@ -2,229 +2,28 @@
 
 #pragma once
 
-// Copy texture data from one 32bpp image to another
-static void CopyTextureData(const uint8* Source, uint8* Dest, uint32 SizeX, uint32 SizeY, uint32 BytesPerPixel, uint32 SourceStride, uint32 DestStride)
+struct FPaperAtlasTextureHelpers
 {
-	const uint32 NumBytesPerRow = SizeX * BytesPerPixel;
+public:
+	// Copy texture data from one 32bpp image to another
+	static void CopyTextureData(const uint8* Source, uint8* Dest, uint32 SizeX, uint32 SizeY, uint32 BytesPerPixel, uint32 SourceStride, uint32 DestStride);
 
-	for (uint32 Y = 0; Y < SizeY; ++Y)
-	{
-		FMemory::Memcpy(
-			Dest + (DestStride * Y),
-			Source + (SourceStride * Y),
-			NumBytesPerRow
-			);
-	}
-}
+	// Read sprite data into the 4 bytes per pixel target array
+	static bool ReadSpriteTexture(UTexture* SourceTexture, const FIntPoint& SourceXY, const FIntPoint& SourceSize, TArray<uint8>& TargetBuffer);
 
-// Read sprite data into the 4 bytes per pixel target array
-static bool ReadSpriteTexture(TArray<uint8>& TargetBuffer, FIntPoint& OutSpriteSize, UPaperSprite* Sprite)
-{
-	const FVector2D SpriteSizeFloat = Sprite->GetSourceSize();
-	const FIntPoint SpriteSize(FMath::TruncToInt(SpriteSizeFloat.X), FMath::TruncToInt(SpriteSizeFloat.Y));
-	{
-		const int32 BytesPerPixel = 4;
-		TargetBuffer.Empty();
-		TargetBuffer.AddZeroed(SpriteSize.X * SpriteSize.Y * BytesPerPixel);
-	}
-
-	check(Sprite->GetSourceTexture());
-	FTextureSource& SourceData = Sprite->GetSourceTexture()->Source;
-	if (SourceData.GetFormat() == TSF_BGRA8)
-	{
-		uint32 BytesPerPixel = SourceData.GetBytesPerPixel();
-		uint8* OffsetSource = SourceData.LockMip(0) + (FMath::TruncToInt(Sprite->GetSourceUV().X) + FMath::TruncToInt(Sprite->GetSourceUV().Y) * SourceData.GetSizeX()) * BytesPerPixel;
-		uint8* OffsetDest = TargetBuffer.GetData();
-		CopyTextureData(OffsetSource, OffsetDest, SpriteSize.X, SpriteSize.Y, BytesPerPixel, SourceData.GetSizeX() * BytesPerPixel, SpriteSize.X * BytesPerPixel);
-		SourceData.UnlockMip(0);
-	}
-	else
-	{
-		UE_LOG(LogPaper2DEditor, Error, TEXT("Sprite %s is not BGRA8, which isn't supported in atlases yet"), *(Sprite->GetPathName()));
-	}
-
-	OutSpriteSize = SpriteSize;
-	return true;
-}
-
-// Fills the padding space with the correct values
-static void PadSprite(const FPaperSpriteAtlasSlot& Slot, EPaperSpriteAtlasPadding PaddingType, int32 Padding, FIntPoint &SpriteSize, int32 AtlasWidth, int32 AtlasBytesPerPixel, TArray<uint8>& TextureData)
-{
-#define PIXEL(PIXEL__X, PIXEL__Y) TextureData[((Slot.Y + (PIXEL__Y)) * AtlasWidth + (Slot.X + (PIXEL__X))) * AtlasBytesPerPixel + (PixelIndex)]
-	if (PaddingType == EPaperSpriteAtlasPadding::DilateBorder)
-	{
-		for (int32 X = 0; X < Padding; ++X)
-		{
-			for (int32 Y = 0; Y < SpriteSize.Y + Padding * 2; ++Y)
-			{
-				for (int32 PixelIndex = 0; PixelIndex < AtlasBytesPerPixel; ++PixelIndex)
-				{
-					int32 ClampedY = Padding + FMath::Clamp(Y - Padding, 0, SpriteSize.Y - 1);
-					PIXEL(X, Y) = PIXEL(Padding, ClampedY);
-					PIXEL(Padding + SpriteSize.X + X, Y) = PIXEL(Padding + SpriteSize.X - 1, ClampedY);
-				}
-			}
-		}
-		for (int32 Y = 0; Y < Padding; ++Y)
-		{
-			for (int32 X = 0; X < SpriteSize.X + Padding * 2; ++X)
-			{
-				for (int32 PixelIndex = 0; PixelIndex < AtlasBytesPerPixel; ++PixelIndex)
-				{
-					int32 ClampedX = Padding + FMath::Clamp(X - Padding, 0, SpriteSize.X - 1);
-					PIXEL(X, Y) = PIXEL(ClampedX, Padding);
-					PIXEL(X, Padding + SpriteSize.Y + Y) = PIXEL(ClampedX, Padding + SpriteSize.Y - 1);
-				}
-			}
-		}
-	}
-	else if (PaddingType == EPaperSpriteAtlasPadding::PadWithZero)
-	{
-		for (int32 X = 0; X < Padding; ++X)
-		{
-			for (int32 Y = 0; Y < SpriteSize.Y + Padding * 2; ++Y)
-			{
-				for (int32 PixelIndex = 0; PixelIndex < AtlasBytesPerPixel; ++PixelIndex)
-				{
-					PIXEL(X, Y) = 0;
-					PIXEL(Padding + SpriteSize.X + X, Y) = 0;
-				}
-			}
-		}
-		for (int32 Y = 0; Y < Padding; ++Y)
-		{
-			for (int32 X = 0; X < SpriteSize.X + Padding * 2; ++X)
-			{
-				for (int32 PixelIndex = 0; PixelIndex < AtlasBytesPerPixel; ++PixelIndex)
-				{
-					PIXEL(X, Y) = 0;
-					PIXEL(X, Padding + SpriteSize.Y + Y) = 0;
-				}
-			}
-		}
-	}
-#undef PIXEL
-}
-
-// Copies the sprite texture data at the position described by slot
-static void CopySpriteToAtlasTextureData(TArray<uint8>& TextureData, int32 AtlasWidth, int32 AtlasHeight, int32 AtlasBytesPerPixel, EPaperSpriteAtlasPadding PaddingType, int32 Padding, UPaperSprite* Sprite, const FPaperSpriteAtlasSlot& Slot)
-{
-	FIntPoint SpriteSize;
-	TArray<uint8> DummyBuffer;
-	ReadSpriteTexture(DummyBuffer, /*out*/SpriteSize, Sprite);
-
-	// Copy the texture into the texture buffer
-	for (int32 Y = 0; Y < SpriteSize.Y; ++Y)
-	{
-		for (int32 X = 0; X < SpriteSize.X; ++X)
-		{
-			for (int32 PixelIndex = 0; PixelIndex < AtlasBytesPerPixel; ++PixelIndex)
-			{
-				TextureData[((Slot.Y + Y + Padding) * AtlasWidth + (Slot.X + X + Padding)) * AtlasBytesPerPixel + PixelIndex] = DummyBuffer[(Y * SpriteSize.X + X) * AtlasBytesPerPixel + PixelIndex];
-			}
-		}
-	}
-
-	// Padding
-	PadSprite(Slot, PaddingType, Padding, SpriteSize, AtlasWidth, AtlasBytesPerPixel, TextureData);
-}
-
-static int32 ClampMips(int Width, int Height, int MipCount)
-{
-	int32 NumMips = 1;
-	while (MipCount > 1 && (Width % 2) == 0 && (Height % 2) == 0)
-	{
-		Width /= 2;
-		Height /= 2;
-		MipCount--;
-		NumMips++;
-	}
-	return NumMips;
-}
-
-static void GenerateMipChainARGB(const TArray<FPaperSpriteAtlasSlot>& Slots, TArray<uint8>& AtlasTextureData, int32 MipCount, int32 Width, int32 Height)
-{
-	int32 SourceMipOffset = 0;
-	int32 SourceMipWidth = Width;
-	int32 SourceMipHeight = Height;
-
-	// Mask bitmap stores all valid pixels in the image, i.e. of the rects that contain data
-	TArray<int8> MaskBitmap;
-	MaskBitmap.AddZeroed(SourceMipWidth * SourceMipHeight);
-	for (const FPaperSpriteAtlasSlot& Slot : Slots)
-	{
-		for (int32 Y = 0; Y < Slot.Height; ++Y)
-		{
-			for (int32 X = 0; X < Slot.Width; ++X)
-			{
-				MaskBitmap[(Slot.Y + Y) * SourceMipWidth + (Slot.X + X)] = 1;
-			}
-		}
-	}
-
-	int32 BytesPerPixel = 4;
-	check(BytesPerPixel == 4); // Only support 4 bytes per pixel
-
-	// Mipmap offset to mip level being written into
-	int32 TargetMipOffset = SourceMipWidth * SourceMipHeight * BytesPerPixel;
-	int32 TargetMipWidth = SourceMipWidth / 2;
-	int32 TargetMipHeight = SourceMipHeight / 2;
-
-	for (int32 MipIndex = 1; MipIndex < MipCount; ++MipIndex)
-	{
-		int32 MipLevelSize = TargetMipHeight * TargetMipWidth * BytesPerPixel;
-
-		for (int32 Y = 0; Y < TargetMipHeight; ++Y)
-		{
-			for (int32 X = 0; X < TargetMipWidth; ++X)
-			{
-				int32 TotalChannelValues[4] = { 0, 0, 0, 0 };
-				check(BytesPerPixel == 4); // refer to above
-
-				int32 ValidPixelCount = 0;
-				int32 Mask0 = MaskBitmap[(Y * 2 + 0) * SourceMipWidth + (X * 2 + 0)];
-				int32 Mask1 = MaskBitmap[(Y * 2 + 0) * SourceMipWidth + (X * 2 + 1)];
-				int32 Mask2 = MaskBitmap[(Y * 2 + 1) * SourceMipWidth + (X * 2 + 0)];
-				int32 Mask3 = MaskBitmap[(Y * 2 + 1) * SourceMipWidth + (X * 2 + 1)];
-				ValidPixelCount = Mask0 + Mask1 + Mask2 + Mask3;
+	// Fills the padding space with the correct values
+	static void PadSprite(const FPaperSpriteAtlasSlot& Slot, EPaperSpriteAtlasPadding PaddingType, int32 Padding, const FIntPoint& SpriteSize, int32 AtlasWidth, int32 AtlasBytesPerPixel, TArray<uint8>& TextureData);
 
-				for (int32 PixelIndex = 0; PixelIndex < BytesPerPixel; ++PixelIndex)
-				{
-					TotalChannelValues[PixelIndex] += AtlasTextureData[SourceMipOffset + ((Y * 2 + 0) * SourceMipWidth + (X * 2 + 0)) * BytesPerPixel + PixelIndex];
-					TotalChannelValues[PixelIndex] += AtlasTextureData[SourceMipOffset + ((Y * 2 + 0) * SourceMipWidth + (X * 2 + 1)) * BytesPerPixel + PixelIndex];
-					TotalChannelValues[PixelIndex] += AtlasTextureData[SourceMipOffset + ((Y * 2 + 1) * SourceMipWidth + (X * 2 + 0)) * BytesPerPixel + PixelIndex];
-					TotalChannelValues[PixelIndex] += AtlasTextureData[SourceMipOffset + ((Y * 2 + 1) * SourceMipWidth + (X * 2 + 1)) * BytesPerPixel + PixelIndex];
-				}
+	// Copies the sprite texture data at the position described by slot
+	static void CopySpriteToAtlasTextureData(TArray<uint8>& TextureData, int32 AtlasWidth, int32 AtlasHeight, int32 AtlasBytesPerPixel, EPaperSpriteAtlasPadding PaddingType, int32 Padding, UPaperSprite* Sprite, const FPaperSpriteAtlasSlot& Slot);
 
-				for (int32 PixelIndex = 0; PixelIndex < BytesPerPixel; ++PixelIndex)
-				{
-					uint8 TargetPixelValue = (ValidPixelCount > 0) ? (TotalChannelValues[PixelIndex] / ValidPixelCount) : 0;
-					AtlasTextureData[TargetMipOffset + (Y * TargetMipWidth + X) * BytesPerPixel + PixelIndex] = TargetPixelValue;
-				}
-			}
-		}
+	// Copies the texture data at the position described by slot
+	static void CopyTextureRegionToAtlasTextureData(TArray<uint8>& TextureData, int32 AtlasWidth, int32 AtlasHeight, int32 AtlasBytesPerPixel, EPaperSpriteAtlasPadding PaddingType, int32 Padding, TArray<uint8>& SourceData, const FIntPoint& SourceSize, const FPaperSpriteAtlasSlot& Slot);
 
-		// Downsample mask in place, if any of the 4 mask pixels == 1, the target mask should be 1
-		for (int32 Y = 0; Y < TargetMipHeight; ++Y)
-		{
-			for (int32 X = 0; X < TargetMipWidth; ++X)
-			{
-				uint8 TargetPixelValue = 0;
-				TargetPixelValue |= MaskBitmap[(Y * 2 + 0) * SourceMipWidth + (X * 2 + 0)];
-				TargetPixelValue |= MaskBitmap[(Y * 2 + 0) * SourceMipWidth + (X * 2 + 1)];
-				TargetPixelValue |= MaskBitmap[(Y * 2 + 1) * SourceMipWidth + (X * 2 + 0)];
-				TargetPixelValue |= MaskBitmap[(Y * 2 + 1) * SourceMipWidth + (X * 2 + 1)];
-				MaskBitmap[Y * TargetMipWidth + X] = TargetPixelValue;
-			}
-		}
+	static int32 ClampMips(int Width, int Height, int MipCount);
 
-		// Update for next mip
-		SourceMipOffset = TargetMipOffset;
-		SourceMipWidth = TargetMipWidth;
-		SourceMipHeight = TargetMipHeight;
-		TargetMipOffset += MipLevelSize;
-		TargetMipWidth /= 2;
-		TargetMipHeight /= 2;
-	}
-}
+	static void GenerateMipChainARGB(const TArray<FPaperSpriteAtlasSlot>& Slots, TArray<uint8>& AtlasTextureData, int32 MipCount, int32 Width, int32 Height);
 
+private:
+	FPaperAtlasTextureHelpers() {}
+};
\ No newline at end of file
diff --git a/Engine/Plugins/2D/Paper2D/Source/Paper2DEditor/Private/TileSetAssetTypeActions.cpp b/Engine/Plugins/2D/Paper2D/Source/Paper2DEditor/Private/TileSetAssetTypeActions.cpp
--- a/Engine/Plugins/2D/Paper2D/Source/Paper2DEditor/Private/TileSetAssetTypeActions.cpp
+++ b/Engine/Plugins/2D/Paper2D/Source/Paper2DEditor/Private/TileSetAssetTypeActions.cpp
@@ -6,6 +6,7 @@
 #include "AssetToolsModule.h"
 #include "ContentBrowserModule.h"
 #include "PaperTileMapFactory.h"
+#include "TileSetEditor/TileSheetPaddingFactory.h"
 
 #define LOCTEXT_NAMESPACE "AssetTypeActions"
 
@@ -63,6 +64,13 @@ void FTileSetAssetTypeActions::GetActions(const TArray<UObject*>& InObjects, FMe
 			FSlateIcon(FEditorStyle::GetStyleSetName(), "ClassIcon.PaperTileSet"),
 			FUIAction(FExecuteAction::CreateSP(this, &FTileSetAssetTypeActions::ExecuteCreateTileMap, TileSets[0]))
 			);
+
+		MenuBuilder.AddMenuEntry(
+			LOCTEXT("TileSet_ConditionTileSet", "Condition Tile Sheet Texture"),
+			LOCTEXT("TileSet_ConditionTileSetTooltip", "Conditions the tile sheet texture for the selected tile set by duplicating tile edges to create a buffer zone around each tile"),
+			FSlateIcon(FEditorStyle::GetStyleSetName(), "ClassIcon.Texture2D"),
+			FUIAction(FExecuteAction::CreateSP(this, &FTileSetAssetTypeActions::ExecutePadTileSetTexture, TileSets[0]))
+			);
 	}
 }
 
@@ -82,7 +90,7 @@ void FTileSetAssetTypeActions::ExecuteCreateTileMap(TWeakObjectPtr<UPaperTileSet
 		const FString TileSetPathName = TileSet->GetOutermost()->GetPathName();
 		const FString LongPackagePath = FPackageName::GetLongPackagePath(TileSetPathName);
 
-		const FString NewTileMapDefaultPath = LongPackagePath + TEXT("/") + EffectiveTileSetName;
+		const FString NewTileMapDefaultPath = LongPackagePath / EffectiveTileSetName;
 
 		// Make sure the name is unique
 		FString AssetName;
@@ -97,6 +105,36 @@ void FTileSetAssetTypeActions::ExecuteCreateTileMap(TWeakObjectPtr<UPaperTileSet
 	}
 }
 
+void FTileSetAssetTypeActions::ExecutePadTileSetTexture(TWeakObjectPtr<UPaperTileSet> TileSetPtr)
+{
+	FAssetToolsModule& AssetToolsModule = FModuleManager::GetModuleChecked<FAssetToolsModule>("AssetTools");
+	FContentBrowserModule& ContentBrowserModule = FModuleManager::LoadModuleChecked<FContentBrowserModule>("ContentBrowser");
+
+	if (UPaperTileSet* TileSet = TileSetPtr.Get())
+	{
+		if (TileSet->TileSheet != nullptr)
+		{
+			const FString TileSheetSuffix(TEXT("Padded"));
+			const FString TileSheetPathName = TileSet->TileSheet->GetOutermost()->GetPathName();
+			const FString LongPackagePath = FPackageName::GetLongPackagePath(TileSheetPathName);
+
+			const FString EffectiveTileSheetName = TileSet->TileSheet->GetName();
+			const FString NewTileSheetDefaultPath = LongPackagePath / EffectiveTileSheetName;
+
+			// Make sure the name is unique
+			FString AssetName;
+			FString PackageName;
+			AssetToolsModule.Get().CreateUniqueAssetName(NewTileSheetDefaultPath, TileSheetSuffix, /*out*/ PackageName, /*out*/ AssetName);
+			const FString PackagePath = FPackageName::GetLongPackagePath(PackageName);
+
+			// Create the new tile sheet
+			UTileSheetPaddingFactory* TileSheetPaddingFactory = NewObject<UTileSheetPaddingFactory>();
+			TileSheetPaddingFactory->SourceTileSet = TileSet;
+			ContentBrowserModule.Get().CreateNewAsset(AssetName, PackagePath, UTexture::StaticClass(), TileSheetPaddingFactory);
+		}
+	}
+}
+
 //////////////////////////////////////////////////////////////////////////
 
 #undef LOCTEXT_NAMESPACE
\ No newline at end of file
diff --git a/Engine/Plugins/2D/Paper2D/Source/Paper2DEditor/Private/TileSetAssetTypeActions.h b/Engine/Plugins/2D/Paper2D/Source/Paper2DEditor/Private/TileSetAssetTypeActions.h
--- a/Engine/Plugins/2D/Paper2D/Source/Paper2DEditor/Private/TileSetAssetTypeActions.h
+++ b/Engine/Plugins/2D/Paper2D/Source/Paper2DEditor/Private/TileSetAssetTypeActions.h
@@ -26,4 +26,5 @@ private:
 
 private:
 	void ExecuteCreateTileMap(TWeakObjectPtr<UPaperTileSet> TileSetPtr);
+	void ExecutePadTileSetTexture(TWeakObjectPtr<UPaperTileSet> TileSetPtr);
 };
diff --git a/Engine/Plugins/2D/Paper2D/Source/Paper2DEditor/Private/TileSetEditor/TileSheetPaddingFactory.cpp b/Engine/Plugins/2D/Paper2D/Source/Paper2DEditor/Private/TileSetEditor/TileSheetPaddingFactory.cpp
new file mode 100644
--- /dev/null
+++ b/Engine/Plugins/2D/Paper2D/Source/Paper2DEditor/Private/TileSetEditor/TileSheetPaddingFactory.cpp
@@ -0,0 +1,95 @@
+// Copyright 1998-2015 Epic Games, Inc. All Rights Reserved.
+
+#include "Paper2DEditorPrivatePCH.h"
+#include "TileSheetPaddingFactory.h"
+#include "Atlasing/PaperAtlasTextureHelpers.h"
+
+#define LOCTEXT_NAMESPACE "Paper2D"
+
+/////////////////////////////////////////////////////
+// UTileSheetPaddingFactory
+
+UTileSheetPaddingFactory::UTileSheetPaddingFactory(const FObjectInitializer& ObjectInitializer)
+	: Super(ObjectInitializer)
+{
+	bCreateNew = false;
+	bEditAfterNew = true;
+	SupportedClass = UTexture::StaticClass();
+
+	ExtrusionAmount = 2;
+	bPadToPowerOf2 = false;
+}
+
+UObject* UTileSheetPaddingFactory::FactoryCreateNew(UClass* Class, UObject* InParent, FName Name, EObjectFlags Flags, UObject* Context, FFeedbackContext* Warn)
+{
+	check(SourceTileSet);
+	check(SourceTileSet->TileSheet);
+
+	UTexture* SourceTexture = SourceTileSet->TileSheet;
+
+	if (SourceTexture->Source.GetFormat() != TSF_BGRA8)
+	{
+		UE_LOG(LogPaper2DEditor, Error, TEXT("Tile sheet texture '%s' is not BGRA8, cannot create a padded texture from it"), *SourceTexture->GetName());
+		return nullptr;
+	}
+
+	// Determine how big the new texture needs to be
+	const int32 NumTilesX = SourceTileSet->GetTileCountX();
+	const int32 NumTilesY = SourceTileSet->GetTileCountY();
+	const int32 TileWidth = SourceTileSet->TileWidth;
+	const int32 TileHeight = SourceTileSet->TileHeight;
+
+	const uint32 NewMinTextureWidth = (uint32)(NumTilesX * (TileWidth + 2 * ExtrusionAmount) + (2 * ExtrusionAmount));
+	const uint32 NewMinTextureHeight = (uint32)(NumTilesY * (TileHeight + 2 * ExtrusionAmount) + (2 * ExtrusionAmount));
+
+	const uint32 NewTextureWidth = bPadToPowerOf2 ? FMath::RoundUpToPowerOfTwo(NewMinTextureWidth) : NewMinTextureWidth;
+	const uint32 NewTextureHeight = bPadToPowerOf2 ? FMath::RoundUpToPowerOfTwo(NewMinTextureHeight) : NewMinTextureHeight;
+
+	UTexture2D* Result = NewObject<UTexture2D>(InParent, Name, Flags | RF_Transactional);
+
+	//@TODO: Copy more state across (or start by duplicating the texture maybe?)
+	Result->LODGroup = SourceTexture->LODGroup;
+	Result->CompressionSettings = SourceTexture->CompressionSettings;
+	Result->MipGenSettings = TMGS_NoMipmaps; //@TODO: Don't actually want this...
+	Result->DeferCompression = true;
+
+	TArray<uint8> NewTextureData;
+	NewTextureData.AddZeroed(NewTextureWidth * NewTextureHeight * sizeof(FColor));
+
+	const FIntPoint TileWH(TileWidth, TileHeight);
+	for (int32 TileY = 0; TileY < NumTilesY; ++TileY)
+	{
+		for (int32 TileX = 0; TileX < NumTilesX; ++TileX)
+		{
+			const FIntPoint TileUV = SourceTileSet->GetTileUVFromTileXY(FIntPoint(TileX, TileY));
+
+			TArray<uint8> DummyBuffer;
+			FPaperAtlasTextureHelpers::ReadSpriteTexture(SourceTexture, TileUV, TileWH, DummyBuffer);
+
+			FPaperSpriteAtlasSlot Slot;
+			Slot.X = TileX * (TileWidth + (2 * ExtrusionAmount));
+			Slot.Y = TileY * (TileHeight + (2 * ExtrusionAmount));
+			Slot.Width = TileWidth;// +(2 * ExtrusionAmount);
+			Slot.Height = TileHeight;// +(2 * ExtrusionAmount);
+			Slot.AtlasIndex = 0;
+
+			FPaperAtlasTextureHelpers::CopyTextureRegionToAtlasTextureData(NewTextureData, NewTextureWidth, NewTextureHeight, sizeof(FColor), EPaperSpriteAtlasPadding::DilateBorder, ExtrusionAmount, DummyBuffer, TileWH, Slot);
+		}
+	}
+
+	Result->Source.Init(NewTextureWidth, NewTextureHeight, 1, 1, TSF_BGRA8, NewTextureData.GetData());
+
+	Result->UpdateResource();
+	Result->PostEditChange();
+
+	// Apply the new tile sheet to the specified tile set
+	SourceTileSet->Modify();
+	SourceTileSet->TileSheet = Result;
+	SourceTileSet->Margin = ExtrusionAmount;
+	SourceTileSet->Spacing = 2*ExtrusionAmount;
+	SourceTileSet->PostEditChange();
+
+	return Result;
+}
+
+#undef LOCTEXT_NAMESPACE
\ No newline at end of file
diff --git a/Engine/Plugins/2D/Paper2D/Source/Paper2DEditor/Private/TileSetEditor/TileSheetPaddingFactory.h b/Engine/Plugins/2D/Paper2D/Source/Paper2DEditor/Private/TileSetEditor/TileSheetPaddingFactory.h
new file mode 100644
--- /dev/null
+++ b/Engine/Plugins/2D/Paper2D/Source/Paper2DEditor/Private/TileSetEditor/TileSheetPaddingFactory.h
@@ -0,0 +1,34 @@
+// Copyright 1998-2015 Epic Games, Inc. All Rights Reserved.
+
+#pragma once
+
+class UPaperTileSet;
+
+#include "TileSheetPaddingFactory.generated.h"
+
+/**
+ * Factory used to pad out each individual tile in a tile sheet texture
+ */
+
+UCLASS()
+class UTileSheetPaddingFactory : public UFactory
+{
+	GENERATED_UCLASS_BODY()
+
+	// Source tile sheet texture
+	UPROPERTY()
+	UPaperTileSet* SourceTileSet;
+
+	// The amount to extrude out from each tile (in pixels)
+	UPROPERTY(Category = Settings, EditAnywhere, meta=(UIMin=1, ClampMin=1))
+	int32 ExtrusionAmount;
+
+	// Should we pad the texture to the next power of 2?
+	UPROPERTY(Category=Settings, EditAnywhere)
+	bool bPadToPowerOf2;
+
+
+	// UFactory interface
+	virtual UObject* FactoryCreateNew(UClass* Class, UObject* InParent, FName Name, EObjectFlags Flags, UObject* Context, FFeedbackContext* Warn) override;
+	// End of UFactory interface
+};
