Tile Maps: Allowed non-uniform border margins and per-tile spacing in Tile Sets [UE-14757]
- Changed TileWidth/TileHeight to TileSize
- Made members of UPaperTileSet private and added accessors

diff --git a/Engine/Plugins/2D/Paper2D/Source/Paper2D/Classes/IntMargin.h b/Engine/Plugins/2D/Paper2D/Source/Paper2D/Classes/IntMargin.h
new file mode 100644
--- /dev/null
+++ b/Engine/Plugins/2D/Paper2D/Source/Paper2D/Classes/IntMargin.h
@@ -0,0 +1,126 @@
+// Copyright 1998-2015 Epic Games, Inc. All Rights Reserved.
+
+#pragma once
+
+#include "IntMargin.generated.h"
+
+/**
+ * Describes the space around a 2D area on an integer grid.
+ */
+USTRUCT(BlueprintType)
+struct FIntMargin
+{
+	GENERATED_USTRUCT_BODY()
+
+	/** Holds the margin to the left. */
+	UPROPERTY(EditAnywhere, BlueprintReadWrite, Category=Appearance)
+	int32 Left;
+
+	/** Holds the margin to the top. */
+	UPROPERTY(EditAnywhere, BlueprintReadWrite, Category=Appearance)
+	int32 Top;
+	
+	/** Holds the margin to the right. */
+	UPROPERTY(EditAnywhere, BlueprintReadWrite, Category=Appearance)
+	int32 Right;
+	
+	/** Holds the margin to the bottom. */
+	UPROPERTY(EditAnywhere, BlueprintReadWrite, Category=Appearance)
+	int32 Bottom;
+
+public:
+
+	/**
+	 * Default constructor.
+	 *
+	 * The default margin size is zero on all four sides..
+	 */
+	FIntMargin( )
+		: Left(0)
+		, Top(0)
+		, Right(0)
+		, Bottom(0)
+	{ }
+
+	/** Construct a Margin with uniform space on all sides */
+	FIntMargin(int32 UniformMargin)
+		: Left(UniformMargin)
+		, Top(UniformMargin)
+		, Right(UniformMargin)
+		, Bottom(UniformMargin)
+	{ }
+	
+	/** Construct a Margin where Horizontal describes Left and Right spacing while Vertical describes Top and Bottom spacing */
+	FIntMargin(int32 Horizontal, int32 Vertical)
+		: Left(Horizontal)
+		, Top(Vertical)
+		, Right(Horizontal)
+		, Bottom(Vertical)
+	{ }
+	
+	/** Construct a Margin where the spacing on each side is individually specified. */
+	FIntMargin(int32 InLeft, int32 InTop, int32 InRight, int32 InBottom)
+		: Left(InLeft)
+		, Top(InTop)
+		, Right(InRight)
+		, Bottom(InBottom)
+	{ }
+	
+public:
+
+	/**
+	 * Adds another margin to this margin.
+	 *
+	 * @param Other The margin to add.
+	 * @return A margin that represents this margin plus the other margin.
+	 */
+	FIntMargin operator+(const FIntMargin& InDelta) const
+	{
+		return FIntMargin(Left + InDelta.Left, Top + InDelta.Top, Right + InDelta.Right, Bottom + InDelta.Bottom);
+	}
+
+	/**
+	 * Subtracts another margin from this margin.
+	 *
+	 * @param Other The margin to subtract.
+	 * @return A margin that represents this margin minus the other margin.
+	 */
+	FIntMargin operator-(const FIntMargin& Other) const
+	{
+		return FIntMargin(Left - Other.Left, Top - Other.Top, Right - Other.Right, Bottom - Other.Bottom);
+	}
+
+	/**
+	 * Compares this margin with another for equality.
+	 *
+	 * @param Other The other margin.
+	 * @return true if the two margins are equal, false otherwise.
+	 */
+	bool operator==(const FIntMargin& Other) const
+	{
+		return (Left == Other.Left) && (Right == Other.Right) && (Top == Other.Top) && (Bottom == Other.Bottom);
+	}
+
+	/**
+	 * Compares this margin with another for inequality.
+	 *
+	 * @param Other The other margin.
+	 * @return true if the two margins are not equal, false otherwise.
+	 */
+	bool operator!=(const FIntMargin& Other) const
+	{
+		return (Left != Other.Left) || (Right != Other.Right) || (Top != Other.Top) || (Bottom != Other.Bottom);
+	}
+
+public:
+
+	/**
+	 * Gets the margin's total size.
+	 *
+	 * @return Cumulative margin size.
+	 */
+	FIntPoint GetDesiredSize() const
+	{
+		return FIntPoint(Left + Right, Top + Bottom);
+	}
+};
diff --git a/Engine/Plugins/2D/Paper2D/Source/Paper2D/Classes/PaperTileSet.h b/Engine/Plugins/2D/Paper2D/Source/Paper2D/Classes/PaperTileSet.h
--- a/Engine/Plugins/2D/Paper2D/Source/Paper2D/Classes/PaperTileSet.h
+++ b/Engine/Plugins/2D/Paper2D/Source/Paper2D/Classes/PaperTileSet.h
@@ -5,6 +5,7 @@
 #include "SpriteEditorOnlyTypes.h"
 #include "Engine/DataAsset.h"
 #include "Engine/EngineTypes.h"
+#include "IntMargin.h"
 
 #include "PaperTileSet.generated.h"
 
@@ -77,38 +78,33 @@ class PAPER2D_API UPaperTileSet : public UObject
 {
 	GENERATED_UCLASS_BODY()
 
-	// The width of a single tile (in pixels)
-	UPROPERTY(Category=TileSet, BlueprintReadOnly, EditAnywhere, meta=(UIMin=1, ClampMin=1))
-	int32 TileWidth;
-
-	// The height of a single tile (in pixels)
-	UPROPERTY(Category=TileSet, BlueprintReadOnly, EditAnywhere, meta=(UIMin=1, ClampMin=1))
-	int32 TileHeight;
+private:
+	// The width and height of a single tile (in pixels)
+	UPROPERTY(Category=TileSet, BlueprintReadOnly, EditAnywhere, meta=(UIMin=1, ClampMin=1, AllowPrivateAccess="true"))
+	FIntPoint TileSize;
 
 	// The tile sheet texture associated with this tile set
-	UPROPERTY(Category=TileSet, BlueprintReadOnly, EditAnywhere, meta=(DisplayName="Tile Sheet Texture"))
+	UPROPERTY(Category=TileSet, BlueprintReadOnly, EditAnywhere, meta=(DisplayName="Tile Sheet Texture", AllowPrivateAccess="true"))
 	UTexture2D* TileSheet;
 
-	// The amount of padding around the perimeter of the tile sheet (in pixels)
-	UPROPERTY(Category=TileSet, BlueprintReadOnly, EditAnywhere, meta=(UIMin=0, ClampMin=0))
-	int32 Margin;
+	// The amount of padding around the border of the tile sheet (in pixels)
+	UPROPERTY(Category=TileSet, BlueprintReadOnly, EditAnywhere, meta=(UIMin=0, ClampMin=0, AllowPrivateAccess="true"))
+	FIntMargin BorderMargin;
 
 	// The amount of padding between tiles in the tile sheet (in pixels)
-	UPROPERTY(Category=TileSet, BlueprintReadOnly, EditAnywhere, meta=(UIMin=0, ClampMin=0))
-	int32 Spacing;
+	UPROPERTY(Category=TileSet, BlueprintReadOnly, EditAnywhere, meta=(UIMin=0, ClampMin=0, DisplayName="Per-Tile Spacing", AllowPrivateAccess="true"))
+	FIntPoint PerTileSpacing;
 
 	// The drawing offset for tiles from this set (in pixels)
-	UPROPERTY(Category=TileSet, BlueprintReadOnly, EditAnywhere)
+	UPROPERTY(Category=TileSet, BlueprintReadOnly, EditAnywhere, meta=(AllowPrivateAccess="true"))
 	FIntPoint DrawingOffset;
 
 #if WITH_EDITORONLY_DATA
 	/** The background color displayed in the tile set viewer */
-	UPROPERTY(Category=TileSet, EditAnywhere, meta=(HideAlphaChannel))
+	UPROPERTY(Category=TileSet, EditAnywhere, meta=(HideAlphaChannel), meta=(AllowPrivateAccess="true"))
 	FLinearColor BackgroundColor;
 #endif
 
-
-protected:
 	// Cached width of this tile set (in tiles)
 	UPROPERTY()
 	int32 WidthInTiles;
@@ -145,8 +141,9 @@ protected:
 public:
 
 	// UObject interface
-#if WITH_EDITOR
+	virtual void Serialize(FArchive& Ar) override;
 	virtual void PostLoad() override;
+#if WITH_EDITOR
 	virtual void PostEditChangeProperty(FPropertyChangedEvent& PropertyChangedEvent) override;
 #endif
 	// End of UObject interface
@@ -193,4 +190,102 @@ public:
 #if WITH_EDITOR
 	static FName GetPerTilePropertyName() { return GET_MEMBER_NAME_CHECKED(UPaperTileSet, PerTileData); }
 #endif
+
+	// Sets the size of a tile (in pixels)
+	// Note: Does not trigger a rebuild of any tile map components using this tile set
+	inline void SetTileSize(FIntPoint NewSize)
+	{
+		TileSize = FIntPoint(FMath::Max(NewSize.X, 1), FMath::Max(NewSize.Y, 1));
+	}
+
+	// Returns the size of a tile (in pixels)
+	inline FIntPoint GetTileSize() const
+	{
+		return TileSize;
+	}
+
+	// Sets the tile sheet texture associated with this tile set
+	// Note: Does not trigger a rebuild of any tile map components using this tile set
+	inline void SetTileSheetTexture(UTexture2D* NewTileSheet)
+	{
+		TileSheet = NewTileSheet;
+	}
+
+	// Returns the tile sheet texture associated with this tile set
+	inline UTexture2D* GetTileSheetTexture() const
+	{
+		return TileSheet;
+	}
+
+	// Returns the imported size of the tile sheet texture (in pixels)
+	inline FIntPoint GetTileSheetAuthoredSize() const
+	{
+		return TileSheet->GetImportedSize();
+	}
+
+	// Returns the amount of padding around the border of the tile sheet (in pixels)
+	inline FIntMargin GetMargin() const
+	{
+		return BorderMargin;
+	}
+
+	// Sets the amount of padding around the border of the tile sheet (in pixels)
+	// Note: Does not trigger a rebuild of any tile map components using this tile set
+	inline void SetMargin(FIntMargin NewMargin)
+	{
+		BorderMargin = NewMargin;
+	}
+
+	// Returns the amount of padding between tiles in the tile sheet (in pixels)
+	inline FIntPoint GetPerTileSpacing() const
+	{
+		return PerTileSpacing;
+	}
+
+	// Sets the amount of padding between tiles in the tile sheet (in pixels)
+	// Note: Does not trigger a rebuild of any tile map components using this tile set
+	inline void SetPerTileSpacing(FIntPoint NewSpacing)
+	{
+		PerTileSpacing = NewSpacing;
+	}
+
+	// Returns the drawing offset for tiles from this set (in pixels)
+	inline FIntPoint GetDrawingOffset() const
+	{
+		return DrawingOffset;
+	}
+
+	// Sets the drawing offset for tiles from this set (in pixels)
+	// Note: Does not trigger a rebuild of any tile map components using this tile set
+	inline void SetDrawingOffset(FIntPoint NewDrawingOffset)
+	{
+		DrawingOffset = NewDrawingOffset;
+	}
+
+#if WITH_EDITORONLY_DATA
+	// Returns the background color displayed in the tile set viewer
+	inline FLinearColor GetBackgroundColor() const
+	{
+		return BackgroundColor;
+	}
+
+	// Returns the background color displayed in the tile set viewer
+	inline void SetBackgroundColor(FLinearColor NewColor)
+	{
+		BackgroundColor = NewColor;
+	}
+#endif
+
+private:
+	UPROPERTY()
+	int32 TileWidth_DEPRECATED;
+
+	UPROPERTY()
+	int32 TileHeight_DEPRECATED;
+
+	UPROPERTY()
+	int32 Margin_DEPRECATED;
+
+	UPROPERTY()
+	int32 Spacing_DEPRECATED;
 };
diff --git a/Engine/Plugins/2D/Paper2D/Source/Paper2D/Private/PaperCustomVersion.h b/Engine/Plugins/2D/Paper2D/Source/Paper2D/Private/PaperCustomVersion.h
--- a/Engine/Plugins/2D/Paper2D/Source/Paper2D/Private/PaperCustomVersion.h
+++ b/Engine/Plugins/2D/Paper2D/Source/Paper2D/Private/PaperCustomVersion.h
@@ -52,6 +52,9 @@ struct FPaperCustomVersion
 		// Refactor sprite render/collision polygon storage to allow more flexible geometry shapes
 		RefactorPolygonStorageToSupportShapes = 14,
 
+		// Change Tile Set margin and padding to allow non-uniform values (also converts TileWidth/TileHeight into TileSize and makes members private)
+		AllowNonUniformPaddingInTileSets = 15,
+
 		// -----<new versions can be added above this line>-------------------------------------------------
 		VersionPlusOne,
 		LatestVersion = VersionPlusOne - 1
diff --git a/Engine/Plugins/2D/Paper2D/Source/Paper2D/Private/PaperTileMap.cpp b/Engine/Plugins/2D/Paper2D/Source/Paper2D/Private/PaperTileMap.cpp
--- a/Engine/Plugins/2D/Paper2D/Source/Paper2D/Private/PaperTileMap.cpp
+++ b/Engine/Plugins/2D/Paper2D/Source/Paper2D/Private/PaperTileMap.cpp
@@ -570,8 +570,9 @@ void UPaperTileMap::InitializeNewEmptyTileMap(UPaperTileSet* InitialTileSet)
 {
 	if (InitialTileSet != nullptr)
 	{
-		TileWidth = InitialTileSet->TileWidth;
-		TileHeight = InitialTileSet->TileHeight;
+		const FIntPoint TileSetTileSize = InitialTileSet->GetTileSize();
+		TileWidth = TileSetTileSize.X;
+		TileHeight = TileSetTileSize.Y;
 		SelectedTileSet = InitialTileSet;
 	}
 
diff --git a/Engine/Plugins/2D/Paper2D/Source/Paper2D/Private/PaperTileMapComponent.cpp b/Engine/Plugins/2D/Paper2D/Source/Paper2D/Private/PaperTileMapComponent.cpp
--- a/Engine/Plugins/2D/Paper2D/Source/Paper2D/Private/PaperTileMapComponent.cpp
+++ b/Engine/Plugins/2D/Paper2D/Source/Paper2D/Private/PaperTileMapComponent.cpp
@@ -143,9 +143,12 @@ void UPaperTileMapComponent::GetUsedTextures(TArray<UTexture*>& OutTextures, EMa
 				for (int32 X = 0; X < TileMap->MapWidth; ++X)
 				{
 					FPaperTileInfo TileInfo = Layer->GetCell(X, Y);
-					if (TileInfo.IsValid() && (TileInfo.TileSet != nullptr) && (TileInfo.TileSet->TileSheet != nullptr))
+					if (TileInfo.IsValid() && (TileInfo.TileSet != nullptr))
 					{
-						OutTextures.AddUnique(TileInfo.TileSet->TileSheet);
+						if (UTexture2D* TileSheet = TileInfo.TileSet->GetTileSheetTexture())
+						{
+							OutTextures.AddUnique(TileSheet);
+						}
 					}
 				}
 			}
@@ -294,7 +297,7 @@ void UPaperTileMapComponent::RebuildRenderData(FPaperTileMapRenderSceneProxy* Pr
 						continue;
 					}
 
-					SourceTexture = TileInfo.TileSet->TileSheet;
+					SourceTexture = TileInfo.TileSet->GetTileSheetTexture();
 					if (SourceTexture == nullptr)
 					{
 						continue;
@@ -317,11 +320,14 @@ void UPaperTileMapComponent::RebuildRenderData(FPaperTileMapRenderSceneProxy* Pr
 
 						if (TileInfo.TileSet != nullptr)
 						{
-							SourceDimensionsUV = FVector2D(TileInfo.TileSet->TileWidth * InverseTextureSize.X, TileInfo.TileSet->TileHeight * InverseTextureSize.Y);
-							TileSizeXY = FVector2D(UnrealUnitsPerPixel * TileInfo.TileSet->TileWidth, UnrealUnitsPerPixel * TileInfo.TileSet->TileHeight);
+							const FIntPoint TileSetTileSize(TileInfo.TileSet->GetTileSize());
+
+							SourceDimensionsUV = FVector2D(TileSetTileSize.X * InverseTextureSize.X, TileSetTileSize.Y * InverseTextureSize.Y);
+							TileSizeXY = FVector2D(UnrealUnitsPerPixel * TileSetTileSize.X, UnrealUnitsPerPixel * TileSetTileSize.Y);
 
-							const float HorizontalCellOffset = TileInfo.TileSet->DrawingOffset.X * UnrealUnitsPerPixel;
-							const float VerticalCellOffset = (-TileInfo.TileSet->DrawingOffset.Y - TileHeight + TileInfo.TileSet->TileHeight) * UnrealUnitsPerPixel;
+							const FIntPoint TileSetDrawingOffset = TileInfo.TileSet->GetDrawingOffset();
+							const float HorizontalCellOffset = TileSetDrawingOffset.X * UnrealUnitsPerPixel;
+							const float VerticalCellOffset = (-TileSetDrawingOffset.Y - TileHeight + TileSetTileSize.Y) * UnrealUnitsPerPixel;
 							TileSetOffset = (HorizontalCellOffset * PaperAxisX) + (VerticalCellOffset * PaperAxisY);
 						}
 						else
diff --git a/Engine/Plugins/2D/Paper2D/Source/Paper2D/Private/PaperTileSet.cpp b/Engine/Plugins/2D/Paper2D/Source/Paper2D/Private/PaperTileSet.cpp
--- a/Engine/Plugins/2D/Paper2D/Source/Paper2D/Private/PaperTileSet.cpp
+++ b/Engine/Plugins/2D/Paper2D/Source/Paper2D/Private/PaperTileSet.cpp
@@ -2,15 +2,17 @@
 
 #include "Paper2DPrivatePCH.h"
 #include "PaperTileSet.h"
+#include "PaperCustomVersion.h"
 
 //////////////////////////////////////////////////////////////////////////
 // UPaperTileSet
 
 UPaperTileSet::UPaperTileSet(const FObjectInitializer& ObjectInitializer)
 	: Super(ObjectInitializer)
+	, TileSize(32, 32)
 {
-	TileWidth = 32;
-	TileHeight = 32;
+	TileWidth_DEPRECATED = 32;
+	TileHeight_DEPRECATED = 32;
 
 #if WITH_EDITORONLY_DATA
 	BackgroundColor = FColor(0, 0, 127);
@@ -21,13 +23,15 @@ int32 UPaperTileSet::GetTileCount() const
 {
 	if (TileSheet != nullptr)
 	{
-		checkSlow((TileWidth > 0) && (TileHeight > 0));
+		checkSlow((TileSize.X > 0) && (TileSize.Y > 0));
 		const FIntPoint TextureSize = TileSheet->GetImportedSize();
 		const int32 TextureWidth = TextureSize.X;
 		const int32 TextureHeight = TextureSize.Y;
 
-		const int32 CellsX = (TextureWidth - (Margin * 2) + Spacing) / (TileWidth + Spacing);
-		const int32 CellsY = (TextureHeight - (Margin * 2) + Spacing) / (TileHeight + Spacing);
+		const FIntPoint SumMargin = BorderMargin.GetDesiredSize();
+
+		const int32 CellsX = (TextureWidth - SumMargin.X + PerTileSpacing.X) / (TileSize.X + PerTileSpacing.X);
+		const int32 CellsY = (TextureHeight - SumMargin.Y + PerTileSpacing.Y) / (TileSize.Y + PerTileSpacing.Y);
 
 		return CellsX * CellsY;
 	}
@@ -41,9 +45,9 @@ int32 UPaperTileSet::GetTileCountX() const
 {
 	if (TileSheet != nullptr)
 	{
-		checkSlow(TileWidth > 0);
+		checkSlow(TileSize.X > 0);
 		const int32 TextureWidth = TileSheet->GetImportedSize().X;
-		const int32 CellsX = (TextureWidth - (Margin * 2) + Spacing) / (TileWidth + Spacing);
+		const int32 CellsX = (TextureWidth - (BorderMargin.Left + BorderMargin.Right) + PerTileSpacing.X) / (TileSize.X + PerTileSpacing.X);
 		return CellsX;
 	}
 	else
@@ -56,9 +60,9 @@ int32 UPaperTileSet::GetTileCountY() const
 {
 	if (TileSheet != nullptr)
 	{
-		checkSlow(TileHeight > 0);
+		checkSlow(TileSize.Y > 0);
 		const int32 TextureHeight = TileSheet->GetImportedSize().Y;
-		const int32 CellsY = (TextureHeight - (Margin * 2) + Spacing) / (TileHeight + Spacing);
+		const int32 CellsY = (TextureHeight - (BorderMargin.Top + BorderMargin.Bottom) + PerTileSpacing.Y) / (TileSize.Y + PerTileSpacing.Y);
 		return CellsY;
 	}
 	else
@@ -112,15 +116,15 @@ bool UPaperTileSet::GetTileUV(int32 TileIndex, /*out*/ FVector2D& Out_TileUV) co
 
 FIntPoint UPaperTileSet::GetTileUVFromTileXY(const FIntPoint& TileXY) const
 {
-	return FIntPoint(TileXY.X * (TileWidth + Spacing) + Margin, TileXY.Y * (TileHeight + Spacing) + Margin);
+	return FIntPoint(TileXY.X * (TileSize.X + PerTileSpacing.X) + BorderMargin.Left, TileXY.Y * (TileSize.Y + PerTileSpacing.Y) + BorderMargin.Top);
 }
 
 FIntPoint UPaperTileSet::GetTileXYFromTextureUV(const FVector2D& TextureUV, bool bRoundUp) const
 {
-	const float DividendX = TextureUV.X - Margin;
-	const float DividendY = TextureUV.Y - Margin;
-	const float DivisorX = TileWidth + Spacing;
-	const float DivisorY = TileHeight + Spacing;
+	const float DividendX = TextureUV.X - BorderMargin.Left;
+	const float DividendY = TextureUV.Y - BorderMargin.Top;
+	const float DivisorX = TileSize.X + PerTileSpacing.X;
+	const float DivisorY = TileSize.Y + PerTileSpacing.Y;
 	const int32 X = bRoundUp ? FMath::DivideAndRoundUp<int32>(DividendX, DivisorX) : FMath::DivideAndRoundDown<int32>(DividendX, DivisorX);
 	const int32 Y = bRoundUp ? FMath::DivideAndRoundUp<int32>(DividendY, DivisorY) : FMath::DivideAndRoundDown<int32>(DividendY, DivisorY);
 	return FIntPoint(X, Y);
@@ -144,13 +148,24 @@ int32 UPaperTileSet::GetTerrainMembership(const FPaperTileInfo& TileInfo) const
 	return INDEX_NONE; //@TODO: TileMapTerrain: Implement this
 }
 
-#if WITH_EDITOR
+void UPaperTileSet::Serialize(FArchive& Ar)
+{
+	Super::Serialize(Ar);
 
-#include "PaperTileMapComponent.h"
-#include "ComponentReregisterContext.h"
+	Ar.UsingCustomVersion(FPaperCustomVersion::GUID);
+}
 
 void UPaperTileSet::PostLoad()
 {
+	const int32 PaperVer = GetLinkerCustomVersion(FPaperCustomVersion::GUID);
+
+	if (PaperVer < FPaperCustomVersion::AllowNonUniformPaddingInTileSets)
+	{
+		BorderMargin = FIntMargin(Margin_DEPRECATED);
+		PerTileSpacing = FIntPoint(Spacing_DEPRECATED, Spacing_DEPRECATED);
+		TileSize = FIntPoint(TileWidth_DEPRECATED, TileHeight_DEPRECATED);
+	}
+
 	if (TileSheet != nullptr)
 	{
 		TileSheet->ConditionalPostLoad();
@@ -163,15 +178,25 @@ void UPaperTileSet::PostLoad()
 	Super::PostLoad();
 }
 
+
+#if WITH_EDITOR
+
+#include "PaperTileMapComponent.h"
+#include "ComponentReregisterContext.h"
+
 void UPaperTileSet::PostEditChangeProperty(FPropertyChangedEvent& PropertyChangedEvent)
 {
 	const FName PropertyName = (PropertyChangedEvent.Property != nullptr) ? PropertyChangedEvent.Property->GetFName() : NAME_None;
 
-	Margin = FMath::Max<int32>(Margin, 0);
-	Spacing = FMath::Max<int32>(Spacing, 0);
+	BorderMargin.Left = FMath::Max<int32>(BorderMargin.Left, 0);
+	BorderMargin.Right = FMath::Max<int32>(BorderMargin.Right, 0);
+	BorderMargin.Top = FMath::Max<int32>(BorderMargin.Top, 0);
+	BorderMargin.Bottom = FMath::Max<int32>(BorderMargin.Bottom, 0);
+	PerTileSpacing.X = FMath::Max<int32>(PerTileSpacing.X, 0);
+	PerTileSpacing.Y = FMath::Max<int32>(PerTileSpacing.Y, 0);
 
-	TileWidth = FMath::Max<int32>(TileWidth, 1);
-	TileHeight = FMath::Max<int32>(TileHeight, 1);
+	TileSize.X = FMath::Max<int32>(TileSize.X, 1);
+	TileSize.Y = FMath::Max<int32>(TileSize.Y, 1);
 
 	WidthInTiles = GetTileCountX();
 	HeightInTiles = GetTileCountY();
diff --git a/Engine/Plugins/2D/Paper2D/Source/Paper2DEditor/Private/PaperImporterSettings.cpp b/Engine/Plugins/2D/Paper2D/Source/Paper2DEditor/Private/PaperImporterSettings.cpp
--- a/Engine/Plugins/2D/Paper2D/Source/Paper2DEditor/Private/PaperImporterSettings.cpp
+++ b/Engine/Plugins/2D/Paper2D/Source/Paper2DEditor/Private/PaperImporterSettings.cpp
@@ -136,8 +136,9 @@ void UPaperImporterSettings::ApplySettingsForTileMapInit(UPaperTileMap* TileMap,
 {
 	if (DefaultTileSet != nullptr)
 	{
-		TileMap->TileWidth = DefaultTileSet->TileWidth;
-		TileMap->TileHeight = DefaultTileSet->TileHeight;
+		const FIntPoint TileSetTileSize(DefaultTileSet->GetTileSize());
+		TileMap->TileWidth = TileSetTileSize.X;
+		TileMap->TileHeight = TileSetTileSize.Y;
 		TileMap->SelectedTileSet = DefaultTileSet;
 	}
 
@@ -148,9 +149,12 @@ void UPaperImporterSettings::ApplySettingsForTileMapInit(UPaperTileMap* TileMap,
 	if (DesiredMaterialType == ESpriteInitMaterialType::Automatic)
 	{
 		// Analyze the texture if desired (to see if it's got greyscale alpha or just binary alpha, picking either a translucent or masked material)
-		if (bPickBestMaterialWhenCreatingTileMaps && (DefaultTileSet != nullptr) && (DefaultTileSet->TileSheet != nullptr))
+		if (bPickBestMaterialWhenCreatingTileMaps && (DefaultTileSet != nullptr))
 		{
-			DesiredMaterialType = AnalyzeTextureForDesiredMaterialType(DefaultTileSet->TileSheet, FIntPoint::ZeroValue, DefaultTileSet->TileSheet->GetImportedSize());
+			if (UTexture2D* TileSheetTexture = DefaultTileSet->GetTileSheetTexture())
+			{
+				DesiredMaterialType = AnalyzeTextureForDesiredMaterialType(TileSheetTexture, FIntPoint::ZeroValue, TileSheetTexture->GetImportedSize());
+			}
 		}
 	}
 
diff --git a/Engine/Plugins/2D/Paper2D/Source/Paper2DEditor/Private/PaperTileSetFactory.cpp b/Engine/Plugins/2D/Paper2D/Source/Paper2DEditor/Private/PaperTileSetFactory.cpp
--- a/Engine/Plugins/2D/Paper2D/Source/Paper2DEditor/Private/PaperTileSetFactory.cpp
+++ b/Engine/Plugins/2D/Paper2D/Source/Paper2DEditor/Private/PaperTileSetFactory.cpp
@@ -20,8 +20,8 @@ UPaperTileSetFactory::UPaperTileSetFactory(const FObjectInitializer& ObjectIniti
 UObject* UPaperTileSetFactory::FactoryCreateNew(UClass* Class, UObject* InParent, FName Name, EObjectFlags Flags, UObject* Context, FFeedbackContext* Warn)
 {
 	UPaperTileSet* NewTileSet = NewObject<UPaperTileSet>(InParent, Class, Name, Flags | RF_Transactional);
-	NewTileSet->TileSheet = InitialTexture;
-	NewTileSet->BackgroundColor = GetDefault<UTileSetEditorSettings>()->DefaultBackgroundColor;
+	NewTileSet->SetTileSheetTexture(InitialTexture);
+	NewTileSet->SetBackgroundColor(GetDefault<UTileSetEditorSettings>()->DefaultBackgroundColor);
 
 	return NewTileSet;
 }
diff --git a/Engine/Plugins/2D/Paper2D/Source/Paper2DEditor/Private/PaperTileSetThumbnailRenderer.cpp b/Engine/Plugins/2D/Paper2D/Source/Paper2DEditor/Private/PaperTileSetThumbnailRenderer.cpp
--- a/Engine/Plugins/2D/Paper2D/Source/Paper2DEditor/Private/PaperTileSetThumbnailRenderer.cpp
+++ b/Engine/Plugins/2D/Paper2D/Source/Paper2DEditor/Private/PaperTileSetThumbnailRenderer.cpp
@@ -15,9 +15,10 @@ UPaperTileSetThumbnailRenderer::UPaperTileSetThumbnailRenderer(const FObjectInit
 void UPaperTileSetThumbnailRenderer::Draw(UObject* Object, int32 X, int32 Y, uint32 Width, uint32 Height, FRenderTarget*, FCanvas* Canvas)
 {
 	UPaperTileSet* TileSet = Cast<UPaperTileSet>(Object);
-	if ((TileSet != nullptr) && (TileSet->TileSheet != nullptr))
+	if ((TileSet != nullptr) && (TileSet->GetTileSheetTexture() != nullptr))
 	{
-		const bool bUseTranslucentBlend = TileSet->TileSheet->HasAlphaChannel();
+		UTexture2D* TileSheetTexture = TileSet->GetTileSheetTexture();
+		const bool bUseTranslucentBlend = TileSheetTexture->HasAlphaChannel();
 
 		// Draw the grid behind the sprite
 		if (bUseTranslucentBlend)
@@ -45,8 +46,8 @@ void UPaperTileSetThumbnailRenderer::Draw(UObject* Object, int32 X, int32 Y, uin
 		}
 
 		// Draw the sprite itself
-		const float TextureWidth = TileSet->TileSheet->GetSurfaceWidth();
-		const float TextureHeight = TileSet->TileSheet->GetSurfaceHeight();
+		const float TextureWidth = TileSheetTexture->GetSurfaceWidth();
+		const float TextureHeight = TileSheetTexture->GetSurfaceHeight();
 
 		float FinalX = (float)X;
 		float FinalY = (float)Y;
@@ -80,6 +81,7 @@ void UPaperTileSetThumbnailRenderer::Draw(UObject* Object, int32 X, int32 Y, uin
 			Canvas->DrawTile(FinalX+FinalWidth, Y, Width-FinalWidth, Height, 0, 0, 1, 1, BlackBarColor, GWhiteTexture, bAlphaBlend);
 		}
 
+		//@TODO: Cut off the border margin when rendering
 
 		// Draw the tile sheet 
 		Canvas->DrawTile(
@@ -92,7 +94,7 @@ void UPaperTileSetThumbnailRenderer::Draw(UObject* Object, int32 X, int32 Y, uin
 			1.0f,
 			1.0f,
 			FLinearColor::White,
-			TileSet->TileSheet->Resource,
+			TileSheetTexture->Resource,
 			bUseTranslucentBlend);
 
 		// Draw a label overlay
diff --git a/Engine/Plugins/2D/Paper2D/Source/Paper2DEditor/Private/TileSetAssetTypeActions.cpp b/Engine/Plugins/2D/Paper2D/Source/Paper2DEditor/Private/TileSetAssetTypeActions.cpp
--- a/Engine/Plugins/2D/Paper2D/Source/Paper2DEditor/Private/TileSetAssetTypeActions.cpp
+++ b/Engine/Plugins/2D/Paper2D/Source/Paper2DEditor/Private/TileSetAssetTypeActions.cpp
@@ -112,13 +112,13 @@ void FTileSetAssetTypeActions::ExecutePadTileSetTexture(TWeakObjectPtr<UPaperTil
 
 	if (UPaperTileSet* TileSet = TileSetPtr.Get())
 	{
-		if (TileSet->TileSheet != nullptr)
+		if (UTexture2D* TileSheetTexture = TileSet->GetTileSheetTexture())
 		{
 			const FString TileSheetSuffix(TEXT("Padded"));
-			const FString TileSheetPathName = TileSet->TileSheet->GetOutermost()->GetPathName();
+			const FString TileSheetPathName = TileSheetTexture->GetOutermost()->GetPathName();
 			const FString LongPackagePath = FPackageName::GetLongPackagePath(TileSheetPathName);
 
-			const FString EffectiveTileSheetName = TileSet->TileSheet->GetName();
+			const FString EffectiveTileSheetName = TileSheetTexture->GetName();
 			const FString NewTileSheetDefaultPath = LongPackagePath / EffectiveTileSheetName;
 
 			// Make sure the name is unique
diff --git a/Engine/Plugins/2D/Paper2D/Source/Paper2DEditor/Private/TileSetEditor/SingleTileEditorViewportClient.cpp b/Engine/Plugins/2D/Paper2D/Source/Paper2DEditor/Private/TileSetEditor/SingleTileEditorViewportClient.cpp
--- a/Engine/Plugins/2D/Paper2D/Source/Paper2DEditor/Private/TileSetEditor/SingleTileEditorViewportClient.cpp
+++ b/Engine/Plugins/2D/Paper2D/Source/Paper2DEditor/Private/TileSetEditor/SingleTileEditorViewportClient.cpp
@@ -66,7 +66,7 @@ void FSingleTileEditorViewportClient::Tick(float DeltaSeconds)
 
 FLinearColor FSingleTileEditorViewportClient::GetBackgroundColor() const
 {
-	return TileSet->BackgroundColor;
+	return TileSet->GetBackgroundColor();
 }
 
 void FSingleTileEditorViewportClient::TrackingStarted(const struct FInputEventState& InInputState, bool bIsDragging, bool bNudge)
@@ -236,20 +236,20 @@ void FSingleTileEditorViewportClient::SetTileIndex(int32 InTileIndex)
 
 		FSpriteAssetInitParameters SpriteReinitParams;
 
-		SpriteReinitParams.Texture = TileSet->TileSheet;
+		SpriteReinitParams.Texture = TileSet->GetTileSheetTexture();
 
 		//@TODO: Should analyze the texture (*at a higher level, not per tile click!*) to pick the correct material
 		FVector2D UV;
 		TileSet->GetTileUV(TileBeingEditedIndex, /*out*/ UV);
 		SpriteReinitParams.Offset = FIntPoint((int32)UV.X, (int32)(UV.Y));
-		SpriteReinitParams.Dimension = FIntPoint(TileSet->TileWidth, TileSet->TileHeight);
+		SpriteReinitParams.Dimension = TileSet->GetTileSize();
 		SpriteReinitParams.SetPixelsPerUnrealUnit(1.0f);
 		DummySprite->InitializeSprite(SpriteReinitParams);
 	}
 	PreviewTileSpriteComponent->SetSprite(DummySprite);
 
 	// Update the default geometry bounds
-	const FVector2D HalfTileSize(TileSet->TileWidth * 0.5f, TileSet->TileHeight * 0.5f);
+	const FVector2D HalfTileSize(TileSet->GetTileSize().X * 0.5f, TileSet->GetTileSize().Y * 0.5f);
 	FBox2D DesiredBounds(ForceInitToZero);
 	DesiredBounds.Min = -HalfTileSize;
 	DesiredBounds.Max = HalfTileSize;
diff --git a/Engine/Plugins/2D/Paper2D/Source/Paper2DEditor/Private/TileSetEditor/TileSetDetailsCustomization.cpp b/Engine/Plugins/2D/Paper2D/Source/Paper2DEditor/Private/TileSetEditor/TileSetDetailsCustomization.cpp
--- a/Engine/Plugins/2D/Paper2D/Source/Paper2DEditor/Private/TileSetEditor/TileSetDetailsCustomization.cpp
+++ b/Engine/Plugins/2D/Paper2D/Source/Paper2DEditor/Private/TileSetEditor/TileSetDetailsCustomization.cpp
@@ -114,17 +114,17 @@ FText FTileSetDetailsCustomization::GetCellDimensionHeaderText() const
 		const int32 NumTilesX = TileSet->GetTileCountX();
 		const int32 NumTilesY = TileSet->GetTileCountY();
 
-		if (TileSet->TileSheet == nullptr)
+		if (TileSet->GetTileSheetTexture() == nullptr)
 		{
-			Result = LOCTEXT("NoTexture", "No TileSheet");
+			Result = LOCTEXT("NoTexture", "No Tile Sheet");
 		}
 		else if (NumTilesX == 0)
 		{
-			Result = LOCTEXT("TextureTooNarrow", "TileSheet too narrow");
+			Result = LOCTEXT("TextureTooNarrow", "Tile Sheet too narrow");
 		}
 		else if (NumTilesY == 0)
 		{
-			Result = LOCTEXT("TextureTooShort", "TileSheet too short");
+			Result = LOCTEXT("TextureTooShort", "Tile Sheet too short");
 		}
 		else
 		{
diff --git a/Engine/Plugins/2D/Paper2D/Source/Paper2DEditor/Private/TileSetEditor/TileSetEditorViewportClient.cpp b/Engine/Plugins/2D/Paper2D/Source/Paper2DEditor/Private/TileSetEditor/TileSetEditorViewportClient.cpp
--- a/Engine/Plugins/2D/Paper2D/Source/Paper2DEditor/Private/TileSetEditor/TileSetEditorViewportClient.cpp
+++ b/Engine/Plugins/2D/Paper2D/Source/Paper2DEditor/Private/TileSetEditor/TileSetEditorViewportClient.cpp
@@ -29,9 +29,7 @@ void FTileSetEditorViewportClient::Draw(FViewport* Viewport, FCanvas* Canvas)
 		return;
 	}
 
-	UTexture2D* Texture = TileSet->TileSheet;
-
-	if (Texture != nullptr)
+	if (UTexture2D* Texture = TileSet->GetTileSheetTexture())
 	{
 		const bool bUseTranslucentBlend = Texture->HasAlphaChannel();
 
@@ -60,8 +58,9 @@ void FTileSetEditorViewportClient::Draw(FViewport* Viewport, FCanvas* Canvas)
 			const FLinearColor MetaDataOverlayColor(1.0f, 0.2f, 0.0f, 0.5f);
 			const FLinearColor InfoOverlayColor = bShowTilesWithCollision ? CollisionOverlayColor : MetaDataOverlayColor;
 
-			const float Width = (TileSet->TileWidth - 2) * ZoomAmount;
-			const float Height = (TileSet->TileHeight - 2) * ZoomAmount;
+			const FIntPoint TileSetTileSize(TileSet->GetTileSize());
+			const float Width = (TileSetTileSize.X - 2) * ZoomAmount;
+			const float Height = (TileSetTileSize.Y - 2) * ZoomAmount;
 
 			for (int32 TileIndex = 0; TileIndex < NumTiles; ++TileIndex)
 			{
@@ -119,7 +118,7 @@ FLinearColor FTileSetEditorViewportClient::GetBackgroundColor() const
 {
 	if (UPaperTileSet* TileSet = TileSetBeingEdited.Get())
 	{
-		return TileSet->BackgroundColor;
+		return TileSet->GetBackgroundColor();
 	}
 	else
 	{
diff --git a/Engine/Plugins/2D/Paper2D/Source/Paper2DEditor/Private/TileSetEditor/TileSetSelectorViewport.cpp b/Engine/Plugins/2D/Paper2D/Source/Paper2DEditor/Private/TileSetEditor/TileSetSelectorViewport.cpp
--- a/Engine/Plugins/2D/Paper2D/Source/Paper2DEditor/Private/TileSetEditor/TileSetSelectorViewport.cpp
+++ b/Engine/Plugins/2D/Paper2D/Source/Paper2DEditor/Private/TileSetEditor/TileSetSelectorViewport.cpp
@@ -100,7 +100,7 @@ void STileSetSelectorViewport::OnFocusViewportToSelection()
 	if (UPaperTileSet* TileSetBeingEdited = TileSetPtr.Get())
 	{
 		const FIntPoint ViewportSize = SceneViewport->GetSize();
-		if (UTexture* Texture = TileSetBeingEdited->TileSheet)
+		if (UTexture* Texture = TileSetBeingEdited->GetTileSheetTexture())
 		{
 			const FIntPoint TextureSize(Texture->Source.GetSizeX(), Texture->Source.GetSizeY());
 
diff --git a/Engine/Plugins/2D/Paper2D/Source/Paper2DEditor/Private/TileSetEditor/TileSheetPaddingFactory.cpp b/Engine/Plugins/2D/Paper2D/Source/Paper2DEditor/Private/TileSetEditor/TileSheetPaddingFactory.cpp
--- a/Engine/Plugins/2D/Paper2D/Source/Paper2DEditor/Private/TileSetEditor/TileSheetPaddingFactory.cpp
+++ b/Engine/Plugins/2D/Paper2D/Source/Paper2DEditor/Private/TileSetEditor/TileSheetPaddingFactory.cpp
@@ -23,9 +23,9 @@ UTileSheetPaddingFactory::UTileSheetPaddingFactory(const FObjectInitializer& Obj
 UObject* UTileSheetPaddingFactory::FactoryCreateNew(UClass* Class, UObject* InParent, FName Name, EObjectFlags Flags, UObject* Context, FFeedbackContext* Warn)
 {
 	check(SourceTileSet);
-	check(SourceTileSet->TileSheet);
 
-	UTexture* SourceTexture = SourceTileSet->TileSheet;
+	UTexture* SourceTexture = SourceTileSet->GetTileSheetTexture();
+	check(SourceTexture);
 
 	if (SourceTexture->Source.GetFormat() != TSF_BGRA8)
 	{
@@ -36,11 +36,10 @@ UObject* UTileSheetPaddingFactory::FactoryCreateNew(UClass* Class, UObject* InPa
 	// Determine how big the new texture needs to be
 	const int32 NumTilesX = SourceTileSet->GetTileCountX();
 	const int32 NumTilesY = SourceTileSet->GetTileCountY();
-	const int32 TileWidth = SourceTileSet->TileWidth;
-	const int32 TileHeight = SourceTileSet->TileHeight;
+	const FIntPoint TileSize = SourceTileSet->GetTileSize();
 
-	const uint32 NewMinTextureWidth = (uint32)(NumTilesX * (TileWidth + 2 * ExtrusionAmount) + (2 * ExtrusionAmount));
-	const uint32 NewMinTextureHeight = (uint32)(NumTilesY * (TileHeight + 2 * ExtrusionAmount) + (2 * ExtrusionAmount));
+	const uint32 NewMinTextureWidth = (uint32)(NumTilesX * (TileSize.X + 2 * ExtrusionAmount) + (2 * ExtrusionAmount));
+	const uint32 NewMinTextureHeight = (uint32)(NumTilesY * (TileSize.Y + 2 * ExtrusionAmount) + (2 * ExtrusionAmount));
 
 	const uint32 NewTextureWidth = bPadToPowerOf2 ? FMath::RoundUpToPowerOfTwo(NewMinTextureWidth) : NewMinTextureWidth;
 	const uint32 NewTextureHeight = bPadToPowerOf2 ? FMath::RoundUpToPowerOfTwo(NewMinTextureHeight) : NewMinTextureHeight;
@@ -56,7 +55,6 @@ UObject* UTileSheetPaddingFactory::FactoryCreateNew(UClass* Class, UObject* InPa
 	TArray<uint8> NewTextureData;
 	NewTextureData.AddZeroed(NewTextureWidth * NewTextureHeight * sizeof(FColor));
 
-	const FIntPoint TileWH(TileWidth, TileHeight);
 	for (int32 TileY = 0; TileY < NumTilesY; ++TileY)
 	{
 		for (int32 TileX = 0; TileX < NumTilesX; ++TileX)
@@ -64,16 +62,16 @@ UObject* UTileSheetPaddingFactory::FactoryCreateNew(UClass* Class, UObject* InPa
 			const FIntPoint TileUV = SourceTileSet->GetTileUVFromTileXY(FIntPoint(TileX, TileY));
 
 			TArray<uint8> DummyBuffer;
-			FPaperAtlasTextureHelpers::ReadSpriteTexture(SourceTexture, TileUV, TileWH, DummyBuffer);
+			FPaperAtlasTextureHelpers::ReadSpriteTexture(SourceTexture, TileUV, TileSize, DummyBuffer);
 
 			FPaperSpriteAtlasSlot Slot;
-			Slot.X = TileX * (TileWidth + (2 * ExtrusionAmount));
-			Slot.Y = TileY * (TileHeight + (2 * ExtrusionAmount));
-			Slot.Width = TileWidth;// +(2 * ExtrusionAmount);
-			Slot.Height = TileHeight;// +(2 * ExtrusionAmount);
+			Slot.X = TileX * (TileSize.X + (2 * ExtrusionAmount));
+			Slot.Y = TileY * (TileSize.Y + (2 * ExtrusionAmount));
+			Slot.Width = TileSize.X;// +(2 * ExtrusionAmount);
+			Slot.Height = TileSize.Y;// +(2 * ExtrusionAmount);
 			Slot.AtlasIndex = 0;
 
-			FPaperAtlasTextureHelpers::CopyTextureRegionToAtlasTextureData(NewTextureData, NewTextureWidth, NewTextureHeight, sizeof(FColor), EPaperSpriteAtlasPadding::DilateBorder, ExtrusionAmount, DummyBuffer, TileWH, Slot);
+			FPaperAtlasTextureHelpers::CopyTextureRegionToAtlasTextureData(NewTextureData, NewTextureWidth, NewTextureHeight, sizeof(FColor), EPaperSpriteAtlasPadding::DilateBorder, ExtrusionAmount, DummyBuffer, TileSize, Slot);
 		}
 	}
 
@@ -84,9 +82,9 @@ UObject* UTileSheetPaddingFactory::FactoryCreateNew(UClass* Class, UObject* InPa
 
 	// Apply the new tile sheet to the specified tile set
 	SourceTileSet->Modify();
-	SourceTileSet->TileSheet = Result;
-	SourceTileSet->Margin = ExtrusionAmount;
-	SourceTileSet->Spacing = 2*ExtrusionAmount;
+	SourceTileSet->SetTileSheetTexture(Result);
+	SourceTileSet->SetMargin(FIntMargin(ExtrusionAmount));
+	SourceTileSet->SetPerTileSpacing(FIntPoint(2 * ExtrusionAmount, 2 * ExtrusionAmount));
 	SourceTileSet->PostEditChange();
 
 	return Result;
diff --git a/Engine/Plugins/2D/Paper2D/Source/PaperTiledImporter/Private/PaperTiledImporterFactory.cpp b/Engine/Plugins/2D/Paper2D/Source/PaperTiledImporter/Private/PaperTiledImporterFactory.cpp
--- a/Engine/Plugins/2D/Paper2D/Source/PaperTiledImporter/Private/PaperTiledImporterFactory.cpp
+++ b/Engine/Plugins/2D/Paper2D/Source/PaperTiledImporter/Private/PaperTiledImporterFactory.cpp
@@ -1164,9 +1164,9 @@ void UPaperTiledImporterFactory::FinalizeTileMap(FTileMapFromTiled& GlobalInfo,
 		BestMaterial = ESpriteInitMaterialType::Automatic;
 		for (UPaperTileSet* TileSet : GlobalInfo.CreatedTileSetAssets)
 		{
-			if ((TileSet != nullptr) && (TileSet->TileSheet != nullptr))
+			if ((TileSet != nullptr) && (TileSet->GetTileSheetTexture() != nullptr))
 			{
-				ESpriteInitMaterialType TileSheetMaterial = ImporterSettings->AnalyzeTextureForDesiredMaterialType(TileSet->TileSheet, FIntPoint::ZeroValue, TileSet->TileSheet->GetImportedSize());
+				ESpriteInitMaterialType TileSheetMaterial = ImporterSettings->AnalyzeTextureForDesiredMaterialType(TileSet->GetTileSheetTexture(), FIntPoint::ZeroValue, TileSet->GetTileSheetAuthoredSize());
 				
 				switch (TileSheetMaterial)
 				{
@@ -1209,17 +1209,19 @@ bool UPaperTiledImporterFactory::ConvertTileSets(FTileMapFromTiled& GlobalInfo,
 			UPaperTileSet* TileSetAsset = CastChecked<UPaperTileSet>(CreateNewAsset(UPaperTileSet::StaticClass(), TargetTileSetPath, TileSetData.Name, Flags));
 			TileSetAsset->Modify();
 
-			TileSetAsset->TileWidth = TileSetData.TileWidth;
-			TileSetAsset->TileHeight = TileSetData.TileHeight;
-			TileSetAsset->Margin = TileSetData.Margin;
-			TileSetAsset->Spacing = TileSetData.Spacing;
-			TileSetAsset->DrawingOffset = FIntPoint(TileSetData.TileOffsetX, TileSetData.TileOffsetY);
+			TileSetAsset->SetTileSize(FIntPoint(TileSetData.TileWidth, TileSetData.TileHeight));
+			TileSetAsset->SetMargin(FIntMargin(TileSetData.Margin));
+			TileSetAsset->SetPerTileSpacing(FIntPoint(TileSetData.Spacing, TileSetData.Spacing));
+			TileSetAsset->SetDrawingOffset(FIntPoint(TileSetData.TileOffsetX, TileSetData.TileOffsetY));
 
 			// Import the texture
 			const FString SourceImageFilename = FPaths::Combine(*CurrentSourcePath, *TileSetData.ImagePath);
-			TileSetAsset->TileSheet = ImportTexture(SourceImageFilename, TargetTexturePath);
 
-			if (TileSetAsset->TileSheet == nullptr)
+			if (UTexture2D* ImportedTileSheetTexture = ImportTexture(SourceImageFilename, TargetTexturePath))
+			{
+				TileSetAsset->SetTileSheetTexture(ImportedTileSheetTexture);
+			}
+			else
 			{
 				UE_LOG(LogPaperTiledImporter, Warning, TEXT("Failed to import tile set image '%s' referenced from tile set '%s'."), *TileSetData.ImagePath, *TileSetData.Name);
 				bLoadedSuccessfully = false;
