Tile Maps: Excluded the per-tile spacing from the bottom right of the tile selection rectangle in the tile set viewport (still currently included in the interior when multiple tiles are selected)

diff --git a/Engine/Plugins/2D/Paper2D/Source/Paper2DEditor/Private/TileSetEditor/TileSetSelectorViewport.cpp b/Engine/Plugins/2D/Paper2D/Source/Paper2DEditor/Private/TileSetEditor/TileSetSelectorViewport.cpp
--- a/Engine/Plugins/2D/Paper2D/Source/Paper2DEditor/Private/TileSetEditor/TileSetSelectorViewport.cpp
+++ b/Engine/Plugins/2D/Paper2D/Source/Paper2DEditor/Private/TileSetEditor/TileSetSelectorViewport.cpp
@@ -191,7 +191,7 @@ void STileSetSelectorViewport::RefreshSelectionRectangle()
 			Client->ValidPaintRectangle.Color = FLinearColor::White;
 
 			const FIntPoint TopLeft = TileSetBeingEdited->GetTileUVFromTileXY(SelectionTopLeft);
-			const FIntPoint BottomRight = TileSetBeingEdited->GetTileUVFromTileXY(SelectionTopLeft + SelectionDimensions);
+			const FIntPoint BottomRight = TileSetBeingEdited->GetTileUVFromTileXY(SelectionTopLeft + SelectionDimensions - FIntPoint(1,1)) + TileSetBeingEdited->GetTileSize();
 			Client->ValidPaintRectangle.Dimensions = BottomRight - TopLeft;
 			Client->ValidPaintRectangle.TopLeft = TopLeft;
 		}
