Tile Sets: Exposed tile set conditioning settings via the tile set editor settings
- Added proper support for power-of-2 padding during conditioning, now sets margin appropriately
- Added an option for how the power-of-2 fill region is initialized (either transparent black or solid white)

diff --git a/Engine/Plugins/2D/Paper2D/Source/Paper2DEditor/Private/TileSetAssetTypeActions.cpp b/Engine/Plugins/2D/Paper2D/Source/Paper2DEditor/Private/TileSetAssetTypeActions.cpp
--- a/Engine/Plugins/2D/Paper2D/Source/Paper2DEditor/Private/TileSetAssetTypeActions.cpp
+++ b/Engine/Plugins/2D/Paper2D/Source/Paper2DEditor/Private/TileSetAssetTypeActions.cpp
@@ -6,6 +6,7 @@
 #include "AssetToolsModule.h"
 #include "ContentBrowserModule.h"
 #include "PaperTileMapFactory.h"
+#include "TileSetEditor/TileSetEditorSettings.h"
 #include "TileSetEditor/TileSheetPaddingFactory.h"
 
 #define LOCTEXT_NAMESPACE "AssetTypeActions"
@@ -109,6 +110,7 @@ void FTileSetAssetTypeActions::ExecutePadTileSetTexture(TWeakObjectPtr<UPaperTil
 {
 	FAssetToolsModule& AssetToolsModule = FModuleManager::GetModuleChecked<FAssetToolsModule>("AssetTools");
 	FContentBrowserModule& ContentBrowserModule = FModuleManager::LoadModuleChecked<FContentBrowserModule>("ContentBrowser");
+	const UTileSetEditorSettings* TileSetEditorSettings = GetDefault<UTileSetEditorSettings>();
 
 	if (UPaperTileSet* TileSet = TileSetPtr.Get())
 	{
@@ -130,6 +132,9 @@ void FTileSetAssetTypeActions::ExecutePadTileSetTexture(TWeakObjectPtr<UPaperTil
 			// Create the new tile sheet
 			UTileSheetPaddingFactory* TileSheetPaddingFactory = NewObject<UTileSheetPaddingFactory>();
 			TileSheetPaddingFactory->SourceTileSet = TileSet;
+			TileSheetPaddingFactory->ExtrusionAmount = TileSetEditorSettings->ExtrusionAmount;
+			TileSheetPaddingFactory->bPadToPowerOf2 = TileSetEditorSettings->bPadToPowerOf2;
+			TileSheetPaddingFactory->bFillWithTransparentBlack = TileSetEditorSettings->bFillWithTransparentBlack;
 			ContentBrowserModule.Get().CreateNewAsset(AssetName, PackagePath, UTexture::StaticClass(), TileSheetPaddingFactory);
 		}
 	}
diff --git a/Engine/Plugins/2D/Paper2D/Source/Paper2DEditor/Private/TileSetEditor/TileSetEditorSettings.cpp b/Engine/Plugins/2D/Paper2D/Source/Paper2DEditor/Private/TileSetEditor/TileSetEditorSettings.cpp
--- a/Engine/Plugins/2D/Paper2D/Source/Paper2DEditor/Private/TileSetEditor/TileSetEditorSettings.cpp
+++ b/Engine/Plugins/2D/Paper2D/Source/Paper2DEditor/Private/TileSetEditor/TileSetEditorSettings.cpp
@@ -9,5 +9,8 @@
 UTileSetEditorSettings::UTileSetEditorSettings()
 	: DefaultBackgroundColor(0, 0, 127)
 	, bShowGridByDefault(true)
+	, ExtrusionAmount(2)
+	, bPadToPowerOf2(true)
+	, bFillWithTransparentBlack(true)
 {
 }
diff --git a/Engine/Plugins/2D/Paper2D/Source/Paper2DEditor/Private/TileSetEditor/TileSetEditorSettings.h b/Engine/Plugins/2D/Paper2D/Source/Paper2DEditor/Private/TileSetEditor/TileSetEditorSettings.h
--- a/Engine/Plugins/2D/Paper2D/Source/Paper2DEditor/Private/TileSetEditor/TileSetEditorSettings.h
+++ b/Engine/Plugins/2D/Paper2D/Source/Paper2DEditor/Private/TileSetEditor/TileSetEditorSettings.h
@@ -13,11 +13,23 @@ class UTileSetEditorSettings : public UObject
 public:
 	UTileSetEditorSettings();
 
-	/** Default background color for new tile set assets */
+	// Default background color for new tile set assets
 	UPROPERTY(config, EditAnywhere, Category=Background, meta=(HideAlphaChannel))
 	FColor DefaultBackgroundColor;
 
-	/** Should the grid be shown by default when the editor is opened? */
+	// Should the grid be shown by default when the editor is opened?
 	UPROPERTY(config, EditAnywhere, Category="Tile Editor")
 	bool bShowGridByDefault;
+
+	// The amount to extrude out from the edge of each tile (in pixels)
+	UPROPERTY(config, EditAnywhere, Category="Tile Sheet Conditioning", meta=(UIMin=1, ClampMin=1))
+	int32 ExtrusionAmount;
+
+	// Should we pad the conditioned texture to the next power of 2?
+	UPROPERTY(config, EditAnywhere, Category="Tile Sheet Conditioning")
+	bool bPadToPowerOf2;
+
+	// Should we use transparent black or white when filling the texture areas that aren't covered by tiles?
+	UPROPERTY(config, EditAnywhere, Category="Tile Sheet Conditioning")
+	bool bFillWithTransparentBlack;
 };
diff --git a/Engine/Plugins/2D/Paper2D/Source/Paper2DEditor/Private/TileSetEditor/TileSheetPaddingFactory.cpp b/Engine/Plugins/2D/Paper2D/Source/Paper2DEditor/Private/TileSetEditor/TileSheetPaddingFactory.cpp
--- a/Engine/Plugins/2D/Paper2D/Source/Paper2DEditor/Private/TileSetEditor/TileSheetPaddingFactory.cpp
+++ b/Engine/Plugins/2D/Paper2D/Source/Paper2DEditor/Private/TileSetEditor/TileSheetPaddingFactory.cpp
@@ -17,7 +17,8 @@ UTileSheetPaddingFactory::UTileSheetPaddingFactory(const FObjectInitializer& Obj
 	SupportedClass = UTexture::StaticClass();
 
 	ExtrusionAmount = 2;
-	bPadToPowerOf2 = false;
+	bPadToPowerOf2 = true;
+	bFillWithTransparentBlack = true;
 }
 
 UObject* UTileSheetPaddingFactory::FactoryCreateNew(UClass* Class, UObject* InParent, FName Name, EObjectFlags Flags, UObject* Context, FFeedbackContext* Warn)
@@ -33,13 +34,19 @@ UObject* UTileSheetPaddingFactory::FactoryCreateNew(UClass* Class, UObject* InPa
 		return nullptr;
 	}
 
-	// Determine how big the new texture needs to be
 	const int32 NumTilesX = SourceTileSet->GetTileCountX();
 	const int32 NumTilesY = SourceTileSet->GetTileCountY();
 	const FIntPoint TileSize = SourceTileSet->GetTileSize();
 
-	const uint32 NewMinTextureWidth = (uint32)(NumTilesX * (TileSize.X + 2 * ExtrusionAmount) + (2 * ExtrusionAmount));
-	const uint32 NewMinTextureHeight = (uint32)(NumTilesY * (TileSize.Y + 2 * ExtrusionAmount) + (2 * ExtrusionAmount));
+	if ((NumTilesX <= 0) || (NumTilesY <= 0))
+	{
+		UE_LOG(LogPaper2DEditor, Error, TEXT("Tile sheet texture '%s' is too small to contain any tiles, cannot create a padded texture from it"), *SourceTexture->GetName());
+		return nullptr;
+	}
+
+	// Determine how big the new texture needs to be
+	const uint32 NewMinTextureWidth = (uint32)(NumTilesX * (TileSize.X + 2 * ExtrusionAmount));
+	const uint32 NewMinTextureHeight = (uint32)(NumTilesY * (TileSize.Y + 2 * ExtrusionAmount));
 
 	const uint32 NewTextureWidth = bPadToPowerOf2 ? FMath::RoundUpToPowerOfTwo(NewMinTextureWidth) : NewMinTextureWidth;
 	const uint32 NewTextureHeight = bPadToPowerOf2 ? FMath::RoundUpToPowerOfTwo(NewMinTextureHeight) : NewMinTextureHeight;
@@ -49,11 +56,14 @@ UObject* UTileSheetPaddingFactory::FactoryCreateNew(UClass* Class, UObject* InPa
 	//@TODO: Copy more state across (or start by duplicating the texture maybe?)
 	Result->LODGroup = SourceTexture->LODGroup;
 	Result->CompressionSettings = SourceTexture->CompressionSettings;
-	Result->MipGenSettings = TMGS_NoMipmaps; //@TODO: Don't actually want this...
+	Result->MipGenSettings = bPadToPowerOf2 ? TMGS_FromTextureGroup : TMGS_NoMipmaps;
 	Result->DeferCompression = true;
 
+	const uint32 TextureDataSize = NewTextureWidth * NewTextureHeight * sizeof(FColor);
+
 	TArray<uint8> NewTextureData;
-	NewTextureData.AddZeroed(NewTextureWidth * NewTextureHeight * sizeof(FColor));
+	NewTextureData.AddUninitialized(TextureDataSize);
+	FMemory::Memset(NewTextureData.GetData(), bFillWithTransparentBlack ? 0x00 : 0xFF, TextureDataSize);
 
 	for (int32 TileY = 0; TileY < NumTilesY; ++TileY)
 	{
@@ -67,8 +77,8 @@ UObject* UTileSheetPaddingFactory::FactoryCreateNew(UClass* Class, UObject* InPa
 			FPaperSpriteAtlasSlot Slot;
 			Slot.X = TileX * (TileSize.X + (2 * ExtrusionAmount));
 			Slot.Y = TileY * (TileSize.Y + (2 * ExtrusionAmount));
-			Slot.Width = TileSize.X;// +(2 * ExtrusionAmount);
-			Slot.Height = TileSize.Y;// +(2 * ExtrusionAmount);
+			Slot.Width = TileSize.X;
+			Slot.Height = TileSize.Y;
 			Slot.AtlasIndex = 0;
 
 			FPaperAtlasTextureHelpers::CopyTextureRegionToAtlasTextureData(NewTextureData, NewTextureWidth, NewTextureHeight, sizeof(FColor), EPaperSpriteAtlasPadding::DilateBorder, ExtrusionAmount, DummyBuffer, TileSize, Slot);
@@ -80,10 +90,15 @@ UObject* UTileSheetPaddingFactory::FactoryCreateNew(UClass* Class, UObject* InPa
 	Result->UpdateResource();
 	Result->PostEditChange();
 
+	// Figure out the margin (the right/bottom might be quite large due the the power-of-2 padding)
+	const int32 ExcessWidth = (int32)(NewTextureWidth - NewMinTextureWidth);
+	const int32 ExcessHeight = (int32)(NewTextureHeight - NewMinTextureHeight);
+	FIntMargin BorderMargin(ExtrusionAmount, ExtrusionAmount, ExtrusionAmount + ExcessWidth, ExtrusionAmount + ExcessHeight);
+
 	// Apply the new tile sheet to the specified tile set
 	SourceTileSet->Modify();
 	SourceTileSet->SetTileSheetTexture(Result);
-	SourceTileSet->SetMargin(FIntMargin(ExtrusionAmount));
+	SourceTileSet->SetMargin(BorderMargin);
 	SourceTileSet->SetPerTileSpacing(FIntPoint(2 * ExtrusionAmount, 2 * ExtrusionAmount));
 	SourceTileSet->PostEditChange();
 
diff --git a/Engine/Plugins/2D/Paper2D/Source/Paper2DEditor/Private/TileSetEditor/TileSheetPaddingFactory.h b/Engine/Plugins/2D/Paper2D/Source/Paper2DEditor/Private/TileSetEditor/TileSheetPaddingFactory.h
--- a/Engine/Plugins/2D/Paper2D/Source/Paper2DEditor/Private/TileSetEditor/TileSheetPaddingFactory.h
+++ b/Engine/Plugins/2D/Paper2D/Source/Paper2DEditor/Private/TileSetEditor/TileSheetPaddingFactory.h
@@ -27,6 +27,9 @@ class UTileSheetPaddingFactory : public UFactory
 	UPROPERTY(Category=Settings, EditAnywhere)
 	bool bPadToPowerOf2;
 
+	// Should we use transparent black or white when filling the texture areas that aren't covered by tiles?
+	UPROPERTY(Category=Settings, EditAnywhere)
+	bool bFillWithTransparentBlack;
 
 	// UFactory interface
 	virtual UObject* FactoryCreateNew(UClass* Class, UObject* InParent, FName Name, EObjectFlags Flags, UObject* Context, FFeedbackContext* Warn) override;
