Tile Set Editor: Added visualization of the border margin to the tile set editor and excluded it from the thumbnail

diff --git a/Engine/Plugins/2D/Paper2D/Source/Paper2DEditor/Private/PaperTileSetThumbnailRenderer.cpp b/Engine/Plugins/2D/Paper2D/Source/Paper2DEditor/Private/PaperTileSetThumbnailRenderer.cpp
--- a/Engine/Plugins/2D/Paper2D/Source/Paper2DEditor/Private/PaperTileSetThumbnailRenderer.cpp
+++ b/Engine/Plugins/2D/Paper2D/Source/Paper2DEditor/Private/PaperTileSetThumbnailRenderer.cpp
@@ -49,12 +49,14 @@ void UPaperTileSetThumbnailRenderer::Draw(UObject* Object, int32 X, int32 Y, uin
 		const float TextureWidth = TileSheetTexture->GetSurfaceWidth();
 		const float TextureHeight = TileSheetTexture->GetSurfaceHeight();
 
+		const FIntMargin Margin = TileSet->GetMargin();
+
 		float FinalX = (float)X;
 		float FinalY = (float)Y;
 		float FinalWidth = (float)Width;
 		float FinalHeight = (float)Height;
-		const float DesiredWidth = TextureWidth;
-		const float DesiredHeight = TextureHeight;
+		const float DesiredWidth = TextureWidth - Margin.GetDesiredSize().X;
+		const float DesiredHeight = TextureHeight - Margin.GetDesiredSize().Y;
 
 		const FLinearColor BlackBarColor(0.0f, 0.0f, 0.0f, 0.5f);
 
@@ -81,18 +83,18 @@ void UPaperTileSetThumbnailRenderer::Draw(UObject* Object, int32 X, int32 Y, uin
 			Canvas->DrawTile(FinalX+FinalWidth, Y, Width-FinalWidth, Height, 0, 0, 1, 1, BlackBarColor, GWhiteTexture, bAlphaBlend);
 		}
 
-		//@TODO: Cut off the border margin when rendering
-
 		// Draw the tile sheet 
+		const float InvWidth = 1.0f / TextureWidth;
+		const float InvHeight = 1.0f / TextureHeight;
 		Canvas->DrawTile(
 			FinalX,
 			FinalY,
 			FinalWidth,
 			FinalHeight,
-			0.0f,
-			0.0f,
-			1.0f,
-			1.0f,
+			Margin.Left * InvWidth,
+			Margin.Top * InvHeight,
+			(TextureWidth - Margin.Right) * InvWidth,
+			(TextureHeight - Margin.Bottom) * InvHeight,
 			FLinearColor::White,
 			TileSheetTexture->Resource,
 			bUseTranslucentBlend);
diff --git a/Engine/Plugins/2D/Paper2D/Source/Paper2DEditor/Private/TileSetEditor/TileSetEditorViewportClient.cpp b/Engine/Plugins/2D/Paper2D/Source/Paper2DEditor/Private/TileSetEditor/TileSetEditorViewportClient.cpp
--- a/Engine/Plugins/2D/Paper2D/Source/Paper2DEditor/Private/TileSetEditor/TileSetEditorViewportClient.cpp
+++ b/Engine/Plugins/2D/Paper2D/Source/Paper2DEditor/Private/TileSetEditor/TileSetEditorViewportClient.cpp
@@ -49,6 +49,23 @@ void FTileSetEditorViewportClient::Draw(FViewport* Viewport, FCanvas* Canvas)
 			Canvas->DrawTile(XPos, YPos, Width, Height, 0.0f, 0.0f, 1.0f, 1.0f, TextureDrawColor, Texture->Resource, bUseTranslucentBlend);
 		}
 
+
+		const FLinearColor BorderRectangleColor(0.3f, 0.3f, 0.3f, 1.0f);
+		{
+			const FIntPoint TextureSize = Texture->GetImportedSize();
+			const FIntMargin BorderSize = TileSet->GetMargin();
+			const FIntRect TileSetRegion(BorderSize.Left, BorderSize.Top, TextureSize.X - BorderSize.Right, TextureSize.Y - BorderSize.Bottom);
+
+			const float X = (TileSetRegion.Min.X - ZoomPos.X) * ZoomAmount;
+			const float Y = (TileSetRegion.Min.Y - ZoomPos.Y) * ZoomAmount;
+			const float W = TileSetRegion.Width() * ZoomAmount;
+			const float H = TileSetRegion.Height() * ZoomAmount;
+
+			FCanvasBoxItem BoxItem(FVector2D(X, Y), FVector2D(W, H));
+			BoxItem.SetColor(BorderRectangleColor);
+			Canvas->DrawItem(BoxItem);
+		}
+
 		if (bShowTilesWithCollision || bShowTilesWithMetaData)
 		{
 			// Draw an overlay rectangle on top of any tiles that have collision or metadata geometry
