Tile Set Editor: Fixed the tile selection rectangle not being updated when the tile set properties were modified

diff --git a/Engine/Plugins/2D/Paper2D/Source/Paper2DEditor/Private/TileSetEditor.cpp b/Engine/Plugins/2D/Paper2D/Source/Paper2DEditor/Private/TileSetEditor.cpp
--- a/Engine/Plugins/2D/Paper2D/Source/Paper2DEditor/Private/TileSetEditor.cpp
+++ b/Engine/Plugins/2D/Paper2D/Source/Paper2DEditor/Private/TileSetEditor.cpp
@@ -298,6 +298,7 @@ void FTileSetEditor::OnPropertyChanged(UObject* ObjectBeingModified, FPropertyCh
 	if (ObjectBeingModified == TileSetBeingEdited)
 	{
 		TileEditorViewportClient->SetTileIndex(TileEditorViewportClient->GetTileIndex());
+		TileSetViewport->RefreshSelectionRectangle();
 	}
 }
 
