Tile Sets: Fixed the 'Create Tile Set' Texture command not allocating per-tile data

diff --git a/Engine/Plugins/2D/Paper2D/Source/Paper2DEditor/Private/PaperTileSetFactory.cpp b/Engine/Plugins/2D/Paper2D/Source/Paper2DEditor/Private/PaperTileSetFactory.cpp
--- a/Engine/Plugins/2D/Paper2D/Source/Paper2DEditor/Private/PaperTileSetFactory.cpp
+++ b/Engine/Plugins/2D/Paper2D/Source/Paper2DEditor/Private/PaperTileSetFactory.cpp
@@ -23,6 +23,8 @@ UObject* UPaperTileSetFactory::FactoryCreateNew(UClass* Class, UObject* InParent
 	NewTileSet->SetTileSheetTexture(InitialTexture);
 	NewTileSet->SetBackgroundColor(GetDefault<UTileSetEditorSettings>()->DefaultBackgroundColor);
 
+	NewTileSet->PostEditChange();
+
 	return NewTileSet;
 }
 
