Tile Set Editor: Prevented a crash when clicking on a tile if there is no per-tile data stored

diff --git a/Engine/Plugins/2D/Paper2D/Source/Paper2DEditor/Private/TileSetEditor/TileSetDetailsCustomization.cpp b/Engine/Plugins/2D/Paper2D/Source/Paper2DEditor/Private/TileSetEditor/TileSetDetailsCustomization.cpp
--- a/Engine/Plugins/2D/Paper2D/Source/Paper2DEditor/Private/TileSetEditor/TileSetDetailsCustomization.cpp
+++ b/Engine/Plugins/2D/Paper2D/Source/Paper2DEditor/Private/TileSetEditor/TileSetDetailsCustomization.cpp
@@ -78,10 +78,13 @@ void FTileSetDetailsCustomization::CustomizeDetails(IDetailLayoutBuilder& Detail
 			// Customize for the single tile being edited
 			IDetailCategoryBuilder& SingleTileCategory = DetailLayout.EditCategory("SingleTileEditor", FText::GetEmpty());
 			
-			TSharedPtr<IPropertyHandle> OneTileEntry = PerTileArrayProperty->GetChildHandle(SelectedSingleTileIndex);
-			SingleTileCategory.AddProperty(OneTileEntry)
-				.ShouldAutoExpand(true);
-
+			uint32 NumChildren;
+			if ((PerTileArrayProperty->GetNumChildren(/*out*/ NumChildren) == FPropertyAccess::Success) && ((uint32)SelectedSingleTileIndex < NumChildren))
+			{
+				TSharedPtr<IPropertyHandle> OneTileEntry = PerTileArrayProperty->GetChildHandle(SelectedSingleTileIndex);
+				SingleTileCategory.AddProperty(OneTileEntry)
+					.ShouldAutoExpand(true);
+			}
 
 			// Add a display of the tile index being edited to the header
 			const FText TileIndexHeaderText = FText::Format(LOCTEXT("SingleTileSectionHeader", "Editing Tile #{0}"), FText::AsNumber(SelectedSingleTileIndex));
