Tile Maps: Changed layer batch submission order to go back to front to help translucent rendering (only works when viewing from the front though since components are sorted as a single unit) [UE-13645]

diff --git a/Engine/Plugins/2D/Paper2D/Source/Paper2D/Private/PaperTileMapComponent.cpp b/Engine/Plugins/2D/Paper2D/Source/Paper2D/Private/PaperTileMapComponent.cpp
--- a/Engine/Plugins/2D/Paper2D/Source/Paper2D/Private/PaperTileMapComponent.cpp
+++ b/Engine/Plugins/2D/Paper2D/Source/Paper2D/Private/PaperTileMapComponent.cpp
@@ -228,7 +228,7 @@ void UPaperTileMapComponent::RebuildRenderData(FPaperTileMapRenderSceneProxy* Pr
 
 	const float UnrealUnitsPerPixel = TileMap->GetUnrealUnitsPerPixel();
 
-	for (int32 Z = 0; Z < TileMap->TileLayers.Num(); ++Z)
+	for (int32 Z = TileMap->TileLayers.Num() - 1; Z >= 0; --Z)
 	{
 		UPaperTileLayer* Layer = TileMap->TileLayers[Z];
 
