Tile Maps: Added rename to layer content menu [UE-14714]

diff --git a/Engine/Plugins/2D/Paper2D/Source/Paper2DEditor/Private/PaperStyle.cpp b/Engine/Plugins/2D/Paper2D/Source/Paper2DEditor/Private/PaperStyle.cpp
--- a/Engine/Plugins/2D/Paper2D/Source/Paper2DEditor/Private/PaperStyle.cpp
+++ b/Engine/Plugins/2D/Paper2D/Source/Paper2DEditor/Private/PaperStyle.cpp
@@ -87,18 +87,27 @@ void FPaperStyle::Initialize()
 		StyleSet->Set("TileMapEditor.DeleteLayer", new IMAGE_PLUGIN_BRUSH(TEXT("Icons/TileMapEditor/icon_TileMapEditor_DeleteLayer_40x"), Icon40x40));
 		StyleSet->Set("TileMapEditor.DeleteLayer.Small", new IMAGE_PLUGIN_BRUSH(TEXT("Icons/TileMapEditor/icon_TileMapEditor_DeleteLayer_40x"), Icon20x20));
 
+// 		StyleSet->Set("TileMapEditor.RenameLayer", new FSlateNoResource(Icon16x16));
+// 		StyleSet->Set("TileMapEditor.RenameLayer.Small", new FSlateNoResource(Icon16x16));
+		
 		StyleSet->Set("TileMapEditor.DuplicateLayer", new IMAGE_PLUGIN_BRUSH(TEXT("Icons/TileMapEditor/icon_TileMapEditor_DuplicateLayer_40x"), Icon40x40));
 		StyleSet->Set("TileMapEditor.DuplicateLayer.Small", new IMAGE_PLUGIN_BRUSH(TEXT("Icons/TileMapEditor/icon_TileMapEditor_DuplicateLayer_40x"), Icon20x20));
 
 		StyleSet->Set("TileMapEditor.MergeLayerDown", new IMAGE_PLUGIN_BRUSH(TEXT("Icons/TileMapEditor/icon_TileMapEditor_MergeLayerDown_40x"), Icon40x40));
 		StyleSet->Set("TileMapEditor.MergeLayerDown.Small", new IMAGE_PLUGIN_BRUSH(TEXT("Icons/TileMapEditor/icon_TileMapEditor_MergeLayerDown_40x"), Icon20x20));
 
+// 		StyleSet->Set("TileMapEditor.MoveLayerToTop", new FSlateNoResource(Icon16x16));
+// 		StyleSet->Set("TileMapEditor.MoveLayerToTop.Small", new FSlateNoResource(Icon16x16));
+
 		StyleSet->Set("TileMapEditor.MoveLayerUp", new IMAGE_PLUGIN_BRUSH(TEXT("Icons/TileMapEditor/icon_TileMapEditor_MoveLayerUp_40x"), Icon40x40));
 		StyleSet->Set("TileMapEditor.MoveLayerUp.Small", new IMAGE_PLUGIN_BRUSH(TEXT("Icons/TileMapEditor/icon_TileMapEditor_MoveLayerUp_40x"), Icon20x20));
 
 		StyleSet->Set("TileMapEditor.MoveLayerDown", new IMAGE_PLUGIN_BRUSH(TEXT("Icons/TileMapEditor/icon_TileMapEditor_MoveLayerDown_40x"), Icon40x40));
 		StyleSet->Set("TileMapEditor.MoveLayerDown.Small", new IMAGE_PLUGIN_BRUSH(TEXT("Icons/TileMapEditor/icon_TileMapEditor_MoveLayerDown_40x"), Icon20x20));
 
+// 		StyleSet->Set("TileMapEditor.MoveLayerToBottom", new FSlateNoResource(Icon16x16));
+// 		StyleSet->Set("TileMapEditor.MoveLayerToBottom.Small", new FSlateNoResource(Icon16x16));
+
 		StyleSet->Set("TileMapEditor.LayerEyeClosed", new IMAGE_PLUGIN_BRUSH(TEXT("Icons/TileMapEditor/icon_EyeClosed_40x"), Icon16x16));
 		StyleSet->Set("TileMapEditor.LayerEyeOpened", new IMAGE_PLUGIN_BRUSH(TEXT("Icons/TileMapEditor/icon_EyeOpened_40x"), Icon16x16));
 
diff --git a/Engine/Plugins/2D/Paper2D/Source/Paper2DEditor/Private/TileMapEditing/STileLayerItem.cpp b/Engine/Plugins/2D/Paper2D/Source/Paper2DEditor/Private/TileMapEditing/STileLayerItem.cpp
--- a/Engine/Plugins/2D/Paper2D/Source/Paper2DEditor/Private/TileMapEditing/STileLayerItem.cpp
+++ b/Engine/Plugins/2D/Paper2D/Source/Paper2DEditor/Private/TileMapEditing/STileLayerItem.cpp
@@ -22,6 +22,12 @@ void STileLayerItem::Construct(const FArguments& InArgs, int32 Index, class UPap
 	EyeClosed = FPaperStyle::Get()->GetBrush(EyeClosedBrushName);
 	EyeOpened = FPaperStyle::Get()->GetBrush(EyeOpenedBrushName);
 
+	LayerNameWidget = SNew(SInlineEditableTextBlock)
+		.Text(this, &STileLayerItem::GetLayerDisplayName)
+		.ToolTipText(this, &STileLayerItem::GetLayerDisplayName)
+		.OnTextCommitted(this, &STileLayerItem::OnLayerNameCommitted)
+		.IsSelected(InIsSelectedDelegate);
+
 	ChildSlot
 	[
 		SNew(SHorizontalBox)
@@ -48,15 +54,16 @@ void STileLayerItem::Construct(const FArguments& InArgs, int32 Index, class UPap
 		.VAlign(VAlign_Center)
 		.Padding(FMargin(4.0f, 4.0f, 4.0f, 4.0f))
 		[
-			SNew(SInlineEditableTextBlock)
-			.Text(this, &STileLayerItem::GetLayerDisplayName)
-			.ToolTipText(this, &STileLayerItem::GetLayerDisplayName)
-			.OnTextCommitted(this, &STileLayerItem::OnLayerNameCommitted)
-			.IsSelected(InIsSelectedDelegate)
+			LayerNameWidget.ToSharedRef()
 		]
 	];
 }
 
+void STileLayerItem::BeginEditingName()
+{
+	LayerNameWidget->EnterEditingMode();
+}
+
 FText STileLayerItem::GetLayerDisplayName() const
 {
 	const FText UnnamedText = LOCTEXT("NoLayerName", "(unnamed)");
diff --git a/Engine/Plugins/2D/Paper2D/Source/Paper2DEditor/Private/TileMapEditing/STileLayerItem.h b/Engine/Plugins/2D/Paper2D/Source/Paper2DEditor/Private/TileMapEditing/STileLayerItem.h
--- a/Engine/Plugins/2D/Paper2D/Source/Paper2DEditor/Private/TileMapEditing/STileLayerItem.h
+++ b/Engine/Plugins/2D/Paper2D/Source/Paper2DEditor/Private/TileMapEditing/STileLayerItem.h
@@ -13,6 +13,8 @@ public:
 
 	void Construct(const FArguments& InArgs, int32 Index, class UPaperTileMap* InMap, FIsSelected InIsSelectedDelegate);
 
+	void BeginEditingName();
+
 protected:
 	int32 MyIndex;
 	class UPaperTileMap* MyMap;
@@ -23,6 +25,8 @@ protected:
 	const FSlateBrush* EyeClosed;
 	const FSlateBrush* EyeOpened;
 
+	TSharedPtr<SInlineEditableTextBlock> LayerNameWidget;
+
 protected:
 	FText GetLayerDisplayName() const;
 	void OnLayerNameCommitted(const FText& NewText, ETextCommit::Type CommitInfo);
diff --git a/Engine/Plugins/2D/Paper2D/Source/Paper2DEditor/Private/TileMapEditing/STileLayerList.cpp b/Engine/Plugins/2D/Paper2D/Source/Paper2DEditor/Private/TileMapEditing/STileLayerList.cpp
--- a/Engine/Plugins/2D/Paper2D/Source/Paper2DEditor/Private/TileMapEditing/STileLayerList.cpp
+++ b/Engine/Plugins/2D/Paper2D/Source/Paper2DEditor/Private/TileMapEditing/STileLayerList.cpp
@@ -38,6 +38,11 @@ void STileLayerList::Construct(const FArguments& InArgs, UPaperTileMap* InTileMa
 		FExecuteAction::CreateSP(this, &STileLayerList::DeleteLayer),
 		FCanExecuteAction::CreateSP(this, &STileLayerList::CanExecuteActionNeedingSelectedLayer));
 
+	CommandList->MapAction(
+		Commands.RenameLayer,
+		FExecuteAction::CreateSP(this, &STileLayerList::RenameLayer),
+		FCanExecuteAction::CreateSP(this, &STileLayerList::CanExecuteActionNeedingSelectedLayer));
+
 	CommandList->MapAction(
 		Commands.DuplicateLayer,
 		FExecuteAction::CreateSP(this, &STileLayerList::DuplicateLayer),
@@ -311,6 +316,28 @@ void STileLayerList::DeleteLayer()
 	}
 }
 
+void STileLayerList::RenameLayer()
+{
+	if (UPaperTileMap* TileMap = TileMapPtr.Get())
+	{
+		const int32 RenameIndex = GetSelectionIndex();
+		if (MirrorList.IsValidIndex(RenameIndex))
+		{
+			TSharedPtr<ITableRow> LayerRowWidget = ListViewWidget->WidgetFromItem(MirrorList[RenameIndex]);
+			if (LayerRowWidget.IsValid())
+			{
+				TSharedPtr<SWidget> RowContent = LayerRowWidget->GetContent();
+				if (RowContent.IsValid())
+				{
+					TSharedPtr<STileLayerItem> LayerWidget = StaticCastSharedPtr<STileLayerItem>(RowContent);
+					LayerWidget->BeginEditingName();
+				}
+			}
+		}
+	}
+}
+
+
 void STileLayerList::DuplicateLayer()
 {
 	if (UPaperTileMap* TileMap = TileMapPtr.Get())
@@ -473,23 +500,26 @@ TSharedPtr<SWidget> STileLayerList::OnConstructContextMenu()
 
 	const FTileMapEditorCommands& Commands = FTileMapEditorCommands::Get();
 
+	FSlateIcon DummyIcon(NAME_None, NAME_None);
+
 	MenuBuilder.BeginSection("BasicOperations", LOCTEXT("BasicOperationsHeader", "Layer actions"));
  	{
-		MenuBuilder.AddMenuEntry(Commands.DuplicateLayer);
-		MenuBuilder.AddMenuEntry(Commands.DeleteLayer);
-		MenuBuilder.AddMenuEntry(Commands.MergeLayerDown);
+		MenuBuilder.AddMenuEntry(Commands.DuplicateLayer, NAME_None, TAttribute<FText>(), TAttribute<FText>(), DummyIcon);
+		MenuBuilder.AddMenuEntry(Commands.DeleteLayer, NAME_None, TAttribute<FText>(), TAttribute<FText>(), DummyIcon);
+		MenuBuilder.AddMenuEntry(Commands.RenameLayer, NAME_None, TAttribute<FText>(), TAttribute<FText>(), DummyIcon);
+		MenuBuilder.AddMenuEntry(Commands.MergeLayerDown, NAME_None, TAttribute<FText>(), TAttribute<FText>(), DummyIcon);
 		MenuBuilder.AddMenuSeparator();
-		MenuBuilder.AddMenuEntry(Commands.SelectLayerAbove);
-		MenuBuilder.AddMenuEntry(Commands.SelectLayerBelow);
+		MenuBuilder.AddMenuEntry(Commands.SelectLayerAbove, NAME_None, TAttribute<FText>(), TAttribute<FText>(), DummyIcon);
+		MenuBuilder.AddMenuEntry(Commands.SelectLayerBelow, NAME_None, TAttribute<FText>(), TAttribute<FText>(), DummyIcon);
 	}
 	MenuBuilder.EndSection();
 
 	MenuBuilder.BeginSection("OrderingOperations", LOCTEXT("OrderingOperationsHeader", "Order actions"));
 	{
-		MenuBuilder.AddMenuEntry(Commands.MoveLayerToTop);
-		MenuBuilder.AddMenuEntry(Commands.MoveLayerUp);
-		MenuBuilder.AddMenuEntry(Commands.MoveLayerDown);
-		MenuBuilder.AddMenuEntry(Commands.MoveLayerToBottom);
+		MenuBuilder.AddMenuEntry(Commands.MoveLayerToTop, NAME_None, TAttribute<FText>(), TAttribute<FText>(), DummyIcon);
+		MenuBuilder.AddMenuEntry(Commands.MoveLayerUp, NAME_None, TAttribute<FText>(), TAttribute<FText>(), DummyIcon);
+		MenuBuilder.AddMenuEntry(Commands.MoveLayerDown, NAME_None, TAttribute<FText>(), TAttribute<FText>(), DummyIcon);
+		MenuBuilder.AddMenuEntry(Commands.MoveLayerToBottom, NAME_None, TAttribute<FText>(), TAttribute<FText>(), DummyIcon);
 	}
 	MenuBuilder.EndSection();
 
diff --git a/Engine/Plugins/2D/Paper2D/Source/Paper2DEditor/Private/TileMapEditing/STileLayerList.h b/Engine/Plugins/2D/Paper2D/Source/Paper2DEditor/Private/TileMapEditing/STileLayerList.h
--- a/Engine/Plugins/2D/Paper2D/Source/Paper2DEditor/Private/TileMapEditing/STileLayerList.h
+++ b/Engine/Plugins/2D/Paper2D/Source/Paper2DEditor/Private/TileMapEditing/STileLayerList.h
@@ -48,6 +48,7 @@ protected:
 	void AddNewLayerAbove();
 	void AddNewLayerBelow();
 	void DeleteLayer();
+	void RenameLayer();
 	void DuplicateLayer();
 	void MergeLayerDown();
 	void MoveLayerUp(bool bForceToTop);
diff --git a/Engine/Plugins/2D/Paper2D/Source/Paper2DEditor/Private/TileMapEditing/TileMapEditorCommands.cpp b/Engine/Plugins/2D/Paper2D/Source/Paper2DEditor/Private/TileMapEditing/TileMapEditorCommands.cpp
--- a/Engine/Plugins/2D/Paper2D/Source/Paper2DEditor/Private/TileMapEditing/TileMapEditorCommands.cpp
+++ b/Engine/Plugins/2D/Paper2D/Source/Paper2DEditor/Private/TileMapEditing/TileMapEditorCommands.cpp
@@ -46,6 +46,7 @@ void FTileMapEditorCommands::RegisterCommands()
 	UI_COMMAND(MoveLayerToBottom, "Send to back", "Send to back", EUserInterfaceActionType::Button, FInputChord(EKeys::RightBracket, EModifierKey::Shift | EModifierKey::Control));
 	UI_COMMAND(SelectLayerAbove, "Select next layer", "Changes the active layer to the next layer", EUserInterfaceActionType::Button, FInputChord(EKeys::RightBracket, EModifierKey::Alt));
 	UI_COMMAND(SelectLayerBelow, "Select previous layer", "Changes the active layer to the previous layer", EUserInterfaceActionType::Button, FInputChord(EKeys::LeftBracket, EModifierKey::Alt));
+	UI_COMMAND(RenameLayer, "Rename layer", "Renames the selected layer", EUserInterfaceActionType::Button, FInputChord(EKeys::F2));
 }
 
 #undef LOCTEXT_NAMESPACE
\ No newline at end of file
diff --git a/Engine/Plugins/2D/Paper2D/Source/Paper2DEditor/Private/TileMapEditing/TileMapEditorCommands.h b/Engine/Plugins/2D/Paper2D/Source/Paper2DEditor/Private/TileMapEditing/TileMapEditorCommands.h
--- a/Engine/Plugins/2D/Paper2D/Source/Paper2DEditor/Private/TileMapEditing/TileMapEditorCommands.h
+++ b/Engine/Plugins/2D/Paper2D/Source/Paper2DEditor/Private/TileMapEditing/TileMapEditorCommands.h
@@ -59,4 +59,5 @@ public:
 	TSharedPtr<FUICommandInfo> MoveLayerToBottom;
 	TSharedPtr<FUICommandInfo> SelectLayerAbove;
 	TSharedPtr<FUICommandInfo> SelectLayerBelow;
+	TSharedPtr<FUICommandInfo> RenameLayer;
 };
\ No newline at end of file
