Paper2D: Hid Alpha component for OutlineColor in the sprite extraction dialog [UE-13341]

diff --git a/Engine/Plugins/2D/Paper2D/Source/Paper2DEditor/Private/ExtractSprites/PaperExtractSpritesSettings.cpp b/Engine/Plugins/2D/Paper2D/Source/Paper2DEditor/Private/ExtractSprites/PaperExtractSpritesSettings.cpp
--- a/Engine/Plugins/2D/Paper2D/Source/Paper2DEditor/Private/ExtractSprites/PaperExtractSpritesSettings.cpp
+++ b/Engine/Plugins/2D/Paper2D/Source/Paper2DEditor/Private/ExtractSprites/PaperExtractSpritesSettings.cpp
@@ -11,18 +11,6 @@ UPaperExtractSpritesSettings::UPaperExtractSpritesSettings(const FObjectInitiali
 	BackgroundColor = FLinearColor(0.1f, 0.1f, 0.1f);
 }
 
-void UPaperExtractSpritesSettings::PostEditChangeProperty(FPropertyChangedEvent& PropertyChangedEvent)
-{
-	const FName PropertyName = (PropertyChangedEvent.Property != nullptr) ? PropertyChangedEvent.Property->GetFName() : NAME_None;
-	if (PropertyName == GET_MEMBER_NAME_CHECKED(UPaperExtractSpritesSettings, SpriteExtractMode))
-	{
-		// Update bools
-		bAutoMode = (SpriteExtractMode == ESpriteExtractMode::Auto);
-		bGridMode = (SpriteExtractMode == ESpriteExtractMode::Grid);
-	}
-}
-
-
 UPaperExtractSpriteGridSettings::UPaperExtractSpriteGridSettings(const FObjectInitializer& ObjectInitializer)
 	: Super(ObjectInitializer)
 {
diff --git a/Engine/Plugins/2D/Paper2D/Source/Paper2DEditor/Private/ExtractSprites/PaperExtractSpritesSettings.h b/Engine/Plugins/2D/Paper2D/Source/Paper2DEditor/Private/ExtractSprites/PaperExtractSpritesSettings.h
--- a/Engine/Plugins/2D/Paper2D/Source/Paper2DEditor/Private/ExtractSprites/PaperExtractSpritesSettings.h
+++ b/Engine/Plugins/2D/Paper2D/Source/Paper2DEditor/Private/ExtractSprites/PaperExtractSpritesSettings.h
@@ -30,7 +30,7 @@ public:
 	ESpriteExtractMode SpriteExtractMode;
 
 	// The color of the sprite boundary outlines
-	UPROPERTY(Category = Settings, EditAnywhere)
+	UPROPERTY(Category = Settings, EditAnywhere, meta=(HideAlphaChannel))
 	FLinearColor OutlineColor;
 
 	// Apply a tint to the texture in the viewport to improve outline visibility in this editor
@@ -41,12 +41,6 @@ public:
 	UPROPERTY(Category = Settings, EditAnywhere, meta=(HideAlphaChannel))
 	FLinearColor BackgroundColor;
 
-	UPROPERTY()
-	bool bAutoMode;
-
-	UPROPERTY()
-	bool bGridMode;
-
 	// The name of the sprite that will be created. {0} will get replaced by the sprite number.
 	UPROPERTY(Category = Naming, EditAnywhere)
 	FString NamingTemplate;
@@ -56,10 +50,6 @@ public:
 	int32 NamingStartIndex;
 
 	UPaperExtractSpritesSettings(const FObjectInitializer& ObjectInitializer);
-
-	// UObject interface
-	virtual void PostEditChangeProperty(FPropertyChangedEvent& PropertyChangedEvent) override;
-	// End of UObject interface
 };
 
 UCLASS()
