Tile Maps: Added cut/copy/paste to the context menu for tile map layers [UE-15249]

diff --git a/Engine/Plugins/2D/Paper2D/Source/Paper2D/Classes/PaperTileMap.h b/Engine/Plugins/2D/Paper2D/Source/Paper2D/Classes/PaperTileMap.h
--- a/Engine/Plugins/2D/Paper2D/Source/Paper2D/Classes/PaperTileMap.h
+++ b/Engine/Plugins/2D/Paper2D/Source/Paper2D/Classes/PaperTileMap.h
@@ -7,6 +7,11 @@
 #include "PaperSprite.h"
 #include "PaperTileMap.generated.h"
 
+class UPaperTileSet;
+class UPaperTileLayer;
+class UBodySetup;
+class UAssetImportData;
+
 // The different kinds of projection modes supported
 UENUM()
 namespace ETileMapProjectionMode
@@ -67,7 +72,7 @@ class PAPER2D_API UPaperTileMap : public UObject
 
 	// Last tile set that was selected when editing the tile map
 	UPROPERTY()
-	TAssetPtr<class UPaperTileSet> SelectedTileSet;
+	TAssetPtr<UPaperTileSet> SelectedTileSet;
 
 	// The material to use on a tile map instance if not overridden
 	UPROPERTY(Category=Setup, EditAnywhere, BlueprintReadOnly)
@@ -75,7 +80,7 @@ class PAPER2D_API UPaperTileMap : public UObject
 
 	// The list of layers
 	UPROPERTY(Instanced, Category=Sprite, BlueprintReadOnly)
-	TArray<class UPaperTileLayer*> TileLayers;
+	TArray<UPaperTileLayer*> TileLayers;
 
 protected:
 	// The extrusion thickness of collision geometry when using a 3D collision domain
@@ -98,13 +103,13 @@ public:
 
 	// Baked physics data.
 	UPROPERTY()
-	class UBodySetup* BodySetup;
+	UBodySetup* BodySetup;
 
 public:
 #if WITH_EDITORONLY_DATA
 	/** Importing data and options used for this tile map */
 	UPROPERTY(Category=ImportSettings, VisibleAnywhere, Instanced)
-	class UAssetImportData* AssetImportData;
+	UAssetImportData* AssetImportData;
 
 	/** The currently selected layer index */
 	UPROPERTY()
@@ -163,11 +168,17 @@ public:
 	FBoxSphereBounds GetRenderBounds() const;
 
 	// Creates and adds a new layer and returns it
-	class UPaperTileLayer* AddNewLayer(int32 InsertionIndex = INDEX_NONE);
+	UPaperTileLayer* AddNewLayer(int32 InsertionIndex = INDEX_NONE);
+
+	// Handles adding an existing layer that does *not* belong to any existing tile map
+	void AddExistingLayer(UPaperTileLayer* NewLayer, int32 InsertionIndex = INDEX_NONE);
 
 	// Creates a reasonable new layer name
 	static FText GenerateNewLayerName(UPaperTileMap* TileMap);
 
+	// Returns true if the specified name is already in use as a layer name
+	bool IsLayerNameInUse(const FText& LayerName) const;
+
 	// Resize the tile map and all layers
 	void ResizeMap(int32 NewWidth, int32 NewHeight, bool bForceResize = true);
 
diff --git a/Engine/Plugins/2D/Paper2D/Source/Paper2D/Private/PaperTileMap.cpp b/Engine/Plugins/2D/Paper2D/Source/Paper2D/Private/PaperTileMap.cpp
--- a/Engine/Plugins/2D/Paper2D/Source/Paper2D/Private/PaperTileMap.cpp
+++ b/Engine/Plugins/2D/Paper2D/Source/Paper2D/Private/PaperTileMap.cpp
@@ -529,6 +529,37 @@ UPaperTileLayer* UPaperTileMap::AddNewLayer(int32 InsertionIndex)
 	return NewLayer;
 }
 
+void UPaperTileMap::AddExistingLayer(UPaperTileLayer* NewLayer, int32 InsertionIndex)
+{
+	NewLayer->SetFlags(RF_Transactional);
+	NewLayer->Modify();
+
+	// Make sure the layer has the correct outer
+	if (NewLayer->GetOuter() != this)
+	{
+		NewLayer->Rename(nullptr, this);
+	}
+
+	// And correct size
+	NewLayer->ResizeMap(MapWidth, MapHeight);
+
+	// And a unique name
+	if (IsLayerNameInUse(NewLayer->LayerName))
+	{
+		NewLayer->LayerName = GenerateNewLayerName(this);
+	}
+
+	// Insert the new layer
+	if (TileLayers.IsValidIndex(InsertionIndex))
+	{
+		TileLayers.Insert(NewLayer, InsertionIndex);
+	}
+	else
+	{
+		TileLayers.Add(NewLayer);
+	}
+}
+
 FText UPaperTileMap::GenerateNewLayerName(UPaperTileMap* TileMap)
 {
 	// Create a set of existing names
@@ -549,6 +580,19 @@ FText UPaperTileMap::GenerateNewLayerName(UPaperTileMap* TileMap)
 	return TestLayerName;
 }
 
+bool UPaperTileMap::IsLayerNameInUse(const FText& LayerName) const
+{
+	for (UPaperTileLayer* ExistingLayer : TileLayers)
+	{
+		if (ExistingLayer->LayerName.EqualToCaseIgnored(LayerName))
+		{
+			return true;
+		}
+	}
+
+	return false;
+}
+
 void UPaperTileMap::ResizeMap(int32 NewWidth, int32 NewHeight, bool bForceResize)
 {
 	if (bForceResize || (NewWidth != MapWidth) || (NewHeight != MapHeight))
diff --git a/Engine/Plugins/2D/Paper2D/Source/Paper2DEditor/Private/PaperStyle.cpp b/Engine/Plugins/2D/Paper2D/Source/Paper2DEditor/Private/PaperStyle.cpp
--- a/Engine/Plugins/2D/Paper2D/Source/Paper2DEditor/Private/PaperStyle.cpp
+++ b/Engine/Plugins/2D/Paper2D/Source/Paper2DEditor/Private/PaperStyle.cpp
@@ -20,6 +20,12 @@ FString FPaperStyle::InContent(const FString& RelativePath, const ANSICHAR* Exte
 TSharedPtr< FSlateStyleSet > FPaperStyle::StyleSet = nullptr;
 TSharedPtr< class ISlateStyle > FPaperStyle::Get() { return StyleSet; }
 
+FName FPaperStyle::GetStyleSetName()
+{
+	static FName PaperStyleName(TEXT("PaperStyle"));
+	return PaperStyleName;
+}
+
 void FPaperStyle::Initialize()
 {
 	// Const icon sizes
@@ -34,7 +40,7 @@ void FPaperStyle::Initialize()
 		return;
 	}
 
-	StyleSet = MakeShareable(new FSlateStyleSet("PaperStyle"));
+	StyleSet = MakeShareable(new FSlateStyleSet(GetStyleSetName()));
 	StyleSet->SetContentRoot(FPaths::EngineContentDir() / TEXT("Editor/Slate"));
 	StyleSet->SetCoreContentRoot(FPaths::EngineContentDir() / TEXT("Slate"));
 
diff --git a/Engine/Plugins/2D/Paper2D/Source/Paper2DEditor/Private/PaperStyle.h b/Engine/Plugins/2D/Paper2D/Source/Paper2DEditor/Private/PaperStyle.h
--- a/Engine/Plugins/2D/Paper2D/Source/Paper2DEditor/Private/PaperStyle.h
+++ b/Engine/Plugins/2D/Paper2D/Source/Paper2DEditor/Private/PaperStyle.h
@@ -11,6 +11,7 @@ public:
 
 	static TSharedPtr< class ISlateStyle > Get();
 
+	static FName GetStyleSetName();
 private:
 	static FString InContent(const FString& RelativePath, const ANSICHAR* Extension);
 
diff --git a/Engine/Plugins/2D/Paper2D/Source/Paper2DEditor/Private/TileMapEditing/STileLayerList.cpp b/Engine/Plugins/2D/Paper2D/Source/Paper2DEditor/Private/TileMapEditing/STileLayerList.cpp
--- a/Engine/Plugins/2D/Paper2D/Source/Paper2DEditor/Private/TileMapEditing/STileLayerList.cpp
+++ b/Engine/Plugins/2D/Paper2D/Source/Paper2DEditor/Private/TileMapEditing/STileLayerList.cpp
@@ -8,9 +8,41 @@
 #include "ScopedTransaction.h"
 
 #include "TileMapEditorCommands.h"
+#include "GenericCommands.h"
+
+#include "UnrealExporter.h"
+#include "Factories.h"
 
 #define LOCTEXT_NAMESPACE "Paper2D"
 
+//////////////////////////////////////////////////////////////////////////
+// FLayerTextFactory
+
+// Text object factory for pasting layers
+struct FLayerTextFactory : public FCustomizableTextObjectFactory
+{
+public:
+	TArray<UPaperTileLayer*> CreatedLayers;
+public:
+	FLayerTextFactory()
+		: FCustomizableTextObjectFactory(GWarn)
+	{
+	}
+
+	// FCustomizableTextObjectFactory interface
+	virtual bool CanCreateClass(UClass* ObjectClass) const override
+	{
+		// Only allow layers to be created
+		return ObjectClass->IsChildOf(UPaperTileLayer::StaticClass());
+	}
+
+	virtual void ProcessConstructedObject(UObject* NewObject) override
+	{
+		CreatedLayers.Add(CastChecked<UPaperTileLayer>(NewObject));
+	}
+	// End of FCustomizableTextObjectFactory interface
+};
+
 //////////////////////////////////////////////////////////////////////////
 // STileLayerList
 
@@ -20,75 +52,97 @@ void STileLayerList::Construct(const FArguments& InArgs, UPaperTileMap* InTileMa
 	NotifyHook = InNotifyHook;
 
 	FTileMapEditorCommands::Register();
-	const FTileMapEditorCommands& Commands = FTileMapEditorCommands::Get();
+	FGenericCommands::Register();
+	const FTileMapEditorCommands& TileMapCommands = FTileMapEditorCommands::Get();
+	const FGenericCommands& GenericCommands = FGenericCommands::Get();
 
 	CommandList = MakeShareable(new FUICommandList);
 	InParentCommandList->Append(CommandList.ToSharedRef());
 
 	CommandList->MapAction(
-		Commands.AddNewLayerAbove,
+		TileMapCommands.AddNewLayerAbove,
 		FExecuteAction::CreateSP(this, &STileLayerList::AddNewLayerAbove));
 
 	CommandList->MapAction(
-		Commands.AddNewLayerBelow,
+		TileMapCommands.AddNewLayerBelow,
 		FExecuteAction::CreateSP(this, &STileLayerList::AddNewLayerBelow));
 
 	CommandList->MapAction(
-		Commands.DeleteLayer,
-		FExecuteAction::CreateSP(this, &STileLayerList::DeleteLayer),
+		GenericCommands.Cut,
+		FExecuteAction::CreateSP(this, &STileLayerList::CutLayer),
 		FCanExecuteAction::CreateSP(this, &STileLayerList::CanExecuteActionNeedingSelectedLayer));
 
 	CommandList->MapAction(
-		Commands.RenameLayer,
-		FExecuteAction::CreateSP(this, &STileLayerList::RenameLayer),
+		GenericCommands.Copy,
+		FExecuteAction::CreateSP(this, &STileLayerList::CopyLayer),
 		FCanExecuteAction::CreateSP(this, &STileLayerList::CanExecuteActionNeedingSelectedLayer));
 
 	CommandList->MapAction(
-		Commands.DuplicateLayer,
+		GenericCommands.Paste,
+		FExecuteAction::CreateSP(this, &STileLayerList::PasteLayerAbove),
+		FCanExecuteAction::CreateSP(this, &STileLayerList::CanPasteLayer));
+
+	CommandList->MapAction(
+		GenericCommands.Duplicate,
 		FExecuteAction::CreateSP(this, &STileLayerList::DuplicateLayer),
 		FCanExecuteAction::CreateSP(this, &STileLayerList::CanExecuteActionNeedingSelectedLayer));
 
 	CommandList->MapAction(
-		Commands.MergeLayerDown,
+		GenericCommands.Delete,
+		FExecuteAction::CreateSP(this, &STileLayerList::DeleteLayer),
+		FCanExecuteAction::CreateSP(this, &STileLayerList::CanExecuteActionNeedingSelectedLayer));
+
+	CommandList->MapAction(
+		GenericCommands.Rename,
+		FExecuteAction::CreateSP(this, &STileLayerList::RenameLayer),
+		FCanExecuteAction::CreateSP(this, &STileLayerList::CanExecuteActionNeedingSelectedLayer));
+
+	CommandList->MapAction(
+		TileMapCommands.MergeLayerDown,
 		FExecuteAction::CreateSP(this, &STileLayerList::MergeLayerDown),
 		FCanExecuteAction::CreateSP(this, &STileLayerList::CanExecuteActionNeedingLayerBelow));
 
 	CommandList->MapAction(
-		Commands.MoveLayerUp,
+		TileMapCommands.MoveLayerUp,
 		FExecuteAction::CreateSP(this, &STileLayerList::MoveLayerUp, /*bForceToTop=*/ false),
 		FCanExecuteAction::CreateSP(this, &STileLayerList::CanExecuteActionNeedingLayerAbove));
 
 	CommandList->MapAction(
-		Commands.MoveLayerDown,
+		TileMapCommands.MoveLayerDown,
 		FExecuteAction::CreateSP(this, &STileLayerList::MoveLayerDown, /*bForceToBottom=*/ false),
 		FCanExecuteAction::CreateSP(this, &STileLayerList::CanExecuteActionNeedingLayerBelow));
 
 	CommandList->MapAction(
-		Commands.MoveLayerToTop,
+		TileMapCommands.MoveLayerToTop,
 		FExecuteAction::CreateSP(this, &STileLayerList::MoveLayerUp, /*bForceToTop=*/ true),
 		FCanExecuteAction::CreateSP(this, &STileLayerList::CanExecuteActionNeedingLayerAbove));
 
 	CommandList->MapAction(
-		Commands.MoveLayerToBottom,
+		TileMapCommands.MoveLayerToBottom,
 		FExecuteAction::CreateSP(this, &STileLayerList::MoveLayerDown, /*bForceToBottom=*/ true),
 		FCanExecuteAction::CreateSP(this, &STileLayerList::CanExecuteActionNeedingLayerBelow));
 
 	CommandList->MapAction(
-		Commands.SelectLayerAbove,
+		TileMapCommands.SelectLayerAbove,
 		FExecuteAction::CreateSP(this, &STileLayerList::SelectLayerAbove, /*bTopmost=*/ false));
 
 	CommandList->MapAction(
-		Commands.SelectLayerBelow,
+		TileMapCommands.SelectLayerBelow,
 		FExecuteAction::CreateSP(this, &STileLayerList::SelectLayerBelow, /*bBottommost=*/ false));
 
+	//
 	FToolBarBuilder ToolbarBuilder(CommandList, FMultiBoxCustomization("TileLayerBrowserToolbar"), TSharedPtr<FExtender>(), Orient_Horizontal, /*InForceSmallIcons=*/ true);
 	ToolbarBuilder.SetLabelVisibility(EVisibility::Collapsed);
 
-	ToolbarBuilder.AddToolBarButton(Commands.AddNewLayerAbove);
-	ToolbarBuilder.AddToolBarButton(Commands.MoveLayerUp);
-	ToolbarBuilder.AddToolBarButton(Commands.MoveLayerDown);
-	ToolbarBuilder.AddToolBarButton(Commands.DuplicateLayer);
-	ToolbarBuilder.AddToolBarButton(Commands.DeleteLayer);
+	ToolbarBuilder.AddToolBarButton(TileMapCommands.AddNewLayerAbove);
+	ToolbarBuilder.AddToolBarButton(TileMapCommands.MoveLayerUp);
+	ToolbarBuilder.AddToolBarButton(TileMapCommands.MoveLayerDown);
+
+	FSlateIcon DuplicateIcon(FPaperStyle::GetStyleSetName(), "TileMapEditor.DuplicateLayer");
+	ToolbarBuilder.AddToolBarButton(GenericCommands.Duplicate, NAME_None, TAttribute<FText>(), TAttribute<FText>(), DuplicateIcon);
+
+	FSlateIcon DeleteIcon(FPaperStyle::GetStyleSetName(), "TileMapEditor.DeleteLayer");
+	ToolbarBuilder.AddToolBarButton(GenericCommands.Delete, NAME_None, TAttribute<FText>(), TAttribute<FText>(), DeleteIcon);
 
 	TSharedRef<SWidget> Toolbar = ToolbarBuilder.MakeWidget();
 
@@ -165,10 +219,23 @@ UPaperTileLayer* STileLayerList::GetSelectedLayer() const
 FText STileLayerList::GenerateDuplicatedLayerName(const FString& InputNameRaw, UPaperTileMap* TileMap)
 {
 	// Create a set of existing names
+	bool bFoundName = false;
 	TSet<FString> ExistingNames;
 	for (UPaperTileLayer* ExistingLayer : TileMap->TileLayers)
 	{
-		ExistingNames.Add(ExistingLayer->LayerName.ToString());
+		const FString LayerName = ExistingLayer->LayerName.ToString();
+		ExistingNames.Add(LayerName);
+
+		if (LayerName == InputNameRaw)
+		{
+			bFoundName = true;
+		}
+	}
+
+	// If the name doesn't already exist, then we're done (can happen when pasting a cut layer, etc...)
+	if (!bFoundName)
+	{
+		return FText::FromString(InputNameRaw);
 	}
 
 	FString BaseName = InputNameRaw;
@@ -291,17 +358,13 @@ int32 STileLayerList::GetSelectionIndex() const
 	return SelectionIndex;
 }
 
-void STileLayerList::DeleteLayer()
+void STileLayerList::DeleteSelectedLayerWithNoTransaction()
 {
 	if (UPaperTileMap* TileMap = TileMapPtr.Get())
 	{
 		const int32 DeleteIndex = GetSelectionIndex();
 		if (DeleteIndex != INDEX_NONE)
 		{
-			const FScopedTransaction Transaction(LOCTEXT("TileMapDeleteLayer", "Delete Layer"));
-			TileMap->SetFlags(RF_Transactional);
-			TileMap->Modify();
-
 			TileMap->TileLayers.RemoveAt(DeleteIndex);
 
 			PostEditNotfications();
@@ -316,6 +379,18 @@ void STileLayerList::DeleteLayer()
 	}
 }
 
+void STileLayerList::DeleteLayer()
+{
+	if (UPaperTileMap* TileMap = TileMapPtr.Get())
+	{
+		const FScopedTransaction Transaction(LOCTEXT("TileMapDeleteLayer", "Delete Layer"));
+		TileMap->SetFlags(RF_Transactional);
+		TileMap->Modify();
+
+		DeleteSelectedLayerWithNoTransaction();
+	}
+}
+
 void STileLayerList::RenameLayer()
 {
 	if (UPaperTileMap* TileMap = TileMapPtr.Get())
@@ -431,6 +506,70 @@ void STileLayerList::SelectLayerBelow(bool bBottommost)
 	SetSelectedLayerIndex(NewIndex);
 }
 
+void STileLayerList::CutLayer()
+{
+	CopyLayer();
+
+	if (UPaperTileMap* TileMap = TileMapPtr.Get())
+	{
+		const FScopedTransaction Transaction(LOCTEXT("TileMapCutLayer", "Cut Layer"));
+		TileMap->SetFlags(RF_Transactional);
+		TileMap->Modify();
+
+		DeleteSelectedLayerWithNoTransaction();
+	}
+}
+
+void STileLayerList::CopyLayer()
+{
+	if (UPaperTileLayer* SelectedLayer = GetSelectedLayer())
+	{
+		UnMarkAllObjects(EObjectMark(OBJECTMARK_TagExp | OBJECTMARK_TagImp));
+		FStringOutputDevice ExportArchive;
+		const FExportObjectInnerContext Context;
+		UExporter::ExportToOutputDevice(&Context, SelectedLayer, nullptr, ExportArchive, TEXT("copy"), 0, PPF_Copy, false, nullptr);
+
+		FPlatformMisc::ClipboardCopy(*ExportArchive);
+	}
+}
+
+void STileLayerList::PasteLayerAbove()
+{
+	if (UPaperTileMap* TileMap = TileMapPtr.Get())
+	{
+		FString ClipboardContent;
+		FPlatformMisc::ClipboardPaste(ClipboardContent);
+
+		if (!ClipboardContent.IsEmpty())
+		{
+			const FScopedTransaction Transaction(LOCTEXT("TileMapPasteLayer", "Paste Layer"));
+			TileMap->SetFlags(RF_Transactional);
+			TileMap->Modify();
+
+			// Turn the text buffer into objects
+			FLayerTextFactory Factory;
+			Factory.ProcessBuffer(TileMap, RF_Transactional, ClipboardContent);
+
+			// Add them to the map and select them (there will currently only ever be 0 or 1)
+			for (UPaperTileLayer* NewLayer : Factory.CreatedLayers)
+			{
+				NewLayer->LayerName = GenerateDuplicatedLayerName(NewLayer->LayerName.ToString(), TileMap);
+				TileMap->AddExistingLayer(NewLayer, GetSelectionIndex());
+				PostEditNotfications();
+				SetSelectedLayer(NewLayer);
+			}
+		}
+	}
+}
+
+bool STileLayerList::CanPasteLayer() const
+{
+	FString ClipboardContent;
+	FPlatformMisc::ClipboardPaste(ClipboardContent);
+
+	return !ClipboardContent.IsEmpty();
+}
+
 void STileLayerList::SetSelectedLayerIndex(int32 NewIndex)
 {
 	if (UPaperTileMap* TileMap = TileMapPtr.Get())
@@ -498,28 +637,32 @@ TSharedPtr<SWidget> STileLayerList::OnConstructContextMenu()
 	const bool bShouldCloseWindowAfterMenuSelection = true;
 	FMenuBuilder MenuBuilder(bShouldCloseWindowAfterMenuSelection, CommandList);
 
-	const FTileMapEditorCommands& Commands = FTileMapEditorCommands::Get();
+	const FTileMapEditorCommands& TileMapCommands = FTileMapEditorCommands::Get();
+	const FGenericCommands& GenericCommands = FGenericCommands::Get();
 
 	FSlateIcon DummyIcon(NAME_None, NAME_None);
 
 	MenuBuilder.BeginSection("BasicOperations", LOCTEXT("BasicOperationsHeader", "Layer actions"));
  	{
-		MenuBuilder.AddMenuEntry(Commands.DuplicateLayer, NAME_None, TAttribute<FText>(), TAttribute<FText>(), DummyIcon);
-		MenuBuilder.AddMenuEntry(Commands.DeleteLayer, NAME_None, TAttribute<FText>(), TAttribute<FText>(), DummyIcon);
-		MenuBuilder.AddMenuEntry(Commands.RenameLayer, NAME_None, TAttribute<FText>(), TAttribute<FText>(), DummyIcon);
-		MenuBuilder.AddMenuEntry(Commands.MergeLayerDown, NAME_None, TAttribute<FText>(), TAttribute<FText>(), DummyIcon);
+		MenuBuilder.AddMenuEntry(GenericCommands.Cut, NAME_None, TAttribute<FText>(), TAttribute<FText>(), DummyIcon);
+		MenuBuilder.AddMenuEntry(GenericCommands.Copy, NAME_None, TAttribute<FText>(), TAttribute<FText>(), DummyIcon);
+		MenuBuilder.AddMenuEntry(GenericCommands.Paste, NAME_None, TAttribute<FText>(), TAttribute<FText>(), DummyIcon);
+		MenuBuilder.AddMenuEntry(GenericCommands.Duplicate, NAME_None, TAttribute<FText>(), TAttribute<FText>(), DummyIcon);
+		MenuBuilder.AddMenuEntry(GenericCommands.Delete, NAME_None, TAttribute<FText>(), TAttribute<FText>(), DummyIcon);
+		MenuBuilder.AddMenuEntry(GenericCommands.Rename, NAME_None, TAttribute<FText>(), TAttribute<FText>(), DummyIcon);
+		MenuBuilder.AddMenuEntry(TileMapCommands.MergeLayerDown, NAME_None, TAttribute<FText>(), TAttribute<FText>(), DummyIcon);
 		MenuBuilder.AddMenuSeparator();
-		MenuBuilder.AddMenuEntry(Commands.SelectLayerAbove, NAME_None, TAttribute<FText>(), TAttribute<FText>(), DummyIcon);
-		MenuBuilder.AddMenuEntry(Commands.SelectLayerBelow, NAME_None, TAttribute<FText>(), TAttribute<FText>(), DummyIcon);
+		MenuBuilder.AddMenuEntry(TileMapCommands.SelectLayerAbove, NAME_None, TAttribute<FText>(), TAttribute<FText>(), DummyIcon);
+		MenuBuilder.AddMenuEntry(TileMapCommands.SelectLayerBelow, NAME_None, TAttribute<FText>(), TAttribute<FText>(), DummyIcon);
 	}
 	MenuBuilder.EndSection();
 
 	MenuBuilder.BeginSection("OrderingOperations", LOCTEXT("OrderingOperationsHeader", "Order actions"));
 	{
-		MenuBuilder.AddMenuEntry(Commands.MoveLayerToTop, NAME_None, TAttribute<FText>(), TAttribute<FText>(), DummyIcon);
-		MenuBuilder.AddMenuEntry(Commands.MoveLayerUp, NAME_None, TAttribute<FText>(), TAttribute<FText>(), DummyIcon);
-		MenuBuilder.AddMenuEntry(Commands.MoveLayerDown, NAME_None, TAttribute<FText>(), TAttribute<FText>(), DummyIcon);
-		MenuBuilder.AddMenuEntry(Commands.MoveLayerToBottom, NAME_None, TAttribute<FText>(), TAttribute<FText>(), DummyIcon);
+		MenuBuilder.AddMenuEntry(TileMapCommands.MoveLayerToTop, NAME_None, TAttribute<FText>(), TAttribute<FText>(), DummyIcon);
+		MenuBuilder.AddMenuEntry(TileMapCommands.MoveLayerUp, NAME_None, TAttribute<FText>(), TAttribute<FText>(), DummyIcon);
+		MenuBuilder.AddMenuEntry(TileMapCommands.MoveLayerDown, NAME_None, TAttribute<FText>(), TAttribute<FText>(), DummyIcon);
+		MenuBuilder.AddMenuEntry(TileMapCommands.MoveLayerToBottom, NAME_None, TAttribute<FText>(), TAttribute<FText>(), DummyIcon);
 	}
 	MenuBuilder.EndSection();
 
diff --git a/Engine/Plugins/2D/Paper2D/Source/Paper2DEditor/Private/TileMapEditing/STileLayerList.h b/Engine/Plugins/2D/Paper2D/Source/Paper2DEditor/Private/TileMapEditing/STileLayerList.h
--- a/Engine/Plugins/2D/Paper2D/Source/Paper2DEditor/Private/TileMapEditing/STileLayerList.h
+++ b/Engine/Plugins/2D/Paper2D/Source/Paper2DEditor/Private/TileMapEditing/STileLayerList.h
@@ -47,15 +47,25 @@ protected:
 
 	void AddNewLayerAbove();
 	void AddNewLayerBelow();
+
+	void CutLayer();
+	void CopyLayer();
+	void PasteLayerAbove();
+	bool CanPasteLayer() const;
+
+	void DuplicateLayer();
 	void DeleteLayer();
 	void RenameLayer();
-	void DuplicateLayer();
+
+	void DeleteSelectedLayerWithNoTransaction();
+
 	void MergeLayerDown();
 	void MoveLayerUp(bool bForceToTop);
 	void MoveLayerDown(bool bForceToBottom);
 	void SelectLayerAbove(bool bTopmost);
 	void SelectLayerBelow(bool bBottommost);
 
+
 	void SetSelectedLayerIndex(int32 NewIndex);
 
 	int32 GetNumLayers() const;
diff --git a/Engine/Plugins/2D/Paper2D/Source/Paper2DEditor/Private/TileMapEditing/TileMapEditorCommands.cpp b/Engine/Plugins/2D/Paper2D/Source/Paper2DEditor/Private/TileMapEditing/TileMapEditorCommands.cpp
--- a/Engine/Plugins/2D/Paper2D/Source/Paper2DEditor/Private/TileMapEditing/TileMapEditorCommands.cpp
+++ b/Engine/Plugins/2D/Paper2D/Source/Paper2DEditor/Private/TileMapEditing/TileMapEditorCommands.cpp
@@ -37,8 +37,6 @@ void FTileMapEditorCommands::RegisterCommands()
 	// Layer actions
 	UI_COMMAND(AddNewLayerAbove, "Add new layer", "Add a new layer above the current layer", EUserInterfaceActionType::Button, FInputChord(EKeys::N, EModifierKey::Control | EModifierKey::Shift));
 	UI_COMMAND(AddNewLayerBelow, "Add new layer below", "Add a new layer below the current layer", EUserInterfaceActionType::Button, FInputChord());
-	UI_COMMAND(DeleteLayer, "Delete layer", "Delete the current layer", EUserInterfaceActionType::Button, FInputChord());
-	UI_COMMAND(DuplicateLayer, "Duplicate layer", "Duplicate the current layer", EUserInterfaceActionType::Button, FInputChord());
 	UI_COMMAND(MergeLayerDown, "Merge layer down", "Merge the current layer down to the layer below it", EUserInterfaceActionType::Button, FInputChord());
 	UI_COMMAND(MoveLayerUp, "Bring forward", "Bring forward", EUserInterfaceActionType::Button, FInputChord(EKeys::RightBracket, EModifierKey::Control));
 	UI_COMMAND(MoveLayerDown, "Send backward", "Send backward", EUserInterfaceActionType::Button, FInputChord(EKeys::LeftBracket, EModifierKey::Control));
@@ -46,7 +44,6 @@ void FTileMapEditorCommands::RegisterCommands()
 	UI_COMMAND(MoveLayerToBottom, "Send to back", "Send to back", EUserInterfaceActionType::Button, FInputChord(EKeys::RightBracket, EModifierKey::Shift | EModifierKey::Control));
 	UI_COMMAND(SelectLayerAbove, "Select next layer", "Changes the active layer to the next layer", EUserInterfaceActionType::Button, FInputChord(EKeys::RightBracket, EModifierKey::Alt));
 	UI_COMMAND(SelectLayerBelow, "Select previous layer", "Changes the active layer to the previous layer", EUserInterfaceActionType::Button, FInputChord(EKeys::LeftBracket, EModifierKey::Alt));
-	UI_COMMAND(RenameLayer, "Rename layer", "Renames the selected layer", EUserInterfaceActionType::Button, FInputChord(EKeys::F2));
 }
 
 #undef LOCTEXT_NAMESPACE
\ No newline at end of file
diff --git a/Engine/Plugins/2D/Paper2D/Source/Paper2DEditor/Private/TileMapEditing/TileMapEditorCommands.h b/Engine/Plugins/2D/Paper2D/Source/Paper2DEditor/Private/TileMapEditing/TileMapEditorCommands.h
--- a/Engine/Plugins/2D/Paper2D/Source/Paper2DEditor/Private/TileMapEditing/TileMapEditorCommands.h
+++ b/Engine/Plugins/2D/Paper2D/Source/Paper2DEditor/Private/TileMapEditing/TileMapEditorCommands.h
@@ -51,7 +51,6 @@ public:
 	TSharedPtr<FUICommandInfo> AddNewLayerAbove;
 	TSharedPtr<FUICommandInfo> AddNewLayerBelow;
 	TSharedPtr<FUICommandInfo> DeleteLayer;
-	TSharedPtr<FUICommandInfo> DuplicateLayer;
 	TSharedPtr<FUICommandInfo> MergeLayerDown;
 	TSharedPtr<FUICommandInfo> MoveLayerUp;
 	TSharedPtr<FUICommandInfo> MoveLayerDown;
@@ -59,5 +58,4 @@ public:
 	TSharedPtr<FUICommandInfo> MoveLayerToBottom;
 	TSharedPtr<FUICommandInfo> SelectLayerAbove;
 	TSharedPtr<FUICommandInfo> SelectLayerBelow;
-	TSharedPtr<FUICommandInfo> RenameLayer;
 };
\ No newline at end of file
