Tile Map Editor: Added search box to details panel in the standalone tile map editor

diff --git a/Engine/Plugins/2D/Paper2D/Source/Paper2DEditor/Private/TileMapEditing/TileMapEditor.cpp b/Engine/Plugins/2D/Paper2D/Source/Paper2DEditor/Private/TileMapEditing/TileMapEditor.cpp
--- a/Engine/Plugins/2D/Paper2D/Source/Paper2DEditor/Private/TileMapEditing/TileMapEditor.cpp
+++ b/Engine/Plugins/2D/Paper2D/Source/Paper2DEditor/Private/TileMapEditing/TileMapEditor.cpp
@@ -182,7 +182,7 @@ public:
 	{
 		TileMapEditorPtr = InTileMapEditor;
 
-		SSingleObjectDetailsPanel::Construct(SSingleObjectDetailsPanel::FArguments().HostCommandList(InTileMapEditor->GetToolkitCommands()));
+		SSingleObjectDetailsPanel::Construct(SSingleObjectDetailsPanel::FArguments().HostCommandList(InTileMapEditor->GetToolkitCommands()), /*bAutoObserve=*/ true, /*bAllowSearch=*/ true);
 	}
 
 	// SSingleObjectDetailsPanel interface
