Tile Maps: Made most variables in UPaperTileLayer private and fixed an issue when merging non-square layers

diff --git a/Engine/Plugins/2D/Paper2D/Source/Paper2D/Classes/PaperTileLayer.h b/Engine/Plugins/2D/Paper2D/Source/Paper2D/Classes/PaperTileLayer.h
--- a/Engine/Plugins/2D/Paper2D/Source/Paper2D/Classes/PaperTileLayer.h
+++ b/Engine/Plugins/2D/Paper2D/Source/Paper2D/Classes/PaperTileLayer.h
@@ -104,34 +104,30 @@ class PAPER2D_API UPaperTileLayer : public UObject
 	GENERATED_UCLASS_BODY()
 
 	// Name of the layer
-	UPROPERTY()
+	UPROPERTY(BlueprintReadOnly, Category=Sprite)
 	FText LayerName;
 
+private:
 	// Width of the layer (in tiles)
-	UPROPERTY(BlueprintReadOnly, Category=Sprite)
+	UPROPERTY(BlueprintReadOnly, Category=Sprite, meta=(AllowPrivateAccess="true"))
 	int32 LayerWidth;
 
 	// Height of the layer (in tiles)
-	UPROPERTY(BlueprintReadOnly, Category=Sprite)
+	UPROPERTY(BlueprintReadOnly, Category=Sprite, meta=(AllowPrivateAccess="true"))
 	int32 LayerHeight;
 
 #if WITH_EDITORONLY_DATA
-	// Opacity of the layer (editor only)
-	UPROPERTY()
-	float LayerOpacity;
-
 	// Is this layer currently hidden in the editor?
-	UPROPERTY()
-	bool bHiddenInEditor;
+	UPROPERTY(EditAnywhere, Category=Sprite, meta=(AllowPrivateAccess="true"))
+	uint32 bHiddenInEditor:1;
 #endif
 
 	// Should this layer be hidden in the game?
-	UPROPERTY(BlueprintReadOnly, Category=Sprite)
-	bool bHiddenInGame;
+	UPROPERTY(BlueprintReadOnly, EditAnywhere, Category=Sprite, meta=(AllowPrivateAccess="true"))
+	uint32 bHiddenInGame:1;
 
-protected:
 	// The color of this layer (multiplied with the tile map color and passed to the material as a vertex color)
-	UPROPERTY(EditAnywhere, Category=Sprite)
+	UPROPERTY(BlueprintReadOnly, EditAnywhere, Category=Sprite, meta=(AllowPrivateAccess="true"))
 	FLinearColor LayerColor;
 
 	// The allocated width of the tile data (used to handle resizing without data loss)
@@ -146,7 +142,6 @@ protected:
 	UPROPERTY()
 	TArray<FPaperTileInfo> AllocatedCells;
 
-private:
 	UPROPERTY()
 	UPaperTileSet* TileSet_DEPRECATED;
 
@@ -168,6 +163,9 @@ public:
 	// Returns the index of this layer in the parent tile map
 	int32 GetLayerIndex() const;
 
+	// Returns whether the specified coordinates are in bounds for the layer
+	bool InBounds(int32 X, int32 Y) const;
+
 	// Returns the tile information about the specified cell
 	FPaperTileInfo GetCell(int32 X, int32 Y) const;
 
@@ -193,6 +191,23 @@ public:
 	// Note: This is a slow operation, it scans each tile!
 	bool UsesTileSet(UPaperTileSet* TileSet) const;
 
+#if WITH_EDITORONLY_DATA
+	// Should this layer be drawn (in the editor)?
+	bool ShouldRenderInEditor() const { return !bHiddenInEditor; }
+
+	// Set whether this layer should be drawn (in the editor)
+	void SetShouldRenderInEditor(bool bShouldRender) { bHiddenInEditor = !bShouldRender; }
+#endif
+
+	// Should this layer be drawn (in game)?
+	bool ShouldRenderInGame() const { return !bHiddenInGame; }
+
+	// Returns the width of the layer (in tiles)
+	int32 GetLayerWidth() const { return LayerWidth; }
+
+	// Returns the height of the layer (in tiles)
+	int32 GetLayerHeight() const { return LayerHeight; }
+
 protected:
 	void ReallocateAndCopyMap();
 };
diff --git a/Engine/Plugins/2D/Paper2D/Source/Paper2D/Private/PaperRenderSceneProxy.h b/Engine/Plugins/2D/Paper2D/Source/Paper2D/Private/PaperRenderSceneProxy.h
--- a/Engine/Plugins/2D/Paper2D/Source/Paper2D/Private/PaperRenderSceneProxy.h
+++ b/Engine/Plugins/2D/Paper2D/Source/Paper2D/Private/PaperRenderSceneProxy.h
@@ -71,27 +71,29 @@ struct PAPER2D_API FSpriteRenderSection
 		return (NumVertices > 0) && (GetBaseTextureResource() != nullptr);
 	}
 
-#if 0
+#if 1
 	template <typename SourceArrayType>
 	void AddTriangles(const FSpriteDrawCallRecord& Record, SourceArrayType& Vertices)
 	{
 		if (NumVertices == 0)
 		{
 			VertexOffset = Vertices.Num();
-			Texture = Record.Texture;
+			BaseTexture = Record.BaseTexture;
+			AdditionalTextures = Record.AdditionalTextures;
 		}
 		else
 		{
 			checkSlow((VertexOffset + NumVertices) == Vertices.Num());
+			checkSlow(BaseTexture == Record.BaseTexture);
+			// Note: Not checking AdditionalTextures for now, since checking BaseTexture should catch most bugs
 		}
 
-		NumVertices += Record.RenderVerts.Num();
-		Vertices.Reserve(NumVertices);
+		const int32 NumNewVerts = Record.RenderVerts.Num();
+		NumVertices += NumNewVerts;
+		Vertices.Reserve(Vertices.Num() + NumNewVerts);
 
-		for (int32 SVI = 0; SVI < Record.RenderVerts.Num(); ++SVI)
+		for (const FVector4& SourceVert : Record.RenderVerts)
 		{
-			const FVector4& SourceVert = Record.RenderVerts[SVI];
-
 			const FVector Pos((PaperAxisX * SourceVert.X) + (PaperAxisY * SourceVert.Y) + Record.Destination);
 			const FVector2D UV(SourceVert.Z, SourceVert.W);
 
diff --git a/Engine/Plugins/2D/Paper2D/Source/Paper2D/Private/PaperTileLayer.cpp b/Engine/Plugins/2D/Paper2D/Source/Paper2D/Private/PaperTileLayer.cpp
--- a/Engine/Plugins/2D/Paper2D/Source/Paper2D/Private/PaperTileLayer.cpp
+++ b/Engine/Plugins/2D/Paper2D/Source/Paper2D/Private/PaperTileLayer.cpp
@@ -137,16 +137,14 @@ protected:
 
 UPaperTileLayer::UPaperTileLayer(const FObjectInitializer& ObjectInitializer)
 	: Super(ObjectInitializer)
-	, LayerColor(FLinearColor::White)
-{
-	LayerWidth = 4;
-	LayerHeight = 4;
-
+	, LayerWidth(4)
+	, LayerHeight(4)
 #if WITH_EDITORONLY_DATA
-	LayerOpacity = 1.0f;
-
-	bHiddenInEditor = false;
+	, bHiddenInEditor(false)
 #endif
+	, bHiddenInGame(false)
+	, LayerColor(FLinearColor::White)
+{
 
 	DestructiveAllocateMap(LayerWidth, LayerHeight);
 }
@@ -154,6 +152,8 @@ UPaperTileLayer::UPaperTileLayer(const FObjectInitializer& ObjectInitializer)
 void UPaperTileLayer::DestructiveAllocateMap(int32 NewWidth, int32 NewHeight)
 {
 	check((NewWidth > 0) && (NewHeight > 0));
+	LayerWidth = NewWidth;
+	LayerHeight = NewHeight;
 
 	const int32 NumCells = NewWidth * NewHeight;
 	AllocatedCells.Empty(NumCells);
@@ -232,24 +232,19 @@ int32 UPaperTileLayer::GetLayerIndex() const
 	return GetTileMap()->TileLayers.Find(const_cast<UPaperTileLayer*>(this));
 }
 
+bool UPaperTileLayer::InBounds(int32 X, int32 Y) const
+{
+	return (X >= 0) && (X < LayerWidth) && (Y >= 0) && (Y < LayerHeight);
+}
+
 FPaperTileInfo UPaperTileLayer::GetCell(int32 X, int32 Y) const
 {
-	if ((X < 0) || (X >= LayerWidth) || (Y < 0) || (Y >= LayerHeight))
-	{
-		return FPaperTileInfo();
-	}
-	else
-	{
-		return AllocatedCells[X + (Y*LayerWidth)];
-	}
+	return InBounds(X, Y) ? AllocatedCells[X + (Y*LayerWidth)] : FPaperTileInfo();
 }
 
 void UPaperTileLayer::SetCell(int32 X, int32 Y, const FPaperTileInfo& NewValue)
 {
-	if ((X < 0) || (X >= LayerWidth) || (Y < 0) || (Y >= LayerHeight))
-	{
-	}
-	else
+	if (InBounds(X, Y))
 	{
 		AllocatedCells[X + (Y*LayerWidth)] = NewValue;
 	}
diff --git a/Engine/Plugins/2D/Paper2D/Source/Paper2D/Private/PaperTileMap.cpp b/Engine/Plugins/2D/Paper2D/Source/Paper2D/Private/PaperTileMap.cpp
--- a/Engine/Plugins/2D/Paper2D/Source/Paper2D/Private/PaperTileMap.cpp
+++ b/Engine/Plugins/2D/Paper2D/Source/Paper2D/Private/PaperTileMap.cpp
@@ -160,7 +160,7 @@ void UPaperTileMap::ValidateSelectedLayerIndex()
 		SelectedLayerIndex = INDEX_NONE;
 		for (int32 LayerIndex = 0; (LayerIndex < TileLayers.Num()) && (SelectedLayerIndex == INDEX_NONE); ++LayerIndex)
 		{
-			if (!TileLayers[LayerIndex]->bHiddenInEditor)
+			if (TileLayers[LayerIndex]->ShouldRenderInEditor())
 			{
 				SelectedLayerIndex = LayerIndex;
 			}
@@ -511,9 +511,7 @@ UPaperTileLayer* UPaperTileMap::AddNewLayer(int32 InsertionIndex)
 	UPaperTileLayer* NewLayer = NewObject<UPaperTileLayer>(this);
 	NewLayer->SetFlags(RF_Transactional);
 
-	NewLayer->LayerWidth = MapWidth;
-	NewLayer->LayerHeight = MapHeight;
-	NewLayer->DestructiveAllocateMap(NewLayer->LayerWidth, NewLayer->LayerHeight);
+	NewLayer->DestructiveAllocateMap(MapWidth, MapHeight);
 	NewLayer->LayerName = GenerateNewLayerName(this);
 
 	// Insert the new layer
diff --git a/Engine/Plugins/2D/Paper2D/Source/Paper2D/Private/PaperTileMapComponent.cpp b/Engine/Plugins/2D/Paper2D/Source/Paper2D/Private/PaperTileMapComponent.cpp
--- a/Engine/Plugins/2D/Paper2D/Source/Paper2D/Private/PaperTileMapComponent.cpp
+++ b/Engine/Plugins/2D/Paper2D/Source/Paper2D/Private/PaperTileMapComponent.cpp
@@ -238,13 +238,12 @@ void UPaperTileMapComponent::RebuildRenderData(FPaperTileMapRenderSceneProxy* Pr
 		}
 
 		FLinearColor DrawColor = TileMapColor * Layer->GetLayerColor();
+
 #if WITH_EDITORONLY_DATA
-		if (Layer->bHiddenInEditor)
+		if (!Layer->ShouldRenderInEditor())
 		{
 			continue;
 		}
-
-		DrawColor.A *= Layer->LayerOpacity;
 #endif
 
 		FSpriteDrawCallRecord* CurrentBatch = nullptr;
diff --git a/Engine/Plugins/2D/Paper2D/Source/Paper2DEditor/Private/TileMapEditing/EdModeTileMap.cpp b/Engine/Plugins/2D/Paper2D/Source/Paper2DEditor/Private/TileMapEditing/EdModeTileMap.cpp
--- a/Engine/Plugins/2D/Paper2D/Source/Paper2DEditor/Private/TileMapEditing/EdModeTileMap.cpp
+++ b/Engine/Plugins/2D/Paper2D/Source/Paper2DEditor/Private/TileMapEditing/EdModeTileMap.cpp
@@ -77,7 +77,7 @@ struct FHorizontalSpan
 	// Indexes a bit in the reachability array
 	static FBitReference Reach(UPaperTileLayer* Layer, TBitArray<>& Reachability, int32 X, int32 Y)
 	{
-		const int32 Index = (Layer->LayerWidth * Y) + X;
+		const int32 Index = (Layer->GetLayerWidth() * Y) + X;
 		return Reachability[Index];
 	}
 
@@ -100,7 +100,7 @@ struct FHorizontalSpan
 		}
 
 		// Go right
-		for (int32 TestX = X1 + 1; TestX < Layer->LayerWidth; ++TestX)
+		for (int32 TestX = X1 + 1; TestX < Layer->GetLayerWidth(); ++TestX)
 		{
 			const FPaperTileInfo ExistingCell = Layer->GetCell(TestX, Y);
 			const bool bCellMatches = (ExistingCell == RequiredInk) || (!ExistingCell.IsValid() && !RequiredInk.IsValid());
@@ -516,7 +516,7 @@ void FEdModeTileMap::DrawHUD(FEditorViewportClient* ViewportClient, FViewport* V
 			NoCommas.UseGrouping = false;
 
 			const FPaperTileInfo Cell = LastLayer->GetCell(LastCursorTileX, LastCursorTileY);
-			const bool bInBounds = (LastCursorTileX >= 0) && (LastCursorTileX < LastLayer->LayerWidth) && (LastCursorTileY >= 0) && (LastCursorTileY < LastLayer->LayerHeight);
+			const bool bInBounds = LastLayer->InBounds(LastCursorTileX, LastCursorTileY);
 
 			FText TileIndexDescription;
 			if (!bInBounds)
@@ -692,20 +692,20 @@ bool FEdModeTileMap::BlitLayer(UPaperTileLayer* SourceLayer, UPaperTileLayer* Ta
 	bool bPaintedOnSomething = false;
 	bool bChangedSomething = false;
 
-	for (int32 SourceY = 0; SourceY < SourceLayer->LayerHeight; ++SourceY)
+	for (int32 SourceY = 0; SourceY < SourceLayer->GetLayerHeight(); ++SourceY)
 	{
 		const int32 TargetY = OffsetY + SourceY;
 
-		if ((TargetY < 0) || (TargetY >= TargetLayer->LayerHeight))
+		if ((TargetY < 0) || (TargetY >= TargetLayer->GetLayerHeight()))
 		{
 			continue;
 		}
 
-		for (int32 SourceX = 0; SourceX < SourceLayer->LayerWidth; ++SourceX)
+		for (int32 SourceX = 0; SourceX < SourceLayer->GetLayerWidth(); ++SourceX)
 		{
 			const int32 TargetX = OffsetX + SourceX;
 
-			if ((TargetX < 0) || (TargetX >= TargetLayer->LayerWidth))
+			if ((TargetX < 0) || (TargetX >= TargetLayer->GetLayerWidth()))
 			{
 				continue;
 			}
@@ -929,7 +929,7 @@ bool FEdModeTileMap::FloodFillTiles(const FViewportCursorLocation& Ray)
 			for (int32 DY = -1; DY <= 1; DY += 2)
 			{
 				const int32 Y = Span.Y + DY;
-				if ((Y < 0) || (Y >= TargetLayer->LayerHeight))
+				if ((Y < 0) || (Y >= TargetLayer->GetLayerHeight()))
 				{
 					continue;
 				}
@@ -954,8 +954,8 @@ bool FEdModeTileMap::FloodFillTiles(const FViewportCursorLocation& Ray)
 
 		// Figure out where the top left square of the map starts in the pattern, based on the seed point
 		UPaperTileLayer* SourceLayer = GetSourceInkLayer();
-		const int32 BrushWidth = SourceLayer->LayerWidth;
-		const int32 BrushHeight = SourceLayer->LayerHeight;
+		const int32 BrushWidth = SourceLayer->GetLayerWidth();
+		const int32 BrushHeight = SourceLayer->GetLayerHeight();
 
 		int32 BrushPatternOffsetX = (BrushWidth - ((DestTileX + BrushWidth) % BrushWidth));
 		int32 BrushPatternOffsetY = (BrushHeight - ((DestTileY + BrushHeight) % BrushHeight));
@@ -1235,15 +1235,15 @@ void FEdModeTileMap::RotateTilesInSelection(bool bIsClockwise)
 	static const uint8 CounterclockwiseRotationMap[8] = { 3, 2, 7, 6, 1, 0, 5, 4 };
 	const uint8* RotationTable = bIsClockwise ? ClockwiseRotationMap : CounterclockwiseRotationMap;
 
-	const int32 OldWidth = PreviewLayer->LayerWidth;
-	const int32 OldHeight = PreviewLayer->LayerHeight;
+	const int32 OldWidth = PreviewLayer->GetLayerWidth();
+	const int32 OldHeight = PreviewLayer->GetLayerHeight();
 
 	// Copy off the tiles and rotate within each tile
 	TArray<FPaperTileInfo> OldTiles;
-	OldTiles.Empty(PreviewLayer->LayerWidth * PreviewLayer->LayerHeight);
-	for (int32 Y = 0; Y < PreviewLayer->LayerHeight; ++Y)
+	OldTiles.Empty(PreviewLayer->GetLayerWidth() * PreviewLayer->GetLayerHeight());
+	for (int32 Y = 0; Y < PreviewLayer->GetLayerHeight(); ++Y)
 	{
-		for (int32 X = 0; X < PreviewLayer->LayerWidth; ++X)
+		for (int32 X = 0; X < PreviewLayer->GetLayerWidth(); ++X)
 		{
 			FPaperTileInfo Cell = PreviewLayer->GetCell(X, Y);
 			if (Cell.IsValid())
@@ -1256,12 +1256,12 @@ void FEdModeTileMap::RotateTilesInSelection(bool bIsClockwise)
 	}
 
 	// Resize, transposing width and height
-	DestructiveResizePreviewComponent(PreviewLayer->LayerHeight, PreviewLayer->LayerWidth);
+	DestructiveResizePreviewComponent(PreviewLayer->GetLayerHeight(), PreviewLayer->GetLayerWidth());
 
 	// Place the tiles back in the rotated layout
-	for (int32 NewY = 0; NewY < PreviewLayer->LayerHeight; ++NewY)
+	for (int32 NewY = 0; NewY < PreviewLayer->GetLayerHeight(); ++NewY)
 	{
-		for (int32 NewX = 0; NewX < PreviewLayer->LayerWidth; ++NewX)
+		for (int32 NewX = 0; NewX < PreviewLayer->GetLayerWidth(); ++NewX)
 		{
 			const int32 OldX = bIsClockwise ? NewY : (OldWidth - 1 - NewY);
 			const int32 OldY = bIsClockwise ? (OldHeight - 1 - NewX) : NewX;
@@ -1314,7 +1314,7 @@ void FEdModeTileMap::RotateSelectionCCW()
 bool FEdModeTileMap::HasValidSelection() const
 {
 	UPaperTileLayer* PreviewLayer = GetSourceInkLayer();
-	return (PreviewLayer->LayerWidth > 0) && (PreviewLayer->LayerHeight > 0) && bHasValidInkSource;
+	return (PreviewLayer->GetLayerWidth() > 0) && (PreviewLayer->GetLayerHeight() > 0) && bHasValidInkSource;
 }
 
 void FEdModeTileMap::SynchronizePreviewWithTileMap(UPaperTileMap* NewTileMap)
@@ -1438,13 +1438,13 @@ int32 FEdModeTileMap::GetBrushWidth() const
 		BrushWidth = FMath::Max<int32>(LastEyeDropperBounds.Width(), 1);
 		break;
 	case ETileMapEditorTool::Paintbrush:
-		BrushWidth = GetSourceInkLayer()->LayerWidth;
+		BrushWidth = GetSourceInkLayer()->GetLayerWidth();
 		break;
 	case ETileMapEditorTool::Eraser:
 		BrushWidth = EraseBrushSize;
 		break;
 	case ETileMapEditorTool::PaintBucket:
-		BrushWidth = GetSourceInkLayer()->LayerWidth;
+		BrushWidth = GetSourceInkLayer()->GetLayerWidth();
 		break;
 	case ETileMapEditorTool::TerrainBrush:
 		BrushWidth = 1;
@@ -1467,13 +1467,13 @@ int32 FEdModeTileMap::GetBrushHeight() const
 		BrushHeight = FMath::Max<int32>(LastEyeDropperBounds.Height(), 1);
 		break;
 	case ETileMapEditorTool::Paintbrush:
-		BrushHeight = GetSourceInkLayer()->LayerHeight;
+		BrushHeight = GetSourceInkLayer()->GetLayerHeight();
 		break;
 	case ETileMapEditorTool::Eraser:
 		BrushHeight = EraseBrushSize;
 		break;
 	case ETileMapEditorTool::PaintBucket:
-		BrushHeight = GetSourceInkLayer()->LayerHeight;
+		BrushHeight = GetSourceInkLayer()->GetLayerHeight();
 		break;
 	case ETileMapEditorTool::TerrainBrush:
 		BrushHeight = 1;
diff --git a/Engine/Plugins/2D/Paper2D/Source/Paper2DEditor/Private/TileMapEditing/STileLayerItem.cpp b/Engine/Plugins/2D/Paper2D/Source/Paper2DEditor/Private/TileMapEditing/STileLayerItem.cpp
--- a/Engine/Plugins/2D/Paper2D/Source/Paper2DEditor/Private/TileMapEditing/STileLayerItem.cpp
+++ b/Engine/Plugins/2D/Paper2D/Source/Paper2DEditor/Private/TileMapEditing/STileLayerItem.cpp
@@ -86,14 +86,14 @@ FReply STileLayerItem::OnToggleVisibility()
 	UPaperTileLayer* MyLayer = GetMyLayer();
 	MyLayer->SetFlags(RF_Transactional);
 	MyLayer->Modify();
-	MyLayer->bHiddenInEditor = !MyLayer->bHiddenInEditor;
+	MyLayer->SetShouldRenderInEditor(!MyLayer->ShouldRenderInEditor());
 	MyLayer->PostEditChange();
 	return FReply::Handled();
 }
 
 const FSlateBrush* STileLayerItem::GetVisibilityBrushForLayer() const
 {
-	return GetMyLayer()->bHiddenInEditor ? EyeClosed : EyeOpened;
+	return GetMyLayer()->ShouldRenderInEditor() ? EyeOpened : EyeClosed;
 }
 
 FSlateColor STileLayerItem::GetForegroundColorForVisibilityButton() const
diff --git a/Engine/Plugins/2D/Paper2D/Source/Paper2DEditor/Private/TileMapEditing/STileLayerList.cpp b/Engine/Plugins/2D/Paper2D/Source/Paper2DEditor/Private/TileMapEditing/STileLayerList.cpp
--- a/Engine/Plugins/2D/Paper2D/Source/Paper2DEditor/Private/TileMapEditing/STileLayerList.cpp
+++ b/Engine/Plugins/2D/Paper2D/Source/Paper2DEditor/Private/TileMapEditing/STileLayerList.cpp
@@ -455,9 +455,9 @@ void STileLayerList::MergeLayerDown()
 			TargetLayer->Modify();
 
 			// Copy the non-empty tiles from the source to the target layer
-			for (int32 Y = 0; Y < SourceLayer->LayerWidth; ++Y)
+			for (int32 Y = 0; Y < SourceLayer->GetLayerHeight(); ++Y)
 			{
-				for (int32 X = 0; X < SourceLayer->LayerWidth; ++X)
+				for (int32 X = 0; X < SourceLayer->GetLayerWidth(); ++X)
 				{
 					FPaperTileInfo TileInfo = SourceLayer->GetCell(X, Y);
 					if (TileInfo.IsValid())
diff --git a/Engine/Plugins/2D/Paper2D/Source/PaperTiledImporter/Private/PaperTiledImporterFactory.cpp b/Engine/Plugins/2D/Paper2D/Source/PaperTiledImporter/Private/PaperTiledImporterFactory.cpp
--- a/Engine/Plugins/2D/Paper2D/Source/PaperTiledImporter/Private/PaperTiledImporterFactory.cpp
+++ b/Engine/Plugins/2D/Paper2D/Source/PaperTiledImporter/Private/PaperTiledImporterFactory.cpp
@@ -232,13 +232,14 @@ UObject* UPaperTiledImporterFactory::FactoryCreateText(UClass* InClass, UObject*
 				NewLayer->SetFlags(RF_Transactional);
 
 				NewLayer->LayerName = FText::FromString(LayerData.Name);
-				NewLayer->bHiddenInEditor = !LayerData.bVisible;
-				NewLayer->LayerOpacity = FMath::Clamp<float>(LayerData.Opacity, 0.0f, 1.0f);
+				NewLayer->SetShouldRenderInEditor(LayerData.bVisible);
+
+				FLinearColor LayerColor = FLinearColor::White;
+				LayerColor.A = FMath::Clamp<float>(LayerData.Opacity, 0.0f, 1.0f);
+				NewLayer->SetLayerColor(LayerColor);
 
 				//@TODO: No support for Objects (and thus Color, ObjectDrawOrder), Properties, OffsetX, or OffsetY
 
-				NewLayer->LayerWidth = LayerData.Width;
-				NewLayer->LayerHeight = LayerData.Height;
 				NewLayer->DestructiveAllocateMap(LayerData.Width, LayerData.Height);
 
 				int32 SourceIndex = 0;
