Tile Map Editor: Prevented the bounding box text from being drawn when 'show stats' is disabled

diff --git a/Engine/Plugins/2D/Paper2D/Source/Paper2DEditor/Private/TileMapEditing/TileMapEditorViewportClient.cpp b/Engine/Plugins/2D/Paper2D/Source/Paper2DEditor/Private/TileMapEditing/TileMapEditorViewportClient.cpp
--- a/Engine/Plugins/2D/Paper2D/Source/Paper2DEditor/Private/TileMapEditing/TileMapEditorViewportClient.cpp
+++ b/Engine/Plugins/2D/Paper2D/Source/Paper2DEditor/Private/TileMapEditing/TileMapEditorViewportClient.cpp
@@ -204,10 +204,10 @@ void FTileMapEditorViewportClient::DrawCanvas(FViewport& Viewport, FSceneView& V
 				YPos = (int32)TextItem.Position.Y;
 			}
 		}
-	}
 
-	// Draw the render bounds
-	DrawBoundsAsText(Viewport, View, Canvas, /*inout*/ YPos);
+		// Draw the render bounds
+		DrawBoundsAsText(Viewport, View, Canvas, /*inout*/ YPos);
+	}
 }
 
 void FTileMapEditorViewportClient::Tick(float DeltaSeconds)
