Tile Maps: Added a way to view/edit per-layer properties [UE-13633]

diff --git a/Engine/Plugins/2D/Paper2D/Source/Paper2D/Classes/PaperTileMap.h b/Engine/Plugins/2D/Paper2D/Source/Paper2D/Classes/PaperTileMap.h
--- a/Engine/Plugins/2D/Paper2D/Source/Paper2D/Classes/PaperTileMap.h
+++ b/Engine/Plugins/2D/Paper2D/Source/Paper2D/Classes/PaperTileMap.h
@@ -79,7 +79,7 @@ class PAPER2D_API UPaperTileMap : public UObject
 	UMaterialInterface* Material;
 
 	// The list of layers
-	UPROPERTY(Instanced, Category=Sprite, BlueprintReadOnly)
+	UPROPERTY(Instanced, Category=Sprite, BlueprintReadOnly, EditAnywhere)
 	TArray<UPaperTileLayer*> TileLayers;
 
 protected:
diff --git a/Engine/Plugins/2D/Paper2D/Source/Paper2DEditor/Private/TileMapEditing/PaperTileMapDetailsCustomization.cpp b/Engine/Plugins/2D/Paper2D/Source/Paper2DEditor/Private/TileMapEditing/PaperTileMapDetailsCustomization.cpp
--- a/Engine/Plugins/2D/Paper2D/Source/Paper2DEditor/Private/TileMapEditing/PaperTileMapDetailsCustomization.cpp
+++ b/Engine/Plugins/2D/Paper2D/Source/Paper2DEditor/Private/TileMapEditing/PaperTileMapDetailsCustomization.cpp
@@ -27,7 +27,7 @@ TSharedRef<IDetailCustomization> FPaperTileMapDetailsCustomization::MakeInstance
 void FPaperTileMapDetailsCustomization::CustomizeDetails(IDetailLayoutBuilder& DetailLayout)
 {
 	const TArray< TWeakObjectPtr<UObject> >& SelectedObjects = DetailLayout.GetDetailsView().GetSelectedObjects();
-	MyDetailLayout = &DetailLayout;
+	MyDetailLayout = nullptr;
 	
 	FNotifyHook* NotifyHook = DetailLayout.GetPropertyUtilities()->GetNotifyHook();
 
@@ -66,9 +66,6 @@ void FPaperTileMapDetailsCustomization::CustomizeDetails(IDetailLayoutBuilder& D
 	// Make sure the Tile Map category is right below the Transform
 	IDetailCategoryBuilder& TileMapCategory = DetailLayout.EditCategory("Tile Map", FText::GetEmpty(), ECategoryPriority::Important);
 
-	// Make sure the setup category is near the top
-	IDetailCategoryBuilder& SetupCategory = DetailLayout.EditCategory("Setup", FText::GetEmpty(), ECategoryPriority::Important);
-
 	// Add the 'instanced' versus 'asset' indicator to the tile map header
 	TileMapCategory.HeaderContent
 	(
@@ -212,20 +209,70 @@ void FPaperTileMapDetailsCustomization::CustomizeDetails(IDetailLayoutBuilder& D
 			+SVerticalBox::Slot()
 			[
 				SNew(STileLayerList, TileMap, NotifyHook, CommandList)
+				.OnSelectedLayerChanged(this, &FPaperTileMapDetailsCustomization::OnSelectedLayerChanged)
 			]
 		];
 	}
 
+	// Hide the layers since they'll get visualized directly
+	TSharedRef<IPropertyHandle> TileLayersProperty = DetailLayout.GetProperty(GET_MEMBER_NAME_CHECKED(UPaperTileMap, TileLayers));
+	DetailLayout.HideProperty(TileLayersProperty);
+
+	// Add properties for the currently selected layer
+	if ((TileMap != nullptr) && ((TileComponent == nullptr) || (TileComponent->OwnsTileMap())))
+	{
+		if (TileMap->TileLayers.IsValidIndex(TileMap->SelectedLayerIndex))
+		{
+			UPaperTileLayer* SelectedLayer = TileMap->TileLayers[TileMap->SelectedLayerIndex];
+
+			const FText LayerCategoryDisplayName = LOCTEXT("LayerCategoryHeading", "Selected Layer");
+			IDetailCategoryBuilder& LayerCategory = DetailLayout.EditCategory(TEXT("SelectedLayer"), LayerCategoryDisplayName, ECategoryPriority::Important);
+
+			LayerCategory.HeaderContent
+			(
+				SNew(SBox)
+				.HAlign(HAlign_Right)
+				[
+					SNew(SHorizontalBox)
+					+SHorizontalBox::Slot()
+					.Padding(FMargin(5.0f, 0.0f))
+					.AutoWidth()
+					[
+						SNew(STextBlock)
+						.Font(FEditorStyle::GetFontStyle("TinyText"))
+						.Text(this, &FPaperTileMapDetailsCustomization::GetLayerSettingsHeadingText)
+						.ToolTipText(LOCTEXT("LayerSettingsTooltip", "Properties specific to the currently selected layer"))
+					]
+				]
+			);
+
+			TArray<UObject*> ListOfSelectedLayers;
+			ListOfSelectedLayers.Add(SelectedLayer);
+
+			for (const UProperty* TestProperty : TFieldRange<UProperty>(SelectedLayer->GetClass()))
+			{
+				if (TestProperty->HasAnyPropertyFlags(CPF_Edit))
+				{
+					const bool bAdvancedDisplay = TestProperty->HasAnyPropertyFlags(CPF_AdvancedDisplay);
+					const EPropertyLocation::Type PropertyLocation = bAdvancedDisplay ? EPropertyLocation::Advanced : EPropertyLocation::Common;
+
+					LayerCategory.AddExternalProperty(ListOfSelectedLayers, TestProperty->GetFName(), PropertyLocation);
+				}
+			}
+		}
+	}
+
+	// Make sure the setup category is near the top (just below the layer browser and layer-specific stuff)
+	IDetailCategoryBuilder& SetupCategory = DetailLayout.EditCategory("Setup", FText::GetEmpty(), ECategoryPriority::Important);
+
 	// Add all of the properties from the inline tilemap
 	if ((TileComponent != nullptr) && (TileComponent->OwnsTileMap()))
 	{
 		TArray<UObject*> ListOfTileMaps;
 		ListOfTileMaps.Add(TileMap);
 
-		for (TFieldIterator<UProperty> PropIt(UPaperTileMap::StaticClass()); PropIt; ++PropIt)
+		for (const UProperty* TestProperty : TFieldRange<UProperty>(TileMap->GetClass()))
 		{
-			UProperty* TestProperty = *PropIt;
-
 			if (TestProperty->HasAnyPropertyFlags(CPF_Edit))
 			{
 				const bool bAdvancedDisplay = TestProperty->HasAnyPropertyFlags(CPF_AdvancedDisplay);
@@ -241,6 +288,8 @@ void FPaperTileMapDetailsCustomization::CustomizeDetails(IDetailLayoutBuilder& D
 			}
 		}
 	}
+
+	MyDetailLayout = &DetailLayout;
 }
 
 FReply FPaperTileMapDetailsCustomization::EnterTileMapEditingMode()
@@ -369,6 +418,31 @@ bool FPaperTileMapDetailsCustomization::IsInstanced() const
 	return false;
 }
 
+void FPaperTileMapDetailsCustomization::OnSelectedLayerChanged()
+{
+	if (MyDetailLayout != nullptr)
+	{
+		IDetailLayoutBuilder* OldLayout = MyDetailLayout;
+		MyDetailLayout = nullptr;
+
+		OldLayout->ForceRefreshDetails();
+	}
+}
+
+FText FPaperTileMapDetailsCustomization::GetLayerSettingsHeadingText() const
+{
+	if (UPaperTileMap* TileMap = TileMapPtr.Get())
+	{
+		if (TileMap->TileLayers.IsValidIndex(TileMap->SelectedLayerIndex))
+		{
+			UPaperTileLayer* SelectedLayer = TileMap->TileLayers[TileMap->SelectedLayerIndex];
+			return SelectedLayer->LayerName;
+		}
+	}
+
+	return FText::GetEmpty();
+}
+
 //////////////////////////////////////////////////////////////////////////
 
 #undef LOCTEXT_NAMESPACE
\ No newline at end of file
diff --git a/Engine/Plugins/2D/Paper2D/Source/Paper2DEditor/Private/TileMapEditing/PaperTileMapDetailsCustomization.h b/Engine/Plugins/2D/Paper2D/Source/Paper2DEditor/Private/TileMapEditing/PaperTileMapDetailsCustomization.h
--- a/Engine/Plugins/2D/Paper2D/Source/Paper2DEditor/Private/TileMapEditing/PaperTileMapDetailsCustomization.h
+++ b/Engine/Plugins/2D/Paper2D/Source/Paper2DEditor/Private/TileMapEditing/PaperTileMapDetailsCustomization.h
@@ -39,4 +39,8 @@ private:
 
 	bool InLevelEditorContext() const;
 	bool IsInstanced() const;
+
+	void OnSelectedLayerChanged();
+
+	FText GetLayerSettingsHeadingText() const;
 };
diff --git a/Engine/Plugins/2D/Paper2D/Source/Paper2DEditor/Private/TileMapEditing/STileLayerList.cpp b/Engine/Plugins/2D/Paper2D/Source/Paper2DEditor/Private/TileMapEditing/STileLayerList.cpp
--- a/Engine/Plugins/2D/Paper2D/Source/Paper2DEditor/Private/TileMapEditing/STileLayerList.cpp
+++ b/Engine/Plugins/2D/Paper2D/Source/Paper2DEditor/Private/TileMapEditing/STileLayerList.cpp
@@ -48,6 +48,7 @@ public:
 
 void STileLayerList::Construct(const FArguments& InArgs, UPaperTileMap* InTileMap, FNotifyHook* InNotifyHook, TSharedPtr<class FUICommandList> InParentCommandList)
 {
+	OnSelectedLayerChanged = InArgs._OnSelectedLayerChanged;
 	TileMapPtr = InTileMap;
 	NotifyHook = InNotifyHook;
 
@@ -684,6 +685,8 @@ void STileLayerList::PostEditNotfications()
 		NotifyHook->NotifyPreChange(TileMapProperty);
 		NotifyHook->NotifyPostChange(FPropertyChangedEvent(TileMapProperty), TileMapProperty);
 	}
+
+	OnSelectedLayerChanged.Execute();
 }
 
 void STileLayerList::RefreshMirrorList()
diff --git a/Engine/Plugins/2D/Paper2D/Source/Paper2DEditor/Private/TileMapEditing/STileLayerList.h b/Engine/Plugins/2D/Paper2D/Source/Paper2DEditor/Private/TileMapEditing/STileLayerList.h
--- a/Engine/Plugins/2D/Paper2D/Source/Paper2DEditor/Private/TileMapEditing/STileLayerList.h
+++ b/Engine/Plugins/2D/Paper2D/Source/Paper2DEditor/Private/TileMapEditing/STileLayerList.h
@@ -9,6 +9,7 @@ class STileLayerList : public SCompoundWidget, public FEditorUndoClient
 {
 public:
 	SLATE_BEGIN_ARGS(STileLayerList) {}
+		SLATE_EVENT(FSimpleDelegate, OnSelectedLayerChanged)
 	SLATE_END_ARGS()
 
 	void Construct(const FArguments& InArgs, UPaperTileMap* TileMap, FNotifyHook* InNotifyHook, TSharedPtr<class FUICommandList> InCommandList);
@@ -30,6 +31,7 @@ protected:
 	TSharedPtr<class FUICommandList> CommandList;
 	TWeakObjectPtr<class UPaperTileMap> TileMapPtr;
 	FNotifyHook* NotifyHook;
+	FSimpleDelegate OnSelectedLayerChanged;
 protected:
 	TSharedRef<ITableRow> OnGenerateLayerListRow(FMirrorEntry Item, const TSharedRef<STableViewBase>& OwnerTable);
 
