Tile Maps: Improved control over collision generation [UE-13617]
- Collision can now be disabled for specific layers
- By default, collision for a layer is now centered on the same plane as the render geometry instead of always being at the layer 0 offset
- The center offset and thickness can be optionally overriden on a per-layer basis
- Changed how layer edits cause component reregistration allowing the per-layer collision changes to take effect immediately

diff --git a/Engine/Plugins/2D/Paper2D/Source/Paper2D/Classes/PaperTileLayer.h b/Engine/Plugins/2D/Paper2D/Source/Paper2D/Classes/PaperTileLayer.h
--- a/Engine/Plugins/2D/Paper2D/Source/Paper2D/Classes/PaperTileLayer.h
+++ b/Engine/Plugins/2D/Paper2D/Source/Paper2D/Classes/PaperTileLayer.h
@@ -126,6 +126,28 @@ private:
 	UPROPERTY(BlueprintReadOnly, EditAnywhere, Category=Sprite, meta=(AllowPrivateAccess="true"))
 	uint32 bHiddenInGame:1;
 
+	// Should this layer generate collision?
+	// Note: This only has an effect if the owning tile map has collision enabled
+	UPROPERTY(BlueprintReadOnly, EditAnywhere, Category=Sprite, meta=(AllowPrivateAccess="true"))
+	uint32 bLayerCollides:1;
+
+	// Should this layer use a custom thickness for generated collision instead of the thickness setting in the tile map?
+	UPROPERTY(BlueprintReadOnly, Category=Sprite, meta=(AllowPrivateAccess="true"))
+	uint32 bOverrideCollisionThickness:1;
+
+	// Should this layer use a custom offset for generated collision instead of the layer drawing offset?
+	UPROPERTY(BlueprintReadOnly, Category=Sprite, meta=(AllowPrivateAccess="true"))
+	uint32 bOverrideCollisionOffset:1;
+
+	// The override thickness of the collision for this layer (when bOverrideCollisionThickness is set)
+	UPROPERTY(BlueprintReadOnly, EditAnywhere, Category=Sprite, meta=(EditCondition=bOverrideCollisionThickness, AllowPrivateAccess="true"))
+	float CollisionThicknessOverride;
+
+	// The override offset of the collision for this layer (when bOverrideCollisionOffset is set)
+	// Note: This is measured in Unreal Units (cm) from the center of the tile map component, not from the previous layer.
+	UPROPERTY(BlueprintReadOnly, EditAnywhere, Category=Sprite, meta=(EditCondition=bOverrideCollisionOffset, AllowPrivateAccess="true"))
+	float CollisionOffsetOverride;
+
 	// The color of this layer (multiplied with the tile map color and passed to the material as a vertex color)
 	UPROPERTY(BlueprintReadOnly, EditAnywhere, Category=Sprite, meta=(AllowPrivateAccess="true"))
 	FLinearColor LayerColor;
@@ -179,7 +201,7 @@ public:
 	void ResizeMap(int32 NewWidth, int32 NewHeight);
 
 	// Adds collision to the specified body setup
-	void AugmentBodySetup(UBodySetup* ShapeBodySetup);
+	void AugmentBodySetup(UBodySetup* ShapeBodySetup, float RenderSeparation);
 
 	// Gets the layer-specific color multiplier
 	FLinearColor GetLayerColor() const;
diff --git a/Engine/Plugins/2D/Paper2D/Source/Paper2D/Private/PaperTileLayer.cpp b/Engine/Plugins/2D/Paper2D/Source/Paper2D/Private/PaperTileLayer.cpp
--- a/Engine/Plugins/2D/Paper2D/Source/Paper2D/Private/PaperTileLayer.cpp
+++ b/Engine/Plugins/2D/Paper2D/Source/Paper2D/Private/PaperTileLayer.cpp
@@ -39,57 +39,17 @@ const static FTransform TilePermutationTransforms[8] =
 	FTransform(FRotator(-90.0f, 0.0f, 0.0f), FVector::ZeroVector, -PaperAxisX.GetAbs() + PaperAxisY.GetAbs() + PaperAxisZ.GetAbs()),
 };
 
-//////////////////////////////////////////////////////////////////////////
-// FTileMapLayerReregisterContext
-
-/** Removes all components that use the specified tile map layer from their scenes for the lifetime of the class. */
-class FTileMapLayerReregisterContext
-{
-public:
-	/** Initialization constructor. */
-	FTileMapLayerReregisterContext(UPaperTileLayer* TargetAsset)
-	{
-		// Look at tile map components
-		for (TObjectIterator<UPaperTileMapComponent> MapIt; MapIt; ++MapIt)
-		{
-			if (UPaperTileMap* TestMap = (*MapIt)->TileMap)
-			{
-				if (TestMap->TileLayers.Contains(TargetAsset))
-				{
-					AddComponentToRefresh(*MapIt);
-				}
-			}
-		}
-	}
-
-protected:
-	void AddComponentToRefresh(UActorComponent* Component)
-	{
-		if (ComponentContexts.Num() == 0)
-		{
-			// wait until resources are released
-			FlushRenderingCommands();
-		}
-
-		new (ComponentContexts) FComponentReregisterContext(Component);
-	}
-
-private:
-	/** The recreate contexts for the individual components. */
-	TIndirectArray<FComponentReregisterContext> ComponentContexts;
-};
-
 //////////////////////////////////////////////////////////////////////////
 // FPaperTileLayerToBodySetupBuilder
 
 class FPaperTileLayerToBodySetupBuilder : public FSpriteGeometryCollisionBuilderBase
 {
 public:
-	FPaperTileLayerToBodySetupBuilder(UPaperTileMap* InTileMap, UBodySetup* InBodySetup, float InZOffset)
+	FPaperTileLayerToBodySetupBuilder(UPaperTileMap* InTileMap, UBodySetup* InBodySetup, float InZOffset, float InThickness)
 		: FSpriteGeometryCollisionBuilderBase(InBodySetup)
 	{
 		UnrealUnitsPerPixel = InTileMap->GetUnrealUnitsPerPixel();
-		CollisionThickness = InTileMap->GetCollisionThickness();
+		CollisionThickness = InThickness;
 		CollisionDomain = InTileMap->GetSpriteCollisionDomain();
 		CurrentCellOffset = FVector2D::ZeroVector;
 		ZOffsetAmount = InZOffset;
@@ -143,6 +103,11 @@ UPaperTileLayer::UPaperTileLayer(const FObjectInitializer& ObjectInitializer)
 	, bHiddenInEditor(false)
 #endif
 	, bHiddenInGame(false)
+	, bLayerCollides(true)
+	, bOverrideCollisionThickness(false)
+	, bOverrideCollisionOffset(false)
+	, CollisionThicknessOverride(50.0f)
+	, CollisionOffsetOverride(0.0f)
 	, LayerColor(FLinearColor::White)
 {
 
@@ -200,8 +165,6 @@ void UPaperTileLayer::ReallocateAndCopyMap()
 
 void UPaperTileLayer::PostEditChangeProperty(struct FPropertyChangedEvent& PropertyChangedEvent)
 {
-	FTileMapLayerReregisterContext ReregisterExistingComponents(this);
-
 	FName PropertyName = (PropertyChangedEvent.Property != nullptr) ? PropertyChangedEvent.Property->GetFName() : NAME_None;
 
 	if ((PropertyName == GET_MEMBER_NAME_CHECKED(UPaperTileLayer, LayerWidth)) || (PropertyName == GET_MEMBER_NAME_CHECKED(UPaperTileLayer, LayerHeight)))
@@ -213,12 +176,11 @@ void UPaperTileLayer::PostEditChangeProperty(struct FPropertyChangedEvent& Prope
 		// Resize the map, trying to preserve existing data
 		ReallocateAndCopyMap();
 	}
-// 	else if (PropertyName == TEXT("AllocatedCells"))
-// 	{
-// 		BakeMap();
-// 	}
 
 	Super::PostEditChangeProperty(PropertyChangedEvent);
+
+	// Force our owning tile map to recreate any component instances
+	GetTileMap()->PostEditChange();
 }
 #endif
 
@@ -250,17 +212,22 @@ void UPaperTileLayer::SetCell(int32 X, int32 Y, const FPaperTileInfo& NewValue)
 	}
 }
 
-void UPaperTileLayer::AugmentBodySetup(UBodySetup* ShapeBodySetup)
+void UPaperTileLayer::AugmentBodySetup(UBodySetup* ShapeBodySetup, float RenderSeparation)
 {
+	if (!bLayerCollides)
+	{
+		return;
+	}
+
 	UPaperTileMap* TileMap = GetTileMap();
 	const float TileWidth = TileMap->TileWidth;
 	const float TileHeight = TileMap->TileHeight;
 
-	//@TODO: Determine if we want collision to be attached to the layer or always relative to the zero layer (probably need a per-layer config value / option, as you may even want to inset layer 0's collision so it's flush with the surface (imagine a top-down game))
-	const float ZOffset = 0.0f;
+	const float EffectiveCollisionOffset = bOverrideCollisionOffset ? CollisionOffsetOverride : RenderSeparation;
+	const float EffectiveCollisionThickness = bOverrideCollisionThickness ? CollisionThicknessOverride : TileMap->GetCollisionThickness();
 
 	// Generate collision for all cells that contain a tile with collision metadata
-	FPaperTileLayerToBodySetupBuilder CollisionBuilder(GetTileMap(), ShapeBodySetup, ZOffset);
+	FPaperTileLayerToBodySetupBuilder CollisionBuilder(TileMap, ShapeBodySetup, EffectiveCollisionOffset, EffectiveCollisionThickness);
 
 	for (int32 CellY = 0; CellY < LayerHeight; ++CellY)
 	{
diff --git a/Engine/Plugins/2D/Paper2D/Source/Paper2D/Private/PaperTileMap.cpp b/Engine/Plugins/2D/Paper2D/Source/Paper2D/Private/PaperTileMap.cpp
--- a/Engine/Plugins/2D/Paper2D/Source/Paper2D/Private/PaperTileMap.cpp
+++ b/Engine/Plugins/2D/Paper2D/Source/Paper2D/Private/PaperTileMap.cpp
@@ -209,7 +209,8 @@ void UPaperTileMap::UpdateBodySetup()
 
 		for (int32 LayerIndex = 0; LayerIndex < TileLayers.Num(); ++LayerIndex)
 		{
-			TileLayers[LayerIndex]->AugmentBodySetup(BodySetup);
+			const float ZSeparation = LayerIndex * SeparationPerLayer;
+			TileLayers[LayerIndex]->AugmentBodySetup(BodySetup, ZSeparation);
 		}
 
 		// Finalize the BodySetup
