Tile Maps: Added a per-instance option to draw only a single layer as a stop-gap for working with multi-layered translucent tile maps [UE-14910]

diff --git a/Engine/Plugins/2D/Paper2D/Source/Paper2D/Classes/PaperTileMapComponent.h b/Engine/Plugins/2D/Paper2D/Source/Paper2D/Classes/PaperTileMapComponent.h
--- a/Engine/Plugins/2D/Paper2D/Source/Paper2D/Classes/PaperTileMapComponent.h
+++ b/Engine/Plugins/2D/Paper2D/Source/Paper2D/Classes/PaperTileMapComponent.h
@@ -56,6 +56,14 @@ private:
 	UPROPERTY(EditAnywhere, Category=Materials)
 	FLinearColor TileMapColor;
 
+	// The index of the single layer to use if enabled
+	UPROPERTY(EditAnywhere, Category=Rendering, meta=(EditCondition=bUseSingleLayer))
+	int32 UseSingleLayerIndex;
+
+	// Should we draw a single layer?
+	UPROPERTY()
+	bool bUseSingleLayer;
+
 #if WITH_EDITOR
 	// The number of batches required to render this tile map
 	int32 NumBatches;
@@ -93,6 +101,9 @@ public:
 	virtual void PostInitProperties() override;
 	virtual void Serialize(FArchive& Ar) override;
 	virtual void PostLoad() override;
+#if WITH_EDITOR
+	virtual void PostEditChangeProperty(FPropertyChangedEvent& PropertyChangedEvent) override;
+#endif
 	// End of UObject interface
 
 	// UActorComponent interface
diff --git a/Engine/Plugins/2D/Paper2D/Source/Paper2D/Private/PaperTileMapComponent.cpp b/Engine/Plugins/2D/Paper2D/Source/Paper2D/Private/PaperTileMapComponent.cpp
--- a/Engine/Plugins/2D/Paper2D/Source/Paper2D/Private/PaperTileMapComponent.cpp
+++ b/Engine/Plugins/2D/Paper2D/Source/Paper2D/Private/PaperTileMapComponent.cpp
@@ -16,6 +16,8 @@
 UPaperTileMapComponent::UPaperTileMapComponent(const FObjectInitializer& ObjectInitializer)
 	: Super(ObjectInitializer)
 	, TileMapColor(FLinearColor::White)
+	, UseSingleLayerIndex(0)
+	, bUseSingleLayer(false)
 {
 	BodyInstance.SetCollisionEnabled(ECollisionEnabled::QueryAndPhysics);
 
@@ -126,6 +128,25 @@ void UPaperTileMapComponent::PostLoad()
 	}
 }
 
+#if WITH_EDITOR
+void UPaperTileMapComponent::PostEditChangeProperty(FPropertyChangedEvent& PropertyChangedEvent)
+{
+	if (TileMap != nullptr)
+	{
+		if (TileMap->TileLayers.Num() > 0)
+		{
+			UseSingleLayerIndex = FMath::Clamp(UseSingleLayerIndex, 0, TileMap->TileLayers.Num() - 1);
+		}
+		else
+		{
+			UseSingleLayerIndex = 0;
+		}
+	}
+
+	Super::PostEditChangeProperty(PropertyChangedEvent);
+}
+#endif
+
 UBodySetup* UPaperTileMapComponent::GetBodySetup()
 {
 	return (TileMap != nullptr) ? TileMap->BodySetup : nullptr;
@@ -237,6 +258,14 @@ void UPaperTileMapComponent::RebuildRenderData(FPaperTileMapRenderSceneProxy* Pr
 			continue;
 		}
 
+		if (bUseSingleLayer)
+		{
+			if (Z != UseSingleLayerIndex)
+			{
+				continue;
+			}
+		}
+
 		FLinearColor DrawColor = TileMapColor * Layer->GetLayerColor();
 
 #if WITH_EDITORONLY_DATA
diff --git a/Engine/Plugins/2D/Paper2D/Source/Paper2D/Private/PaperTileMapRenderSceneProxy.cpp b/Engine/Plugins/2D/Paper2D/Source/Paper2D/Private/PaperTileMapRenderSceneProxy.cpp
--- a/Engine/Plugins/2D/Paper2D/Source/Paper2D/Private/PaperTileMapRenderSceneProxy.cpp
+++ b/Engine/Plugins/2D/Paper2D/Source/Paper2D/Private/PaperTileMapRenderSceneProxy.cpp
@@ -25,6 +25,7 @@ FPaperTileMapRenderSceneProxy::FPaperTileMapRenderSceneProxy(const UPaperTileMap
 	, bShowOutlineWhenUnselected(false)
 #endif
 	, TileMap(nullptr)
+	, OnlyLayerIndex(InComponent->bUseSingleLayer ? InComponent->UseSingleLayerIndex : INDEX_NONE)
 	, WireDepthBias(0.0001f)
 {
 	check(InComponent);
@@ -281,18 +282,25 @@ void FPaperTileMapRenderSceneProxy::GetDynamicMeshElements(const TArray<const FS
 				// Draw the debug outline
 				if (bEffectivelySelected)
 				{
-					const int32 SelectedLayerIndex = TileMap->SelectedLayerIndex;
+					const int32 SelectedLayerIndex = (OnlyLayerIndex != INDEX_NONE) ? OnlyLayerIndex : TileMap->SelectedLayerIndex;
 
 					if (bShowPerLayerGrid)
 					{
-						// Draw a bound for every layer but the selected one (and even that one if the per-tile grid is off)
-						for (int32 LayerIndex = 0; LayerIndex < TileMap->TileLayers.Num(); ++LayerIndex)
+						if (OnlyLayerIndex == INDEX_NONE)
 						{
-							if ((LayerIndex != SelectedLayerIndex) || !bShowPerTileGrid)
+							// Draw a bound for every layer but the selected one (and even that one if the per-tile grid is off)
+							for (int32 LayerIndex = 0; LayerIndex < TileMap->TileLayers.Num(); ++LayerIndex)
 							{
-								DrawBoundsForLayer(PDI, OverrideColor, LayerIndex);
+								if ((LayerIndex != SelectedLayerIndex) || !bShowPerTileGrid)
+								{
+									DrawBoundsForLayer(PDI, OverrideColor, LayerIndex);
+								}
 							}
 						}
+						else if (!bShowPerTileGrid)
+						{
+							DrawBoundsForLayer(PDI, OverrideColor, OnlyLayerIndex);
+						}
 					}
 
 					if (bShowPerTileGrid && (SelectedLayerIndex != INDEX_NONE))
@@ -315,8 +323,8 @@ void FPaperTileMapRenderSceneProxy::GetDynamicMeshElements(const TArray<const FS
 				}
 				else if (View->Family->EngineShowFlags.Grid && bShowOutlineWhenUnselected)
 				{
-					// Draw layer 0 even when not selected, so you can see where the tile map is in the editor
-					DrawBoundsForLayer(PDI, WireframeColor, /*LayerIndex=*/ 0);
+					// Draw a layer rectangle even when not selected, so you can see where the tile map is in the editor
+					DrawBoundsForLayer(PDI, WireframeColor, /*LayerIndex=*/ (OnlyLayerIndex != INDEX_NONE) ? OnlyLayerIndex : 0);
 				}
 #endif
 			}
diff --git a/Engine/Plugins/2D/Paper2D/Source/Paper2D/Private/PaperTileMapRenderSceneProxy.h b/Engine/Plugins/2D/Paper2D/Source/Paper2D/Private/PaperTileMapRenderSceneProxy.h
--- a/Engine/Plugins/2D/Paper2D/Source/Paper2D/Private/PaperTileMapRenderSceneProxy.h
+++ b/Engine/Plugins/2D/Paper2D/Source/Paper2D/Private/PaperTileMapRenderSceneProxy.h
@@ -38,6 +38,9 @@ protected:
 	//@TODO: Not thread safe
 	const class UPaperTileMap* TileMap;
 
+	// The only layer to draw, or INDEX_NONE if the filter is unset
+	const int32 OnlyLayerIndex;
+
 	// Slight depth bias so that the wireframe grid overlay doesn't z-fight with the tiles themselves
 	const float WireDepthBias;
 };
