Tile Set Editor: Fixed an issue where the first row/column of a tile could not be clicked on to select the tile [UE-15315]

diff --git a/Engine/Plugins/2D/Paper2D/Source/Paper2D/Private/PaperTileSet.cpp b/Engine/Plugins/2D/Paper2D/Source/Paper2D/Private/PaperTileSet.cpp
--- a/Engine/Plugins/2D/Paper2D/Source/Paper2D/Private/PaperTileSet.cpp
+++ b/Engine/Plugins/2D/Paper2D/Source/Paper2D/Private/PaperTileSet.cpp
@@ -123,10 +123,10 @@ FIntPoint UPaperTileSet::GetTileXYFromTextureUV(const FVector2D& TextureUV, bool
 {
 	const float DividendX = TextureUV.X - BorderMargin.Left;
 	const float DividendY = TextureUV.Y - BorderMargin.Top;
-	const float DivisorX = TileSize.X + PerTileSpacing.X;
-	const float DivisorY = TileSize.Y + PerTileSpacing.Y;
-	const int32 X = bRoundUp ? FMath::DivideAndRoundUp<int32>(DividendX, DivisorX) : FMath::DivideAndRoundDown<int32>(DividendX, DivisorX);
-	const int32 Y = bRoundUp ? FMath::DivideAndRoundUp<int32>(DividendY, DivisorY) : FMath::DivideAndRoundDown<int32>(DividendY, DivisorY);
+	const int32 DivisorX = TileSize.X + PerTileSpacing.X;
+	const int32 DivisorY = TileSize.Y + PerTileSpacing.Y;
+	const int32 X = bRoundUp ? FMath::DivideAndRoundUp<int32>(FMath::CeilToInt(DividendX), DivisorX) : FMath::DivideAndRoundDown<int32>(DividendX, DivisorX);
+	const int32 Y = bRoundUp ? FMath::DivideAndRoundUp<int32>(FMath::CeilToInt(DividendY), DivisorY) : FMath::DivideAndRoundDown<int32>(DividendY, DivisorY);
 	return FIntPoint(X, Y);
 }
 
