Paper2D: Prevent the import settings from applying a null material to tile maps, and emit a log warning if any specified materials fail to load

diff --git a/Engine/Plugins/2D/Paper2D/Source/Paper2DEditor/Private/PaperImporterSettings.cpp b/Engine/Plugins/2D/Paper2D/Source/Paper2DEditor/Private/PaperImporterSettings.cpp
--- a/Engine/Plugins/2D/Paper2D/Source/Paper2DEditor/Private/PaperImporterSettings.cpp
+++ b/Engine/Plugins/2D/Paper2D/Source/Paper2DEditor/Private/PaperImporterSettings.cpp
@@ -128,7 +128,7 @@ void UPaperImporterSettings::ApplySettingsForSpriteInit(FSpriteAssetInitParamete
 
 		// Apply the materials
 		InitParams.DefaultMaterialOverride = GetDefaultMaterial(DesiredMaterialType, bUseLitMaterial);
-		InitParams.AlternateMaterialOverride = GetDefaultOpaqueMaterial(bUseLitMaterial);
+		InitParams.AlternateMaterialOverride = GetDefaultMaterial(ESpriteInitMaterialType::Opaque, bUseLitMaterial);
 	}
 }
 
@@ -170,7 +170,10 @@ void UPaperImporterSettings::ApplySettingsForTileMapInit(UPaperTileMap* TileMap,
 		const bool bUseLitMaterial = LightingMode == ESpriteInitMaterialLightingMode::ForceLit;
 
 		// Apply the material
-		TileMap->Material = GetDefaultMaterial(DesiredMaterialType, bUseLitMaterial);
+		if (UMaterialInterface* MaterialOverride = GetDefaultMaterial(DesiredMaterialType, bUseLitMaterial))
+		{
+			TileMap->Material = MaterialOverride;
+		}
 	}
 
 	if (bCreateEmptyLayer)
@@ -197,6 +200,8 @@ UMaterialInterface* UPaperImporterSettings::GetDefaultMaskedMaterial(bool bLit)
 
 UMaterialInterface* UPaperImporterSettings::GetDefaultMaterial(ESpriteInitMaterialType MaterialType, bool bUseLitMaterial) const
 {
+	UMaterialInterface* Result = nullptr;
+
 	// Apply the materials
 	switch (MaterialType)
 	{
@@ -206,10 +211,20 @@ UMaterialInterface* UPaperImporterSettings::GetDefaultMaterial(ESpriteInitMateri
 	case ESpriteInitMaterialType::LeaveAsIs:
 	case ESpriteInitMaterialType::Automatic:
 	case ESpriteInitMaterialType::Masked:
-		return GetDefaultMaskedMaterial(bUseLitMaterial);
+		Result = GetDefaultMaskedMaterial(bUseLitMaterial);
+		break;
 	case ESpriteInitMaterialType::Translucent:
-		return GetDefaultTranslucentMaterial(bUseLitMaterial);
+		Result = GetDefaultTranslucentMaterial(bUseLitMaterial);
+		break;
 	case ESpriteInitMaterialType::Opaque:
-		return GetDefaultOpaqueMaterial(bUseLitMaterial);
+		Result = GetDefaultOpaqueMaterial(bUseLitMaterial);
+		break;
+	}
+
+	if (Result == nullptr)
+	{
+		UE_LOG(LogPaper2DEditor, Warning, TEXT("Failed to load material specified in Paper2D import editor preferences (%s %d)"), bUseLitMaterial ? TEXT("lit") : TEXT("unlit"), (int32)MaterialType);
 	}
+
+	return Result;
 }
\ No newline at end of file
