Paper2D: Changed tile maps and sprites that have a null material to render using the default material instead

diff --git a/Engine/Plugins/2D/Paper2D/Source/Paper2D/Private/PaperSpriteSceneProxy.cpp b/Engine/Plugins/2D/Paper2D/Source/Paper2D/Private/PaperSpriteSceneProxy.cpp
--- a/Engine/Plugins/2D/Paper2D/Source/Paper2D/Private/PaperSpriteSceneProxy.cpp
+++ b/Engine/Plugins/2D/Paper2D/Source/Paper2D/Private/PaperSpriteSceneProxy.cpp
@@ -17,8 +17,19 @@ FPaperSpriteSceneProxy::FPaperSpriteSceneProxy(const UPaperSpriteComponent* InCo
 	, SourceSprite(nullptr)
 {
 	WireframeColor = InComponent->GetWireframeColor();
+
 	Material = InComponent->GetMaterial(0);
+	if (Material == nullptr)
+	{
+		Material = UMaterial::GetDefaultMaterial(MD_Surface);
+	}
+
 	AlternateMaterial = InComponent->GetMaterial(1);
+	if (AlternateMaterial == nullptr)
+	{
+		AlternateMaterial = UMaterial::GetDefaultMaterial(MD_Surface);
+	}
+
 	MaterialSplitIndex = INDEX_NONE;
 	MaterialRelevance = InComponent->GetMaterialRelevance(GetScene().GetFeatureLevel());
 
diff --git a/Engine/Plugins/2D/Paper2D/Source/Paper2D/Private/PaperTileMapRenderSceneProxy.cpp b/Engine/Plugins/2D/Paper2D/Source/Paper2D/Private/PaperTileMapRenderSceneProxy.cpp
--- a/Engine/Plugins/2D/Paper2D/Source/Paper2D/Private/PaperTileMapRenderSceneProxy.cpp
+++ b/Engine/Plugins/2D/Paper2D/Source/Paper2D/Private/PaperTileMapRenderSceneProxy.cpp
@@ -33,6 +33,10 @@ FPaperTileMapRenderSceneProxy::FPaperTileMapRenderSceneProxy(const UPaperTileMap
 	WireframeColor = InComponent->GetWireframeColor();
 	TileMap = InComponent->TileMap;
 	Material = InComponent->GetMaterial(0);
+	if (Material == nullptr)
+	{
+		Material = UMaterial::GetDefaultMaterial(MD_Surface);
+	}
 	MaterialRelevance = InComponent->GetMaterialRelevance(GetScene().GetFeatureLevel());
 
 #if WITH_EDITORONLY_DATA
