Paper2D: Added a map check warning to sprites (same as the one for as static meshes) for non-uniform scale of simple collision geometry

diff --git a/Engine/Plugins/2D/Paper2D/Source/Paper2D/Private/PaperSpriteComponent.cpp b/Engine/Plugins/2D/Paper2D/Source/Paper2D/Private/PaperSpriteComponent.cpp
--- a/Engine/Plugins/2D/Paper2D/Source/Paper2D/Private/PaperSpriteComponent.cpp
+++ b/Engine/Plugins/2D/Paper2D/Source/Paper2D/Private/PaperSpriteComponent.cpp
@@ -346,6 +346,24 @@ void UPaperSpriteComponent::CheckForErrors()
 			}
 		}
 	}
+
+	// Make sure any non uniform scaled sprites have appropriate collision
+	if (IsCollisionEnabled() && (SourceSprite != nullptr) && (SourceSprite->BodySetup != nullptr) && (Owner != nullptr))
+	{
+		UBodySetup* BodySetup = SourceSprite->BodySetup;
+
+		// Overall scale factor for this mesh.
+		const FVector& TotalScale3D = ComponentToWorld.GetScale3D();
+		if (!TotalScale3D.IsUniform() && ((BodySetup->AggGeom.BoxElems.Num() > 0) || (BodySetup->AggGeom.SphylElems.Num() > 0) || (BodySetup->AggGeom.SphereElems.Num() > 0)))
+		{
+			FFormatNamedArguments Arguments;
+			Arguments.Add(TEXT("SpriteName"), FText::FromString(SourceSprite->GetName()));
+			FMessageLog("MapCheck").Warning()
+				->AddToken(FUObjectToken::Create(Owner))
+				->AddToken(FTextToken::Create(FText::Format(LOCTEXT("MapCheck_Message_SimpleCollisionButNonUniformScaleSprite", "'{SpriteName}' has simple collision but is being scaled non-uniformly - collision creation will fail"), Arguments)))
+				->AddToken(FMapErrorToken::Create(FMapErrors::SimpleCollisionButNonUniformScale));
+		}
+	}
 }
 #endif
 
