Paper2D: Changed the sprite render proxy to directly cache the BodySetup rather than accessing it via the sprite

diff --git a/Engine/Plugins/2D/Paper2D/Source/Paper2D/Private/PaperSpriteSceneProxy.cpp b/Engine/Plugins/2D/Paper2D/Source/Paper2D/Private/PaperSpriteSceneProxy.cpp
--- a/Engine/Plugins/2D/Paper2D/Source/Paper2D/Private/PaperSpriteSceneProxy.cpp
+++ b/Engine/Plugins/2D/Paper2D/Source/Paper2D/Private/PaperSpriteSceneProxy.cpp
@@ -1,3 +1,4 @@
+
 // Copyright 1998-2015 Epic Games, Inc. All Rights Reserved.
 
 #include "Paper2DPrivatePCH.h"
@@ -12,9 +13,9 @@
 //////////////////////////////////////////////////////////////////////////
 // FPaperSpriteSceneProxy
 
-FPaperSpriteSceneProxy::FPaperSpriteSceneProxy(const UPaperSpriteComponent* InComponent)
+FPaperSpriteSceneProxy::FPaperSpriteSceneProxy(UPaperSpriteComponent* InComponent)
 	: FPaperRenderSceneProxy(InComponent)
-	, SourceSprite(nullptr)
+	, BodySetup(InComponent->GetBodySetup())
 {
 	WireframeColor = InComponent->GetWireframeColor();
 
@@ -32,22 +33,20 @@ FPaperSpriteSceneProxy::FPaperSpriteSceneProxy(const UPaperSpriteComponent* InCo
 
 	MaterialSplitIndex = INDEX_NONE;
 	MaterialRelevance = InComponent->GetMaterialRelevance(GetScene().GetFeatureLevel());
-
-	SourceSprite = InComponent->SourceSprite; //@TODO: This is totally not threadsafe, and won't keep up to date if the actor's sprite changes, etc....
 }
 
 void FPaperSpriteSceneProxy::GetDynamicMeshElements(const TArray<const FSceneView*>& Views, const FSceneViewFamily& ViewFamily, uint32 VisibilityMap, FMeshElementCollector& Collector) const
 {
-	if (SourceSprite != nullptr)
+	if (BodySetup != nullptr)
 	{
 		// Show 3D physics
 		if ((ViewFamily.EngineShowFlags.Collision /*@TODO: && bIsCollisionEnabled*/) && AllowDebugViewmodes())
 		{
-			if (UBodySetup2D* BodySetup2D = Cast<UBodySetup2D>(SourceSprite->BodySetup))
+			if (const UBodySetup2D* BodySetup2D = Cast<const UBodySetup2D>(BodySetup))
 			{
 				//@TODO: Draw 2D debugging geometry
 			}
-			else if (UBodySetup* BodySetup = SourceSprite->BodySetup)
+			else
 			{
 				if (FMath::Abs(GetLocalToWorld().Determinant()) < SMALL_NUMBER)
 				{
@@ -72,7 +71,7 @@ void FPaperSpriteSceneProxy::GetDynamicMeshElements(const TArray<const FSceneVie
 
 							Collector.RegisterOneFrameMaterialProxy(CollisionMaterialInstance);
 
-							// Draw the static mesh's body setup.
+							// Draw the sprite body setup.
 
 							// Get transform without scaling.
 							FTransform GeomTransform(GetLocalToWorld());
diff --git a/Engine/Plugins/2D/Paper2D/Source/Paper2D/Private/PaperSpriteSceneProxy.h b/Engine/Plugins/2D/Paper2D/Source/Paper2D/Private/PaperSpriteSceneProxy.h
--- a/Engine/Plugins/2D/Paper2D/Source/Paper2D/Private/PaperSpriteSceneProxy.h
+++ b/Engine/Plugins/2D/Paper2D/Source/Paper2D/Private/PaperSpriteSceneProxy.h
@@ -4,13 +4,16 @@
 
 #include "PaperRenderSceneProxy.h"
 
+struct FSpriteDrawCallRecord;
+class UPaperSpriteComponent;
+
 //////////////////////////////////////////////////////////////////////////
 // FPaperSpriteSceneProxy
 
 class FPaperSpriteSceneProxy : public FPaperRenderSceneProxy
 {
 public:
-	FPaperSpriteSceneProxy(const class UPaperSpriteComponent* InComponent);
+	FPaperSpriteSceneProxy(UPaperSpriteComponent* InComponent);
 
 	// FPrimitiveSceneProxy interface
 	virtual void GetDynamicMeshElements(const TArray<const FSceneView*>& Views, const FSceneViewFamily& ViewFamily, uint32 VisibilityMap, FMeshElementCollector& Collector) const override;
@@ -24,8 +27,9 @@ protected:
 	virtual void GetDynamicMeshElementsForView(const FSceneView* View, int32 ViewIndex, bool bUseOverrideColor, const FLinearColor& OverrideColor, FMeshElementCollector& Collector) const override;
 	// End of FPaperRenderSceneProxy interface
 
+protected:
 	UMaterialInterface* AlternateMaterial;
 	int32 MaterialSplitIndex;
-	const class UPaperSprite* SourceSprite;
-	TArray<struct FSpriteDrawCallRecord> AlternateBatchedSprites;
+	const UBodySetup* BodySetup;
+	TArray<FSpriteDrawCallRecord> AlternateBatchedSprites;
 };
