Tile Maps: Exposed various tile/tile map query methods to Blueprints [UE-8012][UE-14088]

diff --git a/Engine/Plugins/2D/Paper2D/Source/Paper2D/Classes/PaperTileLayer.h b/Engine/Plugins/2D/Paper2D/Source/Paper2D/Classes/PaperTileLayer.h
--- a/Engine/Plugins/2D/Paper2D/Source/Paper2D/Classes/PaperTileLayer.h
+++ b/Engine/Plugins/2D/Paper2D/Source/Paper2D/Classes/PaperTileLayer.h
@@ -230,6 +230,8 @@ public:
 	// Returns the height of the layer (in tiles)
 	int32 GetLayerHeight() const { return LayerHeight; }
 
+	// Returns the transform for the given packed flag index (0..7)
+	static FTransform GetTileTransform(int32 FlagIndex);
 protected:
 	void ReallocateAndCopyMap();
 };
diff --git a/Engine/Plugins/2D/Paper2D/Source/Paper2D/Classes/PaperTileMapComponent.h b/Engine/Plugins/2D/Paper2D/Source/Paper2D/Classes/PaperTileMapComponent.h
--- a/Engine/Plugins/2D/Paper2D/Source/Paper2D/Classes/PaperTileMapComponent.h
+++ b/Engine/Plugins/2D/Paper2D/Source/Paper2D/Classes/PaperTileMapComponent.h
@@ -130,6 +130,10 @@ public:
 	UFUNCTION(BlueprintCallable, Category="Sprite")
 	virtual bool SetTileMap(class UPaperTileMap* NewTileMap);
 
+	// Returns the size of the tile map
+	UFUNCTION(BlueprintCallable, Category="Sprite")
+	void GetMapSize(int32& MapWidth, int32& MapHeight, int32& NumLayers);
+
 	// Returns the contents of a specified tile cell
 	UFUNCTION(BlueprintPure, Category="Sprite", meta=(Layer="0"))
 	FPaperTileInfo GetTile(int32 X, int32 Y, int32 Layer) const;
@@ -172,6 +176,18 @@ public:
 	UFUNCTION(BlueprintCallable, Category="Sprite")
 	void MakeTileMapEditable();
 
+	// Returns the position of the top left corner of the specified tile
+	UFUNCTION(BlueprintPure, Category="Sprite")
+	FVector GetTileCornerPosition(int32 TileX, int32 TileY, int32 LayerIndex = 0, bool bWorldSpace = false) const;
+
+	// Returns the position of the center of the specified tile
+	UFUNCTION(BlueprintPure, Category="Sprite")
+	FVector GetTileCenterPosition(int32 TileX, int32 TileY, int32 LayerIndex = 0, bool bWorldSpace = false) const;
+
+	// Returns the polygon for the specified tile (will be 4 or 6 vertices as a rectangle, diamond, or hexagon)
+	UFUNCTION(BlueprintPure, Category="Sprite")
+	void GetTilePolygon(int32 TileX, int32 TileY, TArray<FVector>& Points, int32 LayerIndex = 0, bool bWorldSpace = false) const;
+
 #if WITH_EDITOR
 	// Returns the rendering stats for this component
 	void GetRenderingStats(int32& OutNumTriangles, int32& OutNumBatches) const;
diff --git a/Engine/Plugins/2D/Paper2D/Source/Paper2D/Classes/PaperTileSet.h b/Engine/Plugins/2D/Paper2D/Source/Paper2D/Classes/PaperTileSet.h
--- a/Engine/Plugins/2D/Paper2D/Source/Paper2D/Classes/PaperTileSet.h
+++ b/Engine/Plugins/2D/Paper2D/Source/Paper2D/Classes/PaperTileSet.h
@@ -164,6 +164,9 @@ public:
 	// Returns the tile metadata for the specified tile index
 	const FPaperTileMetadata* GetTileMetadata(int32 TileIndex) const;
 
+	// Returns the user data name for the specified tile, or NAME_None if there is no user-specified data
+	FName GetTileUserData(int32 TileIndex) const;
+
 	// Returns the texture-space coordinates of the top left corner of the specified tile index
 	bool GetTileUV(int32 TileIndex, /*out*/ FVector2D& Out_TileUV) const;
 
diff --git a/Engine/Plugins/2D/Paper2D/Source/Paper2D/Classes/TileMapBlueprintLibrary.h b/Engine/Plugins/2D/Paper2D/Source/Paper2D/Classes/TileMapBlueprintLibrary.h
new file mode 100644
--- /dev/null
+++ b/Engine/Plugins/2D/Paper2D/Source/Paper2D/Classes/TileMapBlueprintLibrary.h
@@ -0,0 +1,34 @@
+// Copyright 1998-2015 Epic Games, Inc. All Rights Reserved.
+
+#pragma once
+
+#include "PaperTileLayer.h"
+#include "TileMapBlueprintLibrary.generated.h"
+
+/**
+ * A collection of utility methods for working with tile map components
+ *
+ * @see UPaperTileMap, UPaperTileMapComponent
+ */
+UCLASS()
+class UTileMapBlueprintLibrary : public UBlueprintFunctionLibrary
+{
+	GENERATED_BODY()
+
+public:
+	// Returns the user data name for the specified tile, or NAME_None if there is no user-specified data
+	UFUNCTION(BlueprintPure, Category=Sprite)
+	static FName GetTileUserData(FPaperTileInfo Tile);
+
+	// Returns the transform for a tile
+	UFUNCTION(BlueprintPure, Category=Sprite)
+	static FTransform GetTileTransform(FPaperTileInfo Tile);
+
+	// Breaks out the information for a tile
+	UFUNCTION(BlueprintPure, Category=Sprite)
+	static void BreakTile(FPaperTileInfo Tile, int32& TileIndex, UPaperTileSet*& TileSet, bool& bFlipH, bool& bFlipV, bool& bFlipD);
+
+	// Creates a tile from the specified information
+	UFUNCTION(BlueprintPure, Category=Sprite)
+	static FPaperTileInfo MakeTile(int32 TileIndex, UPaperTileSet* TileSet, bool bFlipH, bool bFlipV, bool bFlipD);
+};
diff --git a/Engine/Plugins/2D/Paper2D/Source/Paper2D/Private/PaperTileLayer.cpp b/Engine/Plugins/2D/Paper2D/Source/Paper2D/Private/PaperTileLayer.cpp
--- a/Engine/Plugins/2D/Paper2D/Source/Paper2D/Private/PaperTileLayer.cpp
+++ b/Engine/Plugins/2D/Paper2D/Source/Paper2D/Private/PaperTileLayer.cpp
@@ -290,3 +290,9 @@ bool UPaperTileLayer::UsesTileSet(UPaperTileSet* TileSet) const
 
 	return false;
 }
+
+FTransform UPaperTileLayer::GetTileTransform(int32 FlagIndex)
+{
+	checkSlow((FlagIndex >= 0) && (FlagIndex < 8));
+	return TilePermutationTransforms[FlagIndex];
+}
diff --git a/Engine/Plugins/2D/Paper2D/Source/Paper2D/Private/PaperTileMapComponent.cpp b/Engine/Plugins/2D/Paper2D/Source/Paper2D/Private/PaperTileMapComponent.cpp
--- a/Engine/Plugins/2D/Paper2D/Source/Paper2D/Private/PaperTileMapComponent.cpp
+++ b/Engine/Plugins/2D/Paper2D/Source/Paper2D/Private/PaperTileMapComponent.cpp
@@ -458,6 +458,22 @@ bool UPaperTileMapComponent::SetTileMap(class UPaperTileMap* NewTileMap)
 	return false;
 }
 
+void UPaperTileMapComponent::GetMapSize(int32& MapWidth, int32& MapHeight, int32& NumLayers)
+{
+	if (TileMap != nullptr)
+	{
+		MapWidth = TileMap->MapWidth;
+		MapHeight = TileMap->MapHeight;
+		NumLayers = TileMap->TileLayers.Num();
+	}
+	else
+	{
+		MapWidth = 1;
+		MapHeight = 1;
+		NumLayers = 1;
+	}
+}
+
 FPaperTileInfo UPaperTileMapComponent::GetTile(int32 X, int32 Y, int32 Layer) const
 {
 	FPaperTileInfo Result;
@@ -569,4 +585,56 @@ void UPaperTileMapComponent::GetRenderingStats(int32& OutNumTriangles, int32& Ou
 }
 #endif
 
+FVector UPaperTileMapComponent::GetTileCornerPosition(int32 TileX, int32 TileY, int32 LayerIndex, bool bWorldSpace) const
+{
+	FVector Result(ForceInitToZero);
+
+	if (TileMap != nullptr)
+	{
+		Result = TileMap->GetTilePositionInLocalSpace(TileX, TileY, LayerIndex);
+	}
+
+	if (bWorldSpace)
+	{
+		Result = ComponentToWorld.TransformPosition(Result);
+	}
+	return Result;
+}
+
+FVector UPaperTileMapComponent::GetTileCenterPosition(int32 TileX, int32 TileY, int32 LayerIndex, bool bWorldSpace) const
+{
+	FVector Result(ForceInitToZero);
+
+	if (TileMap != nullptr)
+	{
+		Result = TileMap->GetTileCenterInLocalSpace(TileX, TileY, LayerIndex);
+	}
+
+	if (bWorldSpace)
+	{
+		Result = ComponentToWorld.TransformPosition(Result);
+	}
+	return Result;
+}
+
+void UPaperTileMapComponent::GetTilePolygon(int32 TileX, int32 TileY, TArray<FVector>& Points, int32 LayerIndex, bool bWorldSpace) const
+{
+	Points.Reset();
+
+	if (TileMap != nullptr)
+	{
+		TileMap->GetTilePolygon(TileX, TileY, LayerIndex, /*out*/ Points);
+	}
+
+	if (bWorldSpace)
+	{
+		for (FVector& Point : Points)
+		{
+			Point = ComponentToWorld.TransformPosition(Point);
+		}
+	}
+}
+
+//////////////////////////////////////////////////////////////////////////
+
 #undef LOCTEXT_NAMESPACE
\ No newline at end of file
diff --git a/Engine/Plugins/2D/Paper2D/Source/Paper2D/Private/PaperTileSet.cpp b/Engine/Plugins/2D/Paper2D/Source/Paper2D/Private/PaperTileSet.cpp
--- a/Engine/Plugins/2D/Paper2D/Source/Paper2D/Private/PaperTileSet.cpp
+++ b/Engine/Plugins/2D/Paper2D/Source/Paper2D/Private/PaperTileSet.cpp
@@ -253,3 +253,15 @@ void UPaperTileSet::ReallocateAndCopyTileData()
 		}
 	}
 }
+
+FName UPaperTileSet::GetTileUserData(int32 TileIndex) const
+{
+	if (const FPaperTileMetadata* Metadata = GetTileMetadata(TileIndex))
+	{
+		return Metadata->UserDataName;
+	}
+	else
+	{
+		return NAME_None;
+	}
+}
diff --git a/Engine/Plugins/2D/Paper2D/Source/Paper2D/Private/TileMapBlueprintLibrary.cpp b/Engine/Plugins/2D/Paper2D/Source/Paper2D/Private/TileMapBlueprintLibrary.cpp
new file mode 100644
--- /dev/null
+++ b/Engine/Plugins/2D/Paper2D/Source/Paper2D/Private/TileMapBlueprintLibrary.cpp
@@ -0,0 +1,44 @@
+// Copyright 1998-2015 Epic Games, Inc. All Rights Reserved.
+
+#include "Paper2DPrivatePCH.h"
+#include "TileMapBlueprintLibrary.h"
+#include "PaperTileLayer.h"
+
+//////////////////////////////////////////////////////////////////////////
+
+FName UTileMapBlueprintLibrary::GetTileUserData(FPaperTileInfo Tile)
+{
+	if (Tile.TileSet != nullptr)
+	{
+		return Tile.TileSet->GetTileUserData(Tile.GetTileIndex());
+	}
+	else
+	{
+		return NAME_None;
+	}
+}
+
+FTransform UTileMapBlueprintLibrary::GetTileTransform(FPaperTileInfo Tile)
+{
+	return UPaperTileLayer::GetTileTransform(Tile.GetFlagsAsIndex());
+}
+
+void UTileMapBlueprintLibrary::BreakTile(FPaperTileInfo Tile, int32& TileIndex, UPaperTileSet*& TileSet, bool& bFlipH, bool& bFlipV, bool& bFlipD)
+{
+	TileIndex = Tile.GetTileIndex();
+	TileSet = Tile.TileSet;
+	bFlipH = Tile.HasFlag(EPaperTileFlags::FlipHorizontal);
+	bFlipV = Tile.HasFlag(EPaperTileFlags::FlipVertical);
+	bFlipD = Tile.HasFlag(EPaperTileFlags::FlipDiagonal);
+}
+
+FPaperTileInfo UTileMapBlueprintLibrary::MakeTile(int32 TileIndex, UPaperTileSet* TileSet, bool bFlipH, bool bFlipV, bool bFlipD)
+{
+	FPaperTileInfo Result;
+	Result.TileSet = TileSet;
+	Result.PackedTileIndex = TileIndex;
+	Result.SetFlagValue(EPaperTileFlags::FlipHorizontal, bFlipH);
+	Result.SetFlagValue(EPaperTileFlags::FlipVertical, bFlipV);
+	Result.SetFlagValue(EPaperTileFlags::FlipDiagonal, bFlipD);
+	return Result;
+}
