UE-14722: Each individual UPaperSprite asset should not register for an OnObjectReimported delegate

diff --git a/Engine/Plugins/2D/Paper2D/Source/Paper2D/Classes/PaperSprite.h b/Engine/Plugins/2D/Paper2D/Source/Paper2D/Classes/PaperSprite.h
--- a/Engine/Plugins/2D/Paper2D/Source/Paper2D/Classes/PaperSprite.h
+++ b/Engine/Plugins/2D/Paper2D/Source/Paper2D/Classes/PaperSprite.h
@@ -252,9 +252,8 @@ public:
 #endif
 
 #if WITH_EDITOR
-	bool bRegisteredObjectReimport;
-	// Called via delegate when an object is re-imported in the editor
-	void OnObjectReimported(UObject* InObject);
+	// Called when an object is re-imported in the editor
+	void OnObjectReimported(UTexture2D* InObject);
 #endif
 
 #if WITH_EDITOR
diff --git a/Engine/Plugins/2D/Paper2D/Source/Paper2D/Private/PaperSprite.cpp b/Engine/Plugins/2D/Paper2D/Source/Paper2D/Private/PaperSprite.cpp
--- a/Engine/Plugins/2D/Paper2D/Source/Paper2D/Private/PaperSprite.cpp
+++ b/Engine/Plugins/2D/Paper2D/Source/Paper2D/Private/PaperSprite.cpp
@@ -392,26 +392,16 @@ UPaperSprite::UPaperSprite(const FObjectInitializer& ObjectInitializer)
 
 	static ConstructorHelpers::FObjectFinder<UMaterialInterface> OpaqueMaterialRef(TEXT("/Paper2D/OpaqueUnlitSpriteMaterial"));
 	AlternateMaterial = OpaqueMaterialRef.Object;
-
-#if WITH_EDITOR
-	// Hook into notifications for object re-imports so that the gameplay tag tree can be reconstructed if the table changes
-	if (GIsEditor && !bRegisteredObjectReimport)
-	{
-		bRegisteredObjectReimport = true;
-		FEditorDelegates::OnAssetReimport.AddUObject(this, &UPaperSprite::OnObjectReimported);
-	}
-#endif
 }
 
 #if WITH_EDITOR
 
-void UPaperSprite::OnObjectReimported(UObject* InObject)
+void UPaperSprite::OnObjectReimported(UTexture2D* Texture)
 {
 	// Check if its our source texture, and if its dimensions have changed
 	// If SourceTetxureDimension == 0, we don't have a previous dimension to work off, so can't
 	// rescale sensibly
-	UTexture2D* Texture = Cast<UTexture2D>(InObject);
-	if ((Texture != nullptr) && (Texture == GetSourceTexture()))
+	if (Texture == GetSourceTexture())
 	{
 		if (NeedRescaleSpriteData())
 		{
diff --git a/Engine/Plugins/2D/Paper2D/Source/Paper2DEditor/Private/Paper2DEditorModule.cpp b/Engine/Plugins/2D/Paper2D/Source/Paper2DEditor/Private/Paper2DEditorModule.cpp
--- a/Engine/Plugins/2D/Paper2D/Source/Paper2DEditor/Private/Paper2DEditorModule.cpp
+++ b/Engine/Plugins/2D/Paper2D/Source/Paper2DEditor/Private/Paper2DEditorModule.cpp
@@ -96,8 +96,8 @@ private:
 	TSharedPtr<IComponentAssetBroker> PaperTileMapBroker;
 
 	TSharedPtr<IMeshPaintGeometryAdapterFactory> SpriteMeshPaintAdapterFactory;
-	FCoreUObjectDelegates::FOnObjectPropertyChanged::FDelegate OnPropertyChangedHandle;
-	FDelegateHandle OnPropertyChangedHandleDelegateHandle;
+	FDelegateHandle OnPropertyChangedDelegateHandle;
+	FDelegateHandle OnAssetReimportDelegateHandle;
 
 	EAssetTypeCategories::Type Paper2DAssetCategoryBit;
 
@@ -153,8 +153,10 @@ public:
 		}
 
 		// Register to be notified when properties are edited
-		OnPropertyChangedHandle = FCoreUObjectDelegates::FOnObjectPropertyChanged::FDelegate::CreateRaw(this, &FPaper2DEditor::OnPropertyChanged);
-		OnPropertyChangedHandleDelegateHandle = FCoreUObjectDelegates::OnObjectPropertyChanged.Add(OnPropertyChangedHandle);
+		OnPropertyChangedDelegateHandle = FCoreUObjectDelegates::OnObjectPropertyChanged.AddRaw(this, &FPaper2DEditor::OnPropertyChanged);
+
+		// Register to be notified when an asset is reimported
+		OnAssetReimportDelegateHandle = FEditorDelegates::OnAssetReimport.AddRaw(this, &FPaper2DEditor::OnObjectReimported);
 
 		// Register the thumbnail renderers
 		UThumbnailManager::Get().RegisterCustomRenderer(UPaperSprite::StaticClass(), UPaperSpriteThumbnailRenderer::StaticClass());
@@ -225,7 +227,10 @@ public:
 			UThumbnailManager::Get().UnregisterCustomRenderer(UPaperFlipbook::StaticClass());
 
 			// Unregister the property modification handler
-			FCoreUObjectDelegates::OnObjectPropertyChanged.Remove(OnPropertyChangedHandleDelegateHandle);
+			FCoreUObjectDelegates::OnObjectPropertyChanged.Remove(OnPropertyChangedDelegateHandle);
+
+			// Unregister the asset reimport handler
+			FEditorDelegates::OnAssetReimport.Remove(OnAssetReimportDelegateHandle);
 		}
 
 		// Unregister the details customization
@@ -268,6 +273,17 @@ private:
 		}
 	}
 
+	void OnObjectReimported(UObject* InObject)
+	{
+		if (UTexture2D* Texture = Cast<UTexture2D>(InObject))
+		{
+			for (TObjectIterator<UPaperSprite> SpriteIt; SpriteIt; ++SpriteIt)
+			{
+				SpriteIt->OnObjectReimported(Texture);
+			}
+		}
+	}
+
 	void RegisterSettings()
 	{
 		if (ISettingsModule* SettingsModule = FModuleManager::GetModulePtr<ISettingsModule>("Settings"))
