Tile Maps: Exposed a way to create a new instanced tile map on a tile map component

diff --git a/Engine/Plugins/2D/Paper2D/Source/Paper2D/Classes/PaperTileMapComponent.h b/Engine/Plugins/2D/Paper2D/Source/Paper2D/Classes/PaperTileMapComponent.h
--- a/Engine/Plugins/2D/Paper2D/Source/Paper2D/Classes/PaperTileMapComponent.h
+++ b/Engine/Plugins/2D/Paper2D/Source/Paper2D/Classes/PaperTileMapComponent.h
@@ -122,6 +122,18 @@ public:
 	// Creates a new tile map internally, replacing the TileMap reference (or dropping the previous owned one)
 	void CreateNewOwnedTileMap();
 
+	/**
+	 * Creates a new tile map of the specified size, replacing the TileMap reference (or dropping the previous owned one)
+	 *
+	 * @param MapWidth Width of the map (in tiles)
+	 * @param MapHeight Height of the map (in tiles)
+	 * @param TileWidth Width of one tile (in pixels)
+	 * @param TileHeight Height of one tile (in pixels)
+	 * @param bCreateLayer Should an empty layer be created?
+	 */
+	UFUNCTION(BlueprintCallable, Category="Sprite")
+	void CreateNewTileMap(int32 MapWidth = 4, int32 MapHeight = 4, int32 TileWidth = 32, int32 TileHeight = 32, float PixelsPerUnrealUnit = 1.0f, bool bCreateLayer = true);
+
 	// Does this component own the tile map (is it instanced instead of being an asset reference)?
 	UFUNCTION(BlueprintCallable, Category="Sprite")
 	bool OwnsTileMap() const;
diff --git a/Engine/Plugins/2D/Paper2D/Source/Paper2D/Private/PaperTileMapComponent.cpp b/Engine/Plugins/2D/Paper2D/Source/Paper2D/Private/PaperTileMapComponent.cpp
--- a/Engine/Plugins/2D/Paper2D/Source/Paper2D/Private/PaperTileMapComponent.cpp
+++ b/Engine/Plugins/2D/Paper2D/Source/Paper2D/Private/PaperTileMapComponent.cpp
@@ -427,6 +427,26 @@ void UPaperTileMapComponent::CreateNewOwnedTileMap()
 	SetTileMap(NewTileMap);
 }
 
+void UPaperTileMapComponent::CreateNewTileMap(int32 MapWidth, int32 MapHeight, int32 TileWidth, int32 TileHeight, float PixelsPerUnrealUnit, bool bCreateLayer)
+{
+	TGuardValue<TEnumAsByte<EComponentMobility::Type>> MobilitySaver(Mobility, EComponentMobility::Movable);
+
+	UPaperTileMap* NewTileMap = NewObject<UPaperTileMap>(this);
+	NewTileMap->SetFlags(RF_Transactional);
+	NewTileMap->MapWidth = MapWidth;
+	NewTileMap->MapHeight = MapHeight;
+	NewTileMap->TileWidth = TileWidth;
+	NewTileMap->TileHeight = TileHeight;
+	NewTileMap->PixelsPerUnrealUnit = PixelsPerUnrealUnit;
+
+	if (bCreateLayer)
+	{
+		NewTileMap->AddNewLayer();
+	}
+
+	SetTileMap(NewTileMap);
+}
+
 bool UPaperTileMapComponent::OwnsTileMap() const
 {
 	return (TileMap != nullptr) && (TileMap->GetOuter() == this);
@@ -498,6 +518,10 @@ void UPaperTileMapComponent::SetTile(int32 X, int32 Y, int32 Layer, FPaperTileIn
 
 			MarkRenderStateDirty();
 		}
+		else
+		{
+			UE_LOG(LogPaper2D, Warning, TEXT("Invalid layer index %d for %s"), Layer, *TileMap->GetPathName());
+		}
 	}
 }
 
