Tile Maps: Removed BlueprintReadWrite from FPaperTileInfo members, opting instead for friendlier make/break functions (since the index contains packed flags)

diff --git a/Engine/Plugins/2D/Paper2D/Source/Paper2D/Classes/PaperTileLayer.h b/Engine/Plugins/2D/Paper2D/Source/Paper2D/Classes/PaperTileLayer.h
--- a/Engine/Plugins/2D/Paper2D/Source/Paper2D/Classes/PaperTileLayer.h
+++ b/Engine/Plugins/2D/Paper2D/Source/Paper2D/Classes/PaperTileLayer.h
@@ -17,17 +17,17 @@ enum class EPaperTileFlags : uint32
 };
 
 // This is the contents of a tile map cell
-USTRUCT(BlueprintType)
+USTRUCT(BlueprintType, meta=(HasNativeBreak="Paper2D.TileMapBlueprintLibrary.BreakTile", HasNativeMake="Paper2D.TileMapBlueprintLibrary.MakeTile"))
 struct FPaperTileInfo
 {
 	GENERATED_USTRUCT_BODY()
 
 	// The tile set that this tile comes from
-	UPROPERTY(EditAnywhere, BlueprintReadWrite, Category=Sprite)
+	UPROPERTY(EditAnywhere, Category=Sprite)
 	UPaperTileSet* TileSet;
 
 	// This is the index of the current tile within the tile set
-	UPROPERTY(EditAnywhere, BlueprintReadWrite, Category=Sprite)
+	UPROPERTY(EditAnywhere, Category=Sprite)
 	int32 PackedTileIndex;
 
 	FPaperTileInfo()
diff --git a/Engine/Plugins/2D/Paper2D/Source/Paper2D/Classes/TileMapBlueprintLibrary.h b/Engine/Plugins/2D/Paper2D/Source/Paper2D/Classes/TileMapBlueprintLibrary.h
--- a/Engine/Plugins/2D/Paper2D/Source/Paper2D/Classes/TileMapBlueprintLibrary.h
+++ b/Engine/Plugins/2D/Paper2D/Source/Paper2D/Classes/TileMapBlueprintLibrary.h
@@ -25,10 +25,10 @@ public:
 	static FTransform GetTileTransform(FPaperTileInfo Tile);
 
 	// Breaks out the information for a tile
-	UFUNCTION(BlueprintPure, Category=Sprite)
+	UFUNCTION(BlueprintPure, Category=Sprite, meta=(AdvancedDisplay=2))
 	static void BreakTile(FPaperTileInfo Tile, int32& TileIndex, UPaperTileSet*& TileSet, bool& bFlipH, bool& bFlipV, bool& bFlipD);
 
 	// Creates a tile from the specified information
-	UFUNCTION(BlueprintPure, Category=Sprite)
+	UFUNCTION(BlueprintPure, Category=Sprite, meta=(AdvancedDisplay=2))
 	static FPaperTileInfo MakeTile(int32 TileIndex, UPaperTileSet* TileSet, bool bFlipH, bool bFlipV, bool bFlipD);
 };
