Paper2D: Changed the name of the 'configure for retro sprites' menu option to reflect that it now pull the settings from the project settings

diff --git a/Engine/Plugins/2D/Paper2D/Source/Paper2DEditor/Private/ContentBrowserExtensions/ContentBrowserExtensions.cpp b/Engine/Plugins/2D/Paper2D/Source/Paper2DEditor/Private/ContentBrowserExtensions/ContentBrowserExtensions.cpp
--- a/Engine/Plugins/2D/Paper2D/Source/Paper2DEditor/Private/ContentBrowserExtensions/ContentBrowserExtensions.cpp
+++ b/Engine/Plugins/2D/Paper2D/Source/Paper2DEditor/Private/ContentBrowserExtensions/ContentBrowserExtensions.cpp
@@ -236,7 +236,7 @@ public:
 			NAME_None,
 			EUserInterfaceActionType::Button);
 
-		// Configure for retro sprites
+		// Configure the selected textures according to the project settings (same as if it got imported from a sprite sheet)
 		TSharedPtr<FConfigureTexturesForSpriteUsageExtension> TextureConfigFunctor = MakeShareable(new FConfigureTexturesForSpriteUsageExtension());
 		TextureConfigFunctor->SelectedAssets = SelectedAssets;
 
@@ -244,8 +244,8 @@ public:
 			FExecuteAction::CreateStatic(&FPaperContentBrowserExtensions_Impl::ExecuteSelectedContentFunctor, StaticCastSharedPtr<FContentBrowserSelectedAssetExtensionBase>(TextureConfigFunctor)));
 
 		MenuBuilder.AddMenuEntry(
-			LOCTEXT("CB_Extension_Texture_ConfigureTextureForSprites", "Configure For Retro Sprites"),
-			LOCTEXT("CB_Extension_Texture_ConfigureTextureForSprites_Tooltip", "Sets compression settings and sampling modes to good defaults for retro sprites (nearest filtering, uncompressed, etc...)"),
+			LOCTEXT("CB_Extension_Texture_ConfigureTextureForSprites", "Apply Paper2D Texture Settings"),
+			LOCTEXT("CB_Extension_Texture_ConfigureTextureForSprites_Tooltip", "Sets compression settings and sampling modes to the defaults specified in the 'Paper2D - Import' editor preferences"),
 			FSlateIcon(PaperStyleSetName, "AssetActions.ConfigureForRetroSprites"),
 			Action_ConfigureTexturesForSprites,
 			NAME_None,
