Paper2D: Fixed an issue where editing the settings for the tile map editor was really editing the flipbook editor settings

diff --git a/Engine/Plugins/2D/Paper2D/Source/Paper2DEditor/Private/Paper2DEditorModule.cpp b/Engine/Plugins/2D/Paper2D/Source/Paper2DEditor/Private/Paper2DEditorModule.cpp
--- a/Engine/Plugins/2D/Paper2D/Source/Paper2DEditor/Private/Paper2DEditorModule.cpp
+++ b/Engine/Plugins/2D/Paper2D/Source/Paper2DEditor/Private/Paper2DEditorModule.cpp
@@ -306,7 +306,7 @@ private:
 			SettingsModule->RegisterSettings("Editor", "ContentEditors", "TileMapEditor",
 				LOCTEXT("TileMapEditorSettingsName", "Tile Map Editor"),
 				LOCTEXT("TileMapEditorSettingsDescription", "Configure the look and feel of the Tile Map Editor."),
-				GetMutableDefault<UFlipbookEditorSettings>());
+				GetMutableDefault<UTileMapEditorSettings>());
 
 			SettingsModule->RegisterSettings("Editor", "ContentEditors", "TileSetEditor",
 				LOCTEXT("TileSetEditorSettingsName", "Tile Set Editor"),
