Paper2D: Moved importer settings to Project Settings from Editor Preferences, since nothing else writing to DefaultEditor was there

diff --git a/Engine/Plugins/2D/Paper2D/Source/Paper2DEditor/Private/ContentBrowserExtensions/ContentBrowserExtensions.cpp b/Engine/Plugins/2D/Paper2D/Source/Paper2DEditor/Private/ContentBrowserExtensions/ContentBrowserExtensions.cpp
--- a/Engine/Plugins/2D/Paper2D/Source/Paper2DEditor/Private/ContentBrowserExtensions/ContentBrowserExtensions.cpp
+++ b/Engine/Plugins/2D/Paper2D/Source/Paper2DEditor/Private/ContentBrowserExtensions/ContentBrowserExtensions.cpp
@@ -245,7 +245,7 @@ public:
 
 		MenuBuilder.AddMenuEntry(
 			LOCTEXT("CB_Extension_Texture_ConfigureTextureForSprites", "Apply Paper2D Texture Settings"),
-			LOCTEXT("CB_Extension_Texture_ConfigureTextureForSprites_Tooltip", "Sets compression settings and sampling modes to the defaults specified in the 'Paper2D - Import' editor preferences"),
+			LOCTEXT("CB_Extension_Texture_ConfigureTextureForSprites_Tooltip", "Sets compression settings and sampling modes to the defaults specified in the 'Paper2D - Import' project settings"),
 			FSlateIcon(PaperStyleSetName, "AssetActions.ConfigureForRetroSprites"),
 			Action_ConfigureTexturesForSprites,
 			NAME_None,
diff --git a/Engine/Plugins/2D/Paper2D/Source/Paper2DEditor/Private/Paper2DEditorModule.cpp b/Engine/Plugins/2D/Paper2D/Source/Paper2DEditor/Private/Paper2DEditorModule.cpp
--- a/Engine/Plugins/2D/Paper2D/Source/Paper2DEditor/Private/Paper2DEditorModule.cpp
+++ b/Engine/Plugins/2D/Paper2D/Source/Paper2DEditor/Private/Paper2DEditorModule.cpp
@@ -313,7 +313,7 @@ private:
 				LOCTEXT("TileSetEditorSettingsDescription", "Configure the look and feel of the Tile Set Editor."),
 				GetMutableDefault<UTileSetEditorSettings>());
 
-			SettingsModule->RegisterSettings("Editor", "General", "Paper2DImport",
+			SettingsModule->RegisterSettings("Project", "Editor", "Paper2DImport",
 				LOCTEXT("PaperImporterSettingsName", "Paper2D - Import"),
 				LOCTEXT("PaperImporterSettingsDescription", "Configure how assets get imported or converted to sprites."),
 				GetMutableDefault<UPaperImporterSettings>());
@@ -324,7 +324,7 @@ private:
 	{
 		if (ISettingsModule* SettingsModule = FModuleManager::GetModulePtr<ISettingsModule>("Settings"))
 		{
-			SettingsModule->UnregisterSettings("Editor", "General", "Paper2DImport");
+			SettingsModule->UnregisterSettings("Project", "Editor", "Paper2DImport");
 			SettingsModule->UnregisterSettings("Editor", "ContentEditors", "TileSetEditor");
 			SettingsModule->UnregisterSettings("Editor", "ContentEditors", "TileMapEditor");
 			SettingsModule->UnregisterSettings("Editor", "ContentEditors", "FlipbookEditor");
diff --git a/Engine/Plugins/2D/Paper2D/Source/Paper2DEditor/Private/PaperImporterSettings.cpp b/Engine/Plugins/2D/Paper2D/Source/Paper2DEditor/Private/PaperImporterSettings.cpp
--- a/Engine/Plugins/2D/Paper2D/Source/Paper2DEditor/Private/PaperImporterSettings.cpp
+++ b/Engine/Plugins/2D/Paper2D/Source/Paper2DEditor/Private/PaperImporterSettings.cpp
@@ -223,7 +223,7 @@ UMaterialInterface* UPaperImporterSettings::GetDefaultMaterial(ESpriteInitMateri
 
 	if (Result == nullptr)
 	{
-		UE_LOG(LogPaper2DEditor, Warning, TEXT("Failed to load material specified in Paper2D import editor preferences (%s %d)"), bUseLitMaterial ? TEXT("lit") : TEXT("unlit"), (int32)MaterialType);
+		UE_LOG(LogPaper2DEditor, Warning, TEXT("Failed to load material specified in Paper2D import settings (%s %d)"), bUseLitMaterial ? TEXT("lit") : TEXT("unlit"), (int32)MaterialType);
 	}
 
 	return Result;
