Added an method to skip data validation for details panel propertynodes.  This should be only used in advanced customization circumstances with huge amounts of data where the way the nodes are invalidated is strictly controlled

diff --git a/Engine/Plugins/2D/Paper2D/Source/Paper2DEditor/Private/TileSetEditor/TileSetDetailsCustomization.cpp b/Engine/Plugins/2D/Paper2D/Source/Paper2DEditor/Private/TileSetEditor/TileSetDetailsCustomization.cpp
--- a/Engine/Plugins/2D/Paper2D/Source/Paper2DEditor/Private/TileSetEditor/TileSetDetailsCustomization.cpp
+++ b/Engine/Plugins/2D/Paper2D/Source/Paper2DEditor/Private/TileSetEditor/TileSetDetailsCustomization.cpp
@@ -72,6 +72,8 @@ void FTileSetDetailsCustomization::CustomizeDetails(IDetailLayoutBuilder& Detail
 		const FName MetadataArrayName = UPaperTileSet::GetPerTilePropertyName();
 		TSharedPtr<IPropertyHandle> PerTileArrayProperty = DetailLayout.GetProperty(MetadataArrayName);
 		DetailLayout.HideProperty(PerTileArrayProperty);
+		// this array is potentially huge and has a costly validation overhead.  We only ever show one element in the array so there is no need to validate every element.
+		PerTileArrayProperty->SetIgnoreValidation(true);
 
 		if (SelectedSingleTileIndex != INDEX_NONE)
 		{
diff --git a/Engine/Source/Editor/PropertyEditor/Private/PropertyHandleImpl.cpp b/Engine/Source/Editor/PropertyEditor/Private/PropertyHandleImpl.cpp
--- a/Engine/Source/Editor/PropertyEditor/Private/PropertyHandleImpl.cpp
+++ b/Engine/Source/Editor/PropertyEditor/Private/PropertyHandleImpl.cpp
@@ -2035,6 +2035,15 @@ bool FPropertyHandleBase::GenerateRestrictionToolTip(const FString& Value, FText
 	return false;
 }
 
+void FPropertyHandleBase::SetIgnoreValidation(bool bInIgnore)
+{
+	TSharedPtr<FPropertyNode> PropertyNode = Implementation->GetPropertyNode();
+	if( PropertyNode.IsValid() )
+	{
+		PropertyNode->SetNodeFlags( EPropertyNodeFlags::SkipChildValidation, bInIgnore); 
+	}
+}
+
 /** Implements common property value functions */
 #define IMPLEMENT_PROPERTY_VALUE( ClassName ) \
 	ClassName::ClassName( TSharedRef<FPropertyNode> PropertyNode, FNotifyHook* NotifyHook, TSharedPtr<IPropertyUtilities> PropertyUtilities ) \
diff --git a/Engine/Source/Editor/PropertyEditor/Private/PropertyHandleImpl.h b/Engine/Source/Editor/PropertyEditor/Private/PropertyHandleImpl.h
--- a/Engine/Source/Editor/PropertyEditor/Private/PropertyHandleImpl.h
+++ b/Engine/Source/Editor/PropertyEditor/Private/PropertyHandleImpl.h
@@ -414,6 +414,7 @@ public:
 	virtual bool IsRestricted(const FString& Value) const override;
 	virtual bool IsRestricted(const FString& Value, TArray<FText>& OutReasons) const override;
 	virtual bool GenerateRestrictionToolTip(const FString& Value, FText& OutTooltip) const override;
+	virtual void SetIgnoreValidation(bool bInIgnore) override;
 
 	TSharedPtr<FPropertyNode> GetPropertyNode();
 protected:
diff --git a/Engine/Source/Editor/PropertyEditor/Private/PropertyNode.cpp b/Engine/Source/Editor/PropertyEditor/Private/PropertyNode.cpp
--- a/Engine/Source/Editor/PropertyEditor/Private/PropertyNode.cpp
+++ b/Engine/Source/Editor/PropertyEditor/Private/PropertyNode.cpp
@@ -321,7 +321,7 @@ FObjectPropertyNode* FPropertyNode::FindRootObjectItemParent()
  */
 FPropertyNode::DataValidationResult FPropertyNode::EnsureDataIsValid()
 {
-	bool bValidateChildren = true;
+	bool bValidateChildren = !HasNodeFlags(EPropertyNodeFlags::SkipChildValidation);
 
 	// The root must always be validated
 	if( GetParentNode() == NULL || HasNodeFlags(EPropertyNodeFlags::RequiresValidation) != 0 )
diff --git a/Engine/Source/Editor/PropertyEditor/Private/PropertyNode.h b/Engine/Source/Editor/PropertyEditor/Private/PropertyNode.h
--- a/Engine/Source/Editor/PropertyEditor/Private/PropertyNode.h
+++ b/Engine/Source/Editor/PropertyEditor/Private/PropertyNode.h
@@ -43,6 +43,8 @@ namespace EPropertyNodeFlags
 
 	const Type	IsReadOnly						= 1 << 20; /** true if this node is overriden to appear as read-only */
 
+	const Type	SkipChildValidation				= 1 << 21; /** true if this node should skip child validation */
+
 	const Type 	NoFlags							= 0;
 
 };
diff --git a/Engine/Source/Editor/PropertyEditor/Public/PropertyHandle.h b/Engine/Source/Editor/PropertyEditor/Public/PropertyHandle.h
--- a/Engine/Source/Editor/PropertyEditor/Public/PropertyHandle.h
+++ b/Engine/Source/Editor/PropertyEditor/Public/PropertyHandle.h
@@ -379,6 +379,12 @@ public:
 	 * @return				True if this value is restricted.
 	 */
 	virtual bool GenerateRestrictionToolTip(const FString& Value, FText& OutTooltip)const = 0;
+
+	 /** 
+	  * Sets whether or not data validation should occur for this property and all of its children. It is generally unsafe to set this value unless you know what you are doing.  Data validation done by the details panel ensures changes to properties out from under the details panel are known
+	  * This should only ever be set for extremely large arrays or other costly validation checks where validation is handled by the customizer
+	  */
+	virtual void SetIgnoreValidation(bool bInIgnore) = 0;
 };
 
 /**
