Paper2D: Implemented a stable vertex buffer as a fast path for FPaperRenderSceneProxy subclasses
- Switched tile maps to use the fast path, giving a huge win on per-frame CPU time in a 512x512 test case (357 ms -> ~0.04 ms)
- Pulled some FLinearColor->FColor conversions out of inner loops, did a presize pass for the vertices, and removed some cell access validation overhead to optimize tile map rebuild times
- Added a tile map rebuild stat
- Removed unused OverrideColor code
Note: The vertex buffer is currently stored per component instance rather than once per asset

diff --git a/Engine/Plugins/2D/Paper2D/Source/Paper2D/Classes/PaperTileLayer.h b/Engine/Plugins/2D/Paper2D/Source/Paper2D/Classes/PaperTileLayer.h
--- a/Engine/Plugins/2D/Paper2D/Source/Paper2D/Classes/PaperTileLayer.h
+++ b/Engine/Plugins/2D/Paper2D/Source/Paper2D/Classes/PaperTileLayer.h
@@ -6,6 +6,8 @@
 
 #include "PaperTileLayer.generated.h"
 
+class UPaperTileMap;
+
 // Flags used in the packed tile index
 enum class EPaperTileFlags : uint32
 {
@@ -180,7 +182,7 @@ public:
 	// End of UObject interface
 
 	// Returns the parent tile map
-	class UPaperTileMap* GetTileMap() const;
+	UPaperTileMap* GetTileMap() const;
 
 	// Returns the index of this layer in the parent tile map
 	int32 GetLayerIndex() const;
@@ -232,6 +234,17 @@ public:
 
 	// Returns the transform for the given packed flag index (0..7)
 	static FTransform GetTileTransform(int32 FlagIndex);
+
+	// Returns the number of occupied cells in the layer
+	int32 GetNumOccupiedCells() const;
+
+	// Returns the raw pointer to the allocated cells, only for use when rendering the tile map and should never be held onto
+	const FPaperTileInfo* PRIVATE_GetAllocatedCells() const
+	{
+		return AllocatedCells.GetData();
+	}
+
 protected:
 	void ReallocateAndCopyMap();
+
 };
diff --git a/Engine/Plugins/2D/Paper2D/Source/Paper2D/Classes/PaperTileMapComponent.h b/Engine/Plugins/2D/Paper2D/Source/Paper2D/Classes/PaperTileMapComponent.h
--- a/Engine/Plugins/2D/Paper2D/Source/Paper2D/Classes/PaperTileMapComponent.h
+++ b/Engine/Plugins/2D/Paper2D/Source/Paper2D/Classes/PaperTileMapComponent.h
@@ -7,6 +7,8 @@
 #include "PaperTileMapComponent.generated.h"
 
 class FPaperTileMapRenderSceneProxy;
+struct FSpriteRenderSection;
+struct FPaperSpriteVertex;
 
 /**
  * A component that handles rendering and collision for a single instance of a UPaperTileMap asset.
@@ -50,7 +52,7 @@ private:
 	UMaterialInterface* Material_DEPRECATED;
 
 	UPROPERTY()
-	TArray<class UPaperTileLayer*> TileLayers_DEPRECATED;
+	TArray<UPaperTileLayer*> TileLayers_DEPRECATED;
 
 	// The color of the tile map (multiplied with the per-layer color and passed to the material as a vertex color)
 	UPROPERTY(EditAnywhere, Category=Materials)
@@ -94,7 +96,7 @@ public:
 protected:
 	friend FPaperTileMapRenderSceneProxy;
 
-	void RebuildRenderData(FPaperTileMapRenderSceneProxy* Proxy);
+	void RebuildRenderData(TArray<FSpriteRenderSection>& Sections, TArray<FPaperSpriteVertex>& Vertices);
 
 public:
 	// UObject interface
@@ -161,7 +163,7 @@ public:
 	// Creates and adds a new layer to the tile map
 	// Note: This will only work on components that own their own tile map (OwnsTileMap returns true), you cannot modify standalone tile map assets
 	UFUNCTION(BlueprintCallable, Category="Sprite")
-	class UPaperTileLayer* AddNewLayer();
+	UPaperTileLayer* AddNewLayer();
 
 	// Gets the tile map global color multiplier (multiplied with the per-layer color and passed to the material as a vertex color)
 	UFUNCTION(BlueprintPure, Category="Sprite")
diff --git a/Engine/Plugins/2D/Paper2D/Source/Paper2D/Classes/SpriteDrawCall.h b/Engine/Plugins/2D/Paper2D/Source/Paper2D/Classes/SpriteDrawCall.h
--- a/Engine/Plugins/2D/Paper2D/Source/Paper2D/Classes/SpriteDrawCall.h
+++ b/Engine/Plugins/2D/Paper2D/Source/Paper2D/Classes/SpriteDrawCall.h
@@ -12,16 +12,16 @@ struct FSpriteDrawCallRecord
 public:
 	GENERATED_USTRUCT_BODY()
 
-	UPROPERTY(Category=Sprite, EditAnywhere)
+	UPROPERTY()
 	FVector Destination;
 
-	UPROPERTY(Category=Sprite, EditAnywhere)
+	UPROPERTY()
 	UTexture* BaseTexture;
 
 	FAdditionalSpriteTextureArray AdditionalTextures;
 
-	UPROPERTY(Category=Sprite, EditAnywhere)
-	FLinearColor Color;
+	UPROPERTY()
+	FColor Color;
 
 	//@TODO: The rest of this struct doesn't need to be properties either, but has to be due to serialization for now
 	// Render triangle list (stored as loose vertices)
@@ -31,7 +31,7 @@ public:
 
 	FSpriteDrawCallRecord()
 		: BaseTexture(nullptr)
-		, Color(FLinearColor::White)
+		, Color(FColor::White)
 	{
 	}
 
diff --git a/Engine/Plugins/2D/Paper2D/Source/Paper2D/Private/PaperRenderSceneProxy.cpp b/Engine/Plugins/2D/Paper2D/Source/Paper2D/Private/PaperRenderSceneProxy.cpp
--- a/Engine/Plugins/2D/Paper2D/Source/Paper2D/Private/PaperRenderSceneProxy.cpp
+++ b/Engine/Plugins/2D/Paper2D/Source/Paper2D/Private/PaperRenderSceneProxy.cpp
@@ -15,6 +15,8 @@
 
 DECLARE_CYCLE_STAT(TEXT("Get Batch Mesh"), STAT_PaperRender_GetBatchMesh, STATGROUP_Paper2D);
 DECLARE_CYCLE_STAT(TEXT("Get New Batch Meshes"), STAT_PaperRender_GetNewBatchMeshes, STATGROUP_Paper2D);
+DECLARE_CYCLE_STAT(TEXT("Convert Batches"), STAT_PaperRender_ConvertBatches, STATGROUP_Paper2D);
+DECLARE_CYCLE_STAT(TEXT("SpriteProxy GDME"), STAT_PaperRenderSceneProxy_GetDynamicMeshElements, STATGROUP_Paper2D);
 
 //////////////////////////////////////////////////////////////////////////
 // FPaperSpriteVertex
@@ -30,23 +32,85 @@ void FPaperSpriteVertex::SetTangentsFromPaperAxes()
 	PackedNormalZ.Vector.W = (GetBasisDeterminantSign(PaperAxisX, PaperAxisY, PaperAxisZ) < 0.0f) ? 0 : 255;
 }
 
+
 //////////////////////////////////////////////////////////////////////////
 // FPaperSpriteVertexBuffer
 
-/** A dummy vertex buffer used to give the FPaperSpriteVertexFactory something to reference as a stream source. */
-class FPaperSpriteVertexBuffer : public FVertexBuffer
+class FDummyResourceArrayWrapper : public FResourceArrayInterface
 {
 public:
-	virtual void InitRHI() override
+	FDummyResourceArrayWrapper(void* InData, uint32 InSize)
+		: Data(InData)
+		, Size(InSize)
 	{
-		FRHIResourceCreateInfo CreateInfo;
-		VertexBufferRHI = RHICreateVertexBuffer(sizeof(FPaperSpriteVertex), BUF_Static, CreateInfo);
 	}
+
+	virtual const void* GetResourceData() const override { return Data; }
+	virtual uint32 GetResourceDataSize() const override { return Size; }
+	virtual void Discard() override { }
+	virtual bool IsStatic() const override { return false; }
+	virtual bool GetAllowCPUAccess() const override { return false; }
+	virtual void SetAllowCPUAccess(bool bInNeedsCPUAccess) override { }
+
+private:
+	void* Data;
+	uint32 Size;
 };
-static TGlobalResource<FPaperSpriteVertexBuffer> GDummyMaterialSpriteVertexBuffer;
+
+void FPaperSpriteVertexBuffer::InitRHI()
+{
+	const uint32 SizeInBytes = Vertices.Num() * sizeof(FPaperSpriteVertex);
+
+	FDummyResourceArrayWrapper Wrapper(Vertices.GetData(), SizeInBytes);
+	FRHIResourceCreateInfo CreateInfo(&Wrapper);
+	VertexBufferRHI = RHICreateVertexBuffer(SizeInBytes, BUF_Static, CreateInfo);
+
+	Vertices.Empty();
+}
+
+//////////////////////////////////////////////////////////////////////////
+// FPaperSpriteVertexFactory
+
+FPaperSpriteVertexFactory::FPaperSpriteVertexFactory()
+{
+}
+
+void FPaperSpriteVertexFactory::Init(const FPaperSpriteVertexBuffer* InVertexBuffer)
+{
+	check(!IsInRenderingThread());
+
+	ENQUEUE_UNIQUE_RENDER_COMMAND_TWOPARAMETER(
+		InitPaperSpriteVertexFactory,
+		FPaperSpriteVertexFactory*,VertexFactory,this,
+		const FPaperSpriteVertexBuffer*,VertexBuffer,InVertexBuffer,
+	{
+		// Initialize the vertex factory's stream components.
+		DataType NewData;
+		NewData.PositionComponent = STRUCTMEMBER_VERTEXSTREAMCOMPONENT(VertexBuffer,FPaperSpriteVertex,Position,VET_Float3);
+		NewData.TangentBasisComponents[0] = STRUCTMEMBER_VERTEXSTREAMCOMPONENT(VertexBuffer,FPaperSpriteVertex,TangentX,VET_PackedNormal);
+		NewData.TangentBasisComponents[1] = STRUCTMEMBER_VERTEXSTREAMCOMPONENT(VertexBuffer,FPaperSpriteVertex,TangentZ,VET_PackedNormal);
+		NewData.ColorComponent = STRUCTMEMBER_VERTEXSTREAMCOMPONENT(VertexBuffer,FPaperSpriteVertex,Color,VET_Color);
+		NewData.TextureCoordinates.Add(FVertexStreamComponent(VertexBuffer, STRUCT_OFFSET(FPaperSpriteVertex,TexCoords), sizeof(FPaperSpriteVertex), VET_Float2));
+		VertexFactory->SetData(NewData);
+	});
+}
+
+//////////////////////////////////////////////////////////////////////////
+// FDummyPaperSpriteVertexBuffer
+
+/** A dummy vertex buffer used to give the FPaperSpriteVertexFactoryDummy something to reference as a stream source. */
+class FDummyPaperSpriteVertexBuffer : public FPaperSpriteVertexBuffer
+{
+public:
+	FDummyPaperSpriteVertexBuffer()
+	{
+		Vertices.Emplace(FVector::ZeroVector, FVector2D::ZeroVector, FLinearColor::Black);
+	}
+};
+static TGlobalResource<FDummyPaperSpriteVertexBuffer> GDummyMaterialSpriteVertexBuffer;
 
 /** The vertex factory used to draw paper sprites. */
-class FPaperSpriteVertexFactory : public FLocalVertexFactory
+class FPaperSpriteVertexFactoryDummy : public FLocalVertexFactory
 {
 public:
 	void AllocateStuff()
@@ -88,15 +152,12 @@ public:
 		SetData(Data);
 	}
 
-	FPaperSpriteVertexFactory()
+	FPaperSpriteVertexFactoryDummy()
 	{
 		AllocateStuff();
 	}
 };
 
-//@TODO: Figure out how to do this properly at module load time
-//static TGlobalResource<FPaperSpriteVertexFactory> GPaperSpriteVertexFactory;
-
 //////////////////////////////////////////////////////////////////////////
 // FSpriteTextureOverrideRenderProxy
 
@@ -218,6 +279,9 @@ FPaperRenderSceneProxy::~FPaperRenderSceneProxy()
 		Batcher->UnregisterManagedProxy(this);
 	}
 #endif
+
+	VertexBuffer.ReleaseResource();
+	MyVertexFactory.ReleaseResource();
 }
 
 void FPaperRenderSceneProxy::CreateRenderThreadResources()
@@ -232,7 +296,7 @@ void FPaperRenderSceneProxy::CreateRenderThreadResources()
 
 void FPaperRenderSceneProxy::GetDynamicMeshElements(const TArray<const FSceneView*>& Views, const FSceneViewFamily& ViewFamily, uint32 VisibilityMap, FMeshElementCollector& Collector) const
 {
-	QUICK_SCOPE_CYCLE_COUNTER(STAT_PaperRenderSceneProxy_GetDynamicMeshElements);
+	SCOPE_CYCLE_COUNTER(STAT_PaperRenderSceneProxy_GetDynamicMeshElements);
 	checkSlow(IsInRenderingThread());
 
 	const FEngineShowFlags& EngineShowFlags = ViewFamily.EngineShowFlags;
@@ -257,26 +321,9 @@ void FPaperRenderSceneProxy::GetDynamicMeshElements(const TArray<const FSceneVie
 			{
 				const FSceneView* View = Views[ViewIndex];
 
-				// Set the selection/hover color from the current engine setting.
-				FLinearColor OverrideColor;
-				bool bUseOverrideColor = false;
-
-				const bool bShowAsSelected = !(GIsEditor && EngineShowFlags.Selection) || IsSelected();
-				if (bShowAsSelected || IsHovered())
-				{
-					bUseOverrideColor = true;
-					OverrideColor = GetSelectionColor(FLinearColor::White, bShowAsSelected, IsHovered());
-				}
-
-				bUseOverrideColor = false;
-
-				// Sprites of locked actors draw in red.
-				//FLinearColor LevelColorToUse = IsSelected() ? ColorToUse : (FLinearColor)LevelColor;
-				//FLinearColor PropertyColorToUse = PropertyColor;
-
 #if TEST_BATCHING
 #else
-				GetDynamicMeshElementsForView(View, ViewIndex, bUseOverrideColor, OverrideColor, Collector);
+				GetDynamicMeshElementsForView(View, ViewIndex, Collector);
 #endif
 			}
 		}
@@ -345,22 +392,21 @@ FVertexFactory* FPaperRenderSceneProxy::GetPaperSpriteVertexFactory() const
 		FPaperSpriteVertex::SetTangentsFromPaperAxes();
 	}
 
-	static TGlobalResource<FPaperSpriteVertexFactory> GPaperSpriteVertexFactory;
+	static TGlobalResource<FPaperSpriteVertexFactoryDummy> GPaperSpriteVertexFactory;
 	return &GPaperSpriteVertexFactory;
 }
 
-void FPaperRenderSceneProxy::GetDynamicMeshElementsForView(const FSceneView* View, int32 ViewIndex, bool bUseOverrideColor, const FLinearColor& OverrideColor, FMeshElementCollector& Collector) const
+void FPaperRenderSceneProxy::GetDynamicMeshElementsForView(const FSceneView* View, int32 ViewIndex, FMeshElementCollector& Collector) const
 {
 	if (Material != nullptr)
 	{
-		GetBatchMesh(View, bUseOverrideColor, OverrideColor, Material, BatchedSprites, ViewIndex, Collector);
+		GetBatchMesh(View, Material, BatchedSprites, ViewIndex, Collector);
 	}
-	GetNewBatchMeshes(View, bUseOverrideColor, OverrideColor, ViewIndex, Collector);
+	GetNewBatchMeshes(View, ViewIndex, Collector);
 }
 
-void FPaperRenderSceneProxy::GetNewBatchMeshes(const FSceneView* View, bool bUseOverrideColor, const FLinearColor& OverrideColor, int32 ViewIndex, FMeshElementCollector& Collector) const
+void FPaperRenderSceneProxy::GetNewBatchMeshes(const FSceneView* View, int32 ViewIndex, FMeshElementCollector& Collector) const
 {
-	//@TODO: Doesn't support OverrideColor yet
 	if (BatchedSections.Num() == 0)
 	{
 		return;
@@ -369,16 +415,14 @@ void FPaperRenderSceneProxy::GetNewBatchMeshes(const FSceneView* View, bool bUse
 	SCOPE_CYCLE_COUNTER(STAT_PaperRender_GetNewBatchMeshes);
 
 	const uint8 DPG = GetDepthPriorityGroup(View);
-	FVertexFactory* VertexFactory = GetPaperSpriteVertexFactory();
 	const bool bIsWireframeView = View->Family->EngineShowFlags.Wireframe;
 
 	for (const FSpriteRenderSection& Batch : BatchedSections)
 	{
-		const FLinearColor EffectiveWireframeColor = (Batch.Material->GetBlendMode() != BLEND_Opaque) ? WireframeColor : FLinearColor::Green;
-		FMaterialRenderProxy* ParentMaterialProxy = Batch.Material->GetRenderProxy((View->Family->EngineShowFlags.Selection) && IsSelected(), IsHovered());
-
 		if (Batch.IsValid())
 		{
+			FMaterialRenderProxy* ParentMaterialProxy = Batch.Material->GetRenderProxy((View->Family->EngineShowFlags.Selection) && IsSelected(), IsHovered());
+
 			FMeshBatch& Mesh = Collector.AllocateMesh();
 
 			Mesh.bCanApplyViewModeOverrides = true;
@@ -387,6 +431,8 @@ void FPaperRenderSceneProxy::GetNewBatchMeshes(const FSceneView* View, bool bUse
 			// Implementing our own wireframe coloring as the automatic one (controlled by Mesh.bCanApplyViewModeOverrides) only supports per-FPrimitiveSceneProxy WireframeColor
 			if (bIsWireframeView)
 			{
+				const FLinearColor EffectiveWireframeColor = (Batch.Material->GetBlendMode() != BLEND_Opaque) ? WireframeColor : FLinearColor::Green;
+
 				auto WireframeMaterialInstance = new FColoredMaterialRenderProxy(
 					GEngine->WireframeMaterial->GetRenderProxy(IsSelected(), IsHovered()),
 					GetSelectionColor(EffectiveWireframeColor, IsSelected(), IsHovered(), false)
@@ -405,12 +451,7 @@ void FPaperRenderSceneProxy::GetNewBatchMeshes(const FSceneView* View, bool bUse
 			FSpriteTextureOverrideRenderProxy* TextureOverrideMaterialProxy = new FSpriteTextureOverrideRenderProxy(ParentMaterialProxy, Batch.BaseTexture, Batch.AdditionalTextures UE_EXPAND_IF_WITH_EDITOR(, TextureOverrideList));
 			Collector.RegisterOneFrameMaterialProxy(TextureOverrideMaterialProxy);
 
-			Mesh.UseDynamicData = true;
-			// Note: memory referenced by Collector must be valid as long as the Collector is around, see Collector.AllocateOneFrameResource
-			Mesh.DynamicVertexData = BatchedVertices.GetData();
-			Mesh.DynamicVertexStride = sizeof(FPaperSpriteVertex);
-
-			Mesh.VertexFactory = VertexFactory;
+			Mesh.VertexFactory = &MyVertexFactory;
 			Mesh.LCI = nullptr;
 			Mesh.ReverseCulling = IsLocalToWorldDeterminantNegative() ? true : false;
 			Mesh.CastShadow = bCastShadow;
@@ -441,7 +482,7 @@ public:
 	TArray<FPaperSpriteVertex, TInlineAllocator<6> > Vertices;
 };
 
-void FPaperRenderSceneProxy::GetBatchMesh(const FSceneView* View, bool bUseOverrideColor, const FLinearColor& OverrideColor, class UMaterialInterface* BatchMaterial, const TArray<FSpriteDrawCallRecord>& Batch, int32 ViewIndex, FMeshElementCollector& Collector) const
+void FPaperRenderSceneProxy::GetBatchMesh(const FSceneView* View, class UMaterialInterface* BatchMaterial, const TArray<FSpriteDrawCallRecord>& Batch, int32 ViewIndex, FMeshElementCollector& Collector) const
 {
 	SCOPE_CYCLE_COUNTER(STAT_PaperRender_GetBatchMesh);
 
@@ -449,13 +490,12 @@ void FPaperRenderSceneProxy::GetBatchMesh(const FSceneView* View, bool bUseOverr
 	FVertexFactory* VertexFactory = GetPaperSpriteVertexFactory();
 
 	const bool bIsWireframeView = View->Family->EngineShowFlags.Wireframe;
-	const FLinearColor EffectiveWireframeColor = (BatchMaterial->GetBlendMode() != BLEND_Opaque) ? WireframeColor : FLinearColor::Green;
 
 	for (const FSpriteDrawCallRecord& Record : Batch)
 	{
 		if (Record.IsValid())
 		{
-			const FLinearColor SpriteColor = bUseOverrideColor ? OverrideColor : Record.Color;
+			const FColor SpriteColor(Record.Color);
 			const FVector EffectiveOrigin = Record.Destination;
 
 			FPaperVertexArray& VertexArray = Collector.AllocateOneFrameResource<FPaperVertexArray>();
@@ -485,6 +525,8 @@ void FPaperRenderSceneProxy::GetBatchMesh(const FSceneView* View, bool bUseOverr
 			// Implementing our own wireframe coloring as the automatic one (controlled by Mesh.bCanApplyViewModeOverrides) only supports per-FPrimitiveSceneProxy WireframeColor
 			if (bIsWireframeView)
 			{
+				const FLinearColor EffectiveWireframeColor = (BatchMaterial->GetBlendMode() != BLEND_Opaque) ? WireframeColor : FLinearColor::Green;
+
 				auto WireframeMaterialInstance = new FColoredMaterialRenderProxy(
 					GEngine->WireframeMaterial->GetRenderProxy(IsSelected(), IsHovered()),
 					GetSelectionColor(EffectiveWireframeColor, IsSelected(), IsHovered(), false)
@@ -566,11 +608,6 @@ FPrimitiveViewRelevance FPaperRenderSceneProxy::GetViewRelevance(const FSceneVie
 	return Result;
 }
 
-void FPaperRenderSceneProxy::OnTransformChanged()
-{
-	Origin = GetLocalToWorld().GetOrigin();
-}
-
 uint32 FPaperRenderSceneProxy::GetMemoryFootprint() const
 {
 	return sizeof(*this) + GetAllocatedSize();
@@ -631,4 +668,36 @@ void FPaperRenderSceneProxy::SetTransientTextureOverride_RenderThread(const UTex
 		TextureOverrideList.Remove(InTextureToModifyOverrideFor);
 	}
 }
-#endif
\ No newline at end of file
+#endif
+
+void FPaperRenderSceneProxy::ConvertBatchesToNewStyle(TArray<FSpriteDrawCallRecord>& SourceBatches)
+{
+	SCOPE_CYCLE_COUNTER(STAT_PaperRender_ConvertBatches);
+
+	VertexBuffer.Vertices.Reset();
+	BatchedSections.Reset();
+
+	for (const FSpriteDrawCallRecord& SourceBatch : SourceBatches)
+	{
+		if (SourceBatch.IsValid())
+		{
+			FSpriteRenderSection& DestBatch = *new (BatchedSections) FSpriteRenderSection();
+
+			DestBatch.BaseTexture = SourceBatch.BaseTexture;
+			DestBatch.AdditionalTextures = SourceBatch.AdditionalTextures;
+			DestBatch.Material = Material;
+
+			DestBatch.AddTriangles(SourceBatch, VertexBuffer.Vertices);
+		}
+	}
+
+	if (VertexBuffer.Vertices.Num() > 0)
+	{
+		// Init the vertex factory
+		MyVertexFactory.Init(&VertexBuffer);
+
+		// Enqueue initialization of render resources
+		BeginInitResource(&VertexBuffer);
+		BeginInitResource(&MyVertexFactory);
+	}
+}
\ No newline at end of file
diff --git a/Engine/Plugins/2D/Paper2D/Source/Paper2D/Private/PaperRenderSceneProxy.h b/Engine/Plugins/2D/Paper2D/Source/Paper2D/Private/PaperRenderSceneProxy.h
--- a/Engine/Plugins/2D/Paper2D/Source/Paper2D/Private/PaperRenderSceneProxy.h
+++ b/Engine/Plugins/2D/Paper2D/Source/Paper2D/Private/PaperRenderSceneProxy.h
@@ -29,7 +29,7 @@ struct PAPER2D_API FPaperSpriteVertex
 
 	FPaperSpriteVertex() {}
 
-	FPaperSpriteVertex(const FVector& InPosition, const FVector2D& InTextureCoordinate, const FLinearColor& InColor)
+	FPaperSpriteVertex(const FVector& InPosition, const FVector2D& InTextureCoordinate, const FColor& InColor)
 		: Position(InPosition)
 		, TangentX(PackedNormalX)
 		, TangentZ(PackedNormalZ)
@@ -43,6 +43,29 @@ struct PAPER2D_API FPaperSpriteVertex
 	static FPackedNormal PackedNormalZ;
 };
 
+//////////////////////////////////////////////////////////////////////////
+// FPaperSpriteVertexBuffer
+
+class FPaperSpriteVertexBuffer : public FVertexBuffer
+{
+public:
+	TArray<FPaperSpriteVertex> Vertices;
+
+	// FRenderResource interface
+	virtual void InitRHI() override;
+	// End of FRenderResource interface
+};
+
+//////////////////////////////////////////////////////////////////////////
+// FPaperSpriteVertexFactory
+
+class FPaperSpriteVertexFactory : public FLocalVertexFactory
+{
+public:
+	FPaperSpriteVertexFactory();
+	void Init(const FPaperSpriteVertexBuffer* VertexBuffer);
+};
+
 //////////////////////////////////////////////////////////////////////////
 // FSpriteRenderSection
 
@@ -56,7 +79,9 @@ struct PAPER2D_API FSpriteRenderSection
 	int32 NumVertices;
 
 	FSpriteRenderSection()
-		: VertexOffset(INDEX_NONE)
+		: Material(nullptr)
+		, BaseTexture(nullptr)
+		, VertexOffset(INDEX_NONE)
 		, NumVertices(0)
 	{
 	}
@@ -68,10 +93,9 @@ struct PAPER2D_API FSpriteRenderSection
 
 	bool IsValid() const
 	{
-		return (NumVertices > 0) && (GetBaseTextureResource() != nullptr);
+		return (Material != nullptr) && (NumVertices > 0) && (GetBaseTextureResource() != nullptr);
 	}
 
-#if 1
 	template <typename SourceArrayType>
 	void AddTriangles(const FSpriteDrawCallRecord& Record, SourceArrayType& Vertices)
 	{
@@ -92,15 +116,24 @@ struct PAPER2D_API FSpriteRenderSection
 		NumVertices += NumNewVerts;
 		Vertices.Reserve(Vertices.Num() + NumNewVerts);
 
+		const FColor VertColor(Record.Color);
 		for (const FVector4& SourceVert : Record.RenderVerts)
 		{
 			const FVector Pos((PaperAxisX * SourceVert.X) + (PaperAxisY * SourceVert.Y) + Record.Destination);
 			const FVector2D UV(SourceVert.Z, SourceVert.W);
 
-			new (Vertices) FPaperSpriteVertex(Pos, UV, Record.Color);
+			new (Vertices) FPaperSpriteVertex(Pos, UV, VertColor);
 		}
 	}
-#endif
+
+	template <typename SourceArrayType>
+	inline void AddVertex(float X, float Y, float U, float V, const FVector& Origin, const FColor& Color, SourceArrayType& Vertices)
+	{
+		const FVector Pos((PaperAxisX * X) + (PaperAxisY * Y) + Origin);
+
+		new (Vertices) FPaperSpriteVertex(Pos, FVector2D(U, V), Color);
+		++NumVertices;
+	}
 };
 
 //////////////////////////////////////////////////////////////////////////
@@ -110,12 +143,11 @@ class PAPER2D_API FPaperRenderSceneProxy : public FPrimitiveSceneProxy
 {
 public:
 	FPaperRenderSceneProxy(const UPrimitiveComponent* InComponent);
-	~FPaperRenderSceneProxy();
+	virtual ~FPaperRenderSceneProxy();
 
 	// FPrimitiveSceneProxy interface.
 	virtual void GetDynamicMeshElements(const TArray<const FSceneView*>& Views, const FSceneViewFamily& ViewFamily, uint32 VisibilityMap, FMeshElementCollector& Collector) const override;
 	virtual FPrimitiveViewRelevance GetViewRelevance(const FSceneView* View) override;
-	virtual void OnTransformChanged() override;
 	virtual uint32 GetMemoryFootprint() const override;
 	virtual bool CanBeOccluded() const override;
 	virtual void CreateRenderThreadResources() override;
@@ -134,10 +166,10 @@ public:
 	}
 
 protected:
-	virtual void GetDynamicMeshElementsForView(const FSceneView* View, int32 ViewIndex, bool bUseOverrideColor, const FLinearColor& OverrideColor, FMeshElementCollector& Collector) const;
+	virtual void GetDynamicMeshElementsForView(const FSceneView* View, int32 ViewIndex, FMeshElementCollector& Collector) const;
 
-	void GetBatchMesh(const FSceneView* View, bool bUseOverrideColor, const FLinearColor& OverrideColor, class UMaterialInterface* BatchMaterial, const TArray<FSpriteDrawCallRecord>& Batch, int32 ViewIndex, FMeshElementCollector& Collector) const;
-	void GetNewBatchMeshes(const FSceneView* View, bool bUseOverrideColor, const FLinearColor& OverrideColor, int32 ViewIndex, FMeshElementCollector& Collector) const;
+	void GetBatchMesh(const FSceneView* View, UMaterialInterface* BatchMaterial, const TArray<FSpriteDrawCallRecord>& Batch, int32 ViewIndex, FMeshElementCollector& Collector) const;
+	void GetNewBatchMeshes(const FSceneView* View, int32 ViewIndex, FMeshElementCollector& Collector) const;
 
 	bool IsCollisionView(const FEngineShowFlags& EngineShowFlags, bool& bDrawSimpleCollision, bool& bDrawComplexCollision) const;
 
@@ -145,13 +177,19 @@ protected:
 
 	FVertexFactory* GetPaperSpriteVertexFactory() const;
 
+	void ConvertBatchesToNewStyle(TArray<FSpriteDrawCallRecord>& SourceBatches);
+
 protected:
-	TArray<FPaperSpriteVertex> BatchedVertices;
+	// New style
+	FPaperSpriteVertexBuffer VertexBuffer;
+	FPaperSpriteVertexFactory MyVertexFactory;
 	TArray<FSpriteRenderSection> BatchedSections;
 
+	// Old style
 	TArray<FSpriteDrawCallRecord> BatchedSprites;
 	class UMaterialInterface* Material;
-	FVector Origin;
+
+	//
 	AActor* Owner;
 	UBodySetup* BodySetup;
 
diff --git a/Engine/Plugins/2D/Paper2D/Source/Paper2D/Private/PaperSpriteSceneProxy.cpp b/Engine/Plugins/2D/Paper2D/Source/Paper2D/Private/PaperSpriteSceneProxy.cpp
--- a/Engine/Plugins/2D/Paper2D/Source/Paper2D/Private/PaperSpriteSceneProxy.cpp
+++ b/Engine/Plugins/2D/Paper2D/Source/Paper2D/Private/PaperSpriteSceneProxy.cpp
@@ -130,15 +130,15 @@ void FPaperSpriteSceneProxy::SetSprite_RenderThread(const FSpriteDrawCallRecord&
 	}
 }
 
-void FPaperSpriteSceneProxy::GetDynamicMeshElementsForView(const FSceneView* View, int32 ViewIndex, bool bUseOverrideColor, const FLinearColor& OverrideColor, FMeshElementCollector& Collector) const
+void FPaperSpriteSceneProxy::GetDynamicMeshElementsForView(const FSceneView* View, int32 ViewIndex, FMeshElementCollector& Collector) const
 {
 	if (Material != nullptr)
 	{
-		GetBatchMesh(View, bUseOverrideColor, OverrideColor, Material, BatchedSprites, ViewIndex, Collector);
+		GetBatchMesh(View, Material, BatchedSprites, ViewIndex, Collector);
 	}
 
 	if ((AlternateMaterial != nullptr) && (AlternateBatchedSprites.Num() > 0))
 	{
-		GetBatchMesh(View, bUseOverrideColor, OverrideColor, AlternateMaterial, AlternateBatchedSprites, ViewIndex, Collector);
+		GetBatchMesh(View, AlternateMaterial, AlternateBatchedSprites, ViewIndex, Collector);
 	}
 }
\ No newline at end of file
diff --git a/Engine/Plugins/2D/Paper2D/Source/Paper2D/Private/PaperSpriteSceneProxy.h b/Engine/Plugins/2D/Paper2D/Source/Paper2D/Private/PaperSpriteSceneProxy.h
--- a/Engine/Plugins/2D/Paper2D/Source/Paper2D/Private/PaperSpriteSceneProxy.h
+++ b/Engine/Plugins/2D/Paper2D/Source/Paper2D/Private/PaperSpriteSceneProxy.h
@@ -24,7 +24,7 @@ public:
 protected:
 
 	// FPaperRenderSceneProxy interface
-	virtual void GetDynamicMeshElementsForView(const FSceneView* View, int32 ViewIndex, bool bUseOverrideColor, const FLinearColor& OverrideColor, FMeshElementCollector& Collector) const override;
+	virtual void GetDynamicMeshElementsForView(const FSceneView* View, int32 ViewIndex, FMeshElementCollector& Collector) const override;
 	// End of FPaperRenderSceneProxy interface
 
 protected:
diff --git a/Engine/Plugins/2D/Paper2D/Source/Paper2D/Private/PaperTileLayer.cpp b/Engine/Plugins/2D/Paper2D/Source/Paper2D/Private/PaperTileLayer.cpp
--- a/Engine/Plugins/2D/Paper2D/Source/Paper2D/Private/PaperTileLayer.cpp
+++ b/Engine/Plugins/2D/Paper2D/Source/Paper2D/Private/PaperTileLayer.cpp
@@ -296,3 +296,18 @@ FTransform UPaperTileLayer::GetTileTransform(int32 FlagIndex)
 	checkSlow((FlagIndex >= 0) && (FlagIndex < 8));
 	return TilePermutationTransforms[FlagIndex];
 }
+
+int32 UPaperTileLayer::GetNumOccupiedCells() const
+{
+	int32 NumOccupiedCells = 0;
+
+	for (const FPaperTileInfo& TileInfo : AllocatedCells)
+	{
+		if (TileInfo.IsValid())
+		{
+			++NumOccupiedCells;
+		}
+	}
+
+	return NumOccupiedCells;
+}
diff --git a/Engine/Plugins/2D/Paper2D/Source/Paper2D/Private/PaperTileMapComponent.cpp b/Engine/Plugins/2D/Paper2D/Source/Paper2D/Private/PaperTileMapComponent.cpp
--- a/Engine/Plugins/2D/Paper2D/Source/Paper2D/Private/PaperTileMapComponent.cpp
+++ b/Engine/Plugins/2D/Paper2D/Source/Paper2D/Private/PaperTileMapComponent.cpp
@@ -10,6 +10,8 @@
 
 #define LOCTEXT_NAMESPACE "Paper2D"
 
+DECLARE_CYCLE_STAT(TEXT("Rebuild Tile Map"), STAT_PaperRender_TileMapRebuild, STATGROUP_Paper2D);
+
 //////////////////////////////////////////////////////////////////////////
 // UPaperTileMapComponent
 
@@ -47,8 +49,15 @@ UPaperTileMapComponent::UPaperTileMapComponent(const FObjectInitializer& ObjectI
 
 FPrimitiveSceneProxy* UPaperTileMapComponent::CreateSceneProxy()
 {
-	FPaperTileMapRenderSceneProxy* Proxy = new FPaperTileMapRenderSceneProxy(this);
-	RebuildRenderData(Proxy);
+	SCOPE_CYCLE_COUNTER(STAT_PaperRender_TileMapRebuild);
+
+	TArray<FSpriteRenderSection>* Sections;
+	TArray<FPaperSpriteVertex>* Vertices;
+	FPaperTileMapRenderSceneProxy* Proxy = FPaperTileMapRenderSceneProxy::CreateTileMapProxy(this, /*out*/ Sections, /*out*/ Vertices);
+
+	RebuildRenderData(*Sections, *Vertices);
+
+	Proxy->FinishConstruction_GameThread();
 
 	return Proxy;
 }
@@ -212,10 +221,8 @@ const UObject* UPaperTileMapComponent::AdditionalStatObject() const
 	return nullptr;
 }
 
-void UPaperTileMapComponent::RebuildRenderData(FPaperTileMapRenderSceneProxy* Proxy)
+void UPaperTileMapComponent::RebuildRenderData(TArray<FSpriteRenderSection>& Sections, TArray<FPaperSpriteVertex>& Vertices)
 {
-	TArray<FSpriteDrawCallRecord> BatchedSprites;
-	
 	if (TileMap == nullptr)
 	{
 		return;
@@ -249,6 +256,30 @@ void UPaperTileMapComponent::RebuildRenderData(FPaperTileMapRenderSceneProxy* Pr
 
 	const float UnrealUnitsPerPixel = TileMap->GetUnrealUnitsPerPixel();
 
+
+	// Run thru the layers and estimate how big of an allocation we will need
+	int32 EstimatedNumVerts = 0;
+
+	for (int32 Z = TileMap->TileLayers.Num() - 1; Z >= 0; --Z)
+	{
+		UPaperTileLayer* Layer = TileMap->TileLayers[Z];
+
+		if ((Layer != nullptr) && (!bUseSingleLayer || (Z == UseSingleLayerIndex)))
+		{
+			const int32 NumOccupiedCells = Layer->GetNumOccupiedCells();
+			EstimatedNumVerts += 6 * NumOccupiedCells;
+		}
+	}
+
+ 	Vertices.Empty(EstimatedNumVerts);
+
+	UMaterialInterface* TileMapMaterial = GetMaterial(0);
+	if (TileMapMaterial == nullptr)
+	{
+		TileMapMaterial = UMaterial::GetDefaultMaterial(MD_Surface);
+	}
+
+	// Actual pass
 	for (int32 Z = TileMap->TileLayers.Num() - 1; Z >= 0; --Z)
 	{
 		UPaperTileLayer* Layer = TileMap->TileLayers[Z];
@@ -266,7 +297,7 @@ void UPaperTileMapComponent::RebuildRenderData(FPaperTileMapRenderSceneProxy* Pr
 			}
 		}
 
-		FLinearColor DrawColor = TileMapColor * Layer->GetLayerColor();
+		const FColor DrawColor(TileMapColor * Layer->GetLayerColor());
 
 #if WITH_EDITORONLY_DATA
 		if (!Layer->ShouldRenderInEditor())
@@ -275,7 +306,9 @@ void UPaperTileMapComponent::RebuildRenderData(FPaperTileMapRenderSceneProxy* Pr
 		}
 #endif
 
-		FSpriteDrawCallRecord* CurrentBatch = nullptr;
+		FSpriteRenderSection* CurrentBatch = nullptr;
+		FVector CurrentDestinationOrigin;
+		const FPaperTileInfo* CurrentCellPtr = Layer->PRIVATE_GetAllocatedCells();
 
 		for (int32 Y = 0; Y < TileMap->MapHeight; ++Y)
 		{
@@ -299,7 +332,7 @@ void UPaperTileMapComponent::RebuildRenderData(FPaperTileMapRenderSceneProxy* Pr
 
 			for (int32 X = 0; X < TileMap->MapWidth; ++X)
 			{
-				const FPaperTileInfo TileInfo = Layer->GetCell(X, Y);
+				const FPaperTileInfo& TileInfo = *CurrentCellPtr++;
 
 				// do stuff
 				const float TotalSeparation = (TileMap->SeparationPerLayer * Z) + (TileMap->SeparationPerTileX * X) + (TileMap->SeparationPerTileY * Y);
@@ -333,12 +366,13 @@ void UPaperTileMapComponent::RebuildRenderData(FPaperTileMapRenderSceneProxy* Pr
 
 					if ((SourceTexture != LastSourceTexture) || (CurrentBatch == nullptr))
 					{
-						CurrentBatch = (new (BatchedSprites) FSpriteDrawCallRecord());
+						CurrentBatch = new (Sections) FSpriteRenderSection();
 						CurrentBatch->BaseTexture = SourceTexture;
 						//CurrentBatch->AdditionalTextures = ?; //@TODO: PAPER2D: Need to add multi-texture support to tile sets / tile maps
 						// Probably also need to change the batch check here to TileSet changing to avoid checking each texture in the array
-						CurrentBatch->Color = DrawColor;
-						CurrentBatch->Destination = TopLeftCornerOfTile.ProjectOnTo(PaperAxisZ);
+						CurrentBatch->Material = TileMapMaterial;
+						CurrentBatch->VertexOffset = Vertices.Num();
+						CurrentDestinationOrigin = TopLeftCornerOfTile.ProjectOnTo(PaperAxisZ);
 					}
 
 					if (SourceTexture != LastSourceTexture)
@@ -371,8 +405,6 @@ void UPaperTileMapComponent::RebuildRenderData(FPaperTileMapRenderSceneProxy* Pr
 					SourceUV.X *= InverseTextureSize.X;
 					SourceUV.Y *= InverseTextureSize.Y;
 
-					FSpriteDrawCallRecord& NewTile = *CurrentBatch;
-
 					const float WX0 = FVector::DotProduct(TopLeftCornerOfTile, PaperAxisX);
 					const float WY0 = FVector::DotProduct(TopLeftCornerOfTile, PaperAxisY);
 
@@ -392,28 +424,22 @@ void UPaperTileMapComponent::RebuildRenderData(FPaperTileMapRenderSceneProxy* Pr
 					const FVector4 TopRight(WX0 + TileSizeWithFlip.X, WY0, UValues[UVIndex2], VValues[UVIndex2]);
 					const FVector4 TopLeft(WX0, WY0, UValues[UVIndex3], VValues[UVIndex3]);
 
-					new (NewTile.RenderVerts) FVector4(BottomLeft);
-					new (NewTile.RenderVerts) FVector4(TopRight);
-					new (NewTile.RenderVerts) FVector4(BottomRight);
+					CurrentBatch->AddVertex(BottomLeft.X, BottomLeft.Y, BottomLeft.Z, BottomLeft.W, CurrentDestinationOrigin, DrawColor, Vertices);
+					CurrentBatch->AddVertex(TopRight.X, TopRight.Y, TopRight.Z, TopRight.W, CurrentDestinationOrigin, DrawColor, Vertices);
+					CurrentBatch->AddVertex(BottomRight.X, BottomRight.Y, BottomRight.Z, BottomRight.W, CurrentDestinationOrigin, DrawColor, Vertices);
 
-					new (NewTile.RenderVerts) FVector4(BottomLeft);
-					new (NewTile.RenderVerts) FVector4(TopLeft);
-					new (NewTile.RenderVerts) FVector4(TopRight);
+					CurrentBatch->AddVertex(BottomLeft.X, BottomLeft.Y, BottomLeft.Z, BottomLeft.W, CurrentDestinationOrigin, DrawColor, Vertices);
+					CurrentBatch->AddVertex(TopLeft.X, TopLeft.Y, TopLeft.Z, TopLeft.W, CurrentDestinationOrigin, DrawColor, Vertices);
+					CurrentBatch->AddVertex(TopRight.X, TopRight.Y, TopRight.Z, TopRight.W, CurrentDestinationOrigin, DrawColor, Vertices);
 				}
 			}
 		}
 	}
 
 #if WITH_EDITOR
-	NumBatches = BatchedSprites.Num();
-	NumTriangles = 0;
-	for (const FSpriteDrawCallRecord& Batch : BatchedSprites)
-	{
-		NumTriangles += Batch.RenderVerts.Num() / 3;
-	}
+	NumBatches = Sections.Num();
+	NumTriangles = Vertices.Num() / 3;
 #endif
-
-	Proxy->SetBatchesHack(BatchedSprites);
 }
 
 void UPaperTileMapComponent::CreateNewOwnedTileMap()
diff --git a/Engine/Plugins/2D/Paper2D/Source/Paper2D/Private/PaperTileMapRenderSceneProxy.cpp b/Engine/Plugins/2D/Paper2D/Source/Paper2D/Private/PaperTileMapRenderSceneProxy.cpp
--- a/Engine/Plugins/2D/Paper2D/Source/Paper2D/Private/PaperTileMapRenderSceneProxy.cpp
+++ b/Engine/Plugins/2D/Paper2D/Source/Paper2D/Private/PaperTileMapRenderSceneProxy.cpp
@@ -46,6 +46,29 @@ FPaperTileMapRenderSceneProxy::FPaperTileMapRenderSceneProxy(const UPaperTileMap
 #endif
 }
 
+FPaperTileMapRenderSceneProxy* FPaperTileMapRenderSceneProxy::CreateTileMapProxy(const UPaperTileMapComponent* InComponent, TArray<FSpriteRenderSection>*& OutSections, TArray<FPaperSpriteVertex>*& OutVertices)
+{
+	FPaperTileMapRenderSceneProxy* NewProxy = new FPaperTileMapRenderSceneProxy(InComponent);
+
+	OutVertices = &(NewProxy->VertexBuffer.Vertices);
+	OutSections = &(NewProxy->BatchedSections);
+
+	return NewProxy;
+}
+
+void FPaperTileMapRenderSceneProxy::FinishConstruction_GameThread()
+{
+	if (VertexBuffer.Vertices.Num() > 0)
+	{
+		// Init the vertex factory
+		MyVertexFactory.Init(&VertexBuffer);
+
+		// Enqueue initialization of render resources
+		BeginInitResource(&VertexBuffer);
+		BeginInitResource(&MyVertexFactory);
+	}
+}
+
 void FPaperTileMapRenderSceneProxy::DrawBoundsForLayer(FPrimitiveDrawInterface* PDI, const FLinearColor& Color, int32 LayerIndex) const
 {
 	const FMatrix& LocalToWorld = GetLocalToWorld();
diff --git a/Engine/Plugins/2D/Paper2D/Source/Paper2D/Private/PaperTileMapRenderSceneProxy.h b/Engine/Plugins/2D/Paper2D/Source/Paper2D/Private/PaperTileMapRenderSceneProxy.h
--- a/Engine/Plugins/2D/Paper2D/Source/Paper2D/Private/PaperTileMapRenderSceneProxy.h
+++ b/Engine/Plugins/2D/Paper2D/Source/Paper2D/Private/PaperTileMapRenderSceneProxy.h
@@ -4,24 +4,28 @@
 
 #include "PaperRenderSceneProxy.h"
 
+class UPaperTileMapComponent;
+class UPaperTileMap;
+
 //////////////////////////////////////////////////////////////////////////
 // FPaperTileMapRenderSceneProxy
 
 class FPaperTileMapRenderSceneProxy : public FPaperRenderSceneProxy
 {
 public:
-	FPaperTileMapRenderSceneProxy(const class UPaperTileMapComponent* InComponent);
-
 	// FPrimitiveSceneProxy interface.
 	virtual void GetDynamicMeshElements(const TArray<const FSceneView*>& Views, const FSceneViewFamily& ViewFamily, uint32 VisibilityMap, FMeshElementCollector& Collector) const override;
 	// End of FPrimitiveSceneProxy interface.
 
-	void SetBatchesHack(TArray<struct FSpriteDrawCallRecord>& InBatchedSprites)
-	{
-		BatchedSprites = InBatchedSprites;
-	}
+	// Construct a tile map scene proxy
+	static FPaperTileMapRenderSceneProxy* CreateTileMapProxy(const UPaperTileMapComponent* InComponent, TArray<FSpriteRenderSection>*& OutSections, TArray<FPaperSpriteVertex>*& OutVertices);
+	
+	// Call this once the tile map sections/vertices are finished
+	void FinishConstruction_GameThread();
 
 protected:
+	FPaperTileMapRenderSceneProxy(const UPaperTileMapComponent* InComponent);
+
 	void DrawBoundsForLayer(FPrimitiveDrawInterface* PDI, const FLinearColor& Color, int32 LayerIndex) const;
 	void DrawNormalGridLines(FPrimitiveDrawInterface* PDI, const FLinearColor& Color, int32 LayerIndex) const;
 	void DrawStaggeredGridLines(FPrimitiveDrawInterface* PDI, const FLinearColor& Color, int32 LayerIndex) const;
@@ -36,7 +40,7 @@ protected:
 #endif
 
 	//@TODO: Not thread safe
-	const class UPaperTileMap* TileMap;
+	const UPaperTileMap* TileMap;
 
 	// The only layer to draw, or INDEX_NONE if the filter is unset
 	const int32 OnlyLayerIndex;
diff --git a/Engine/Plugins/2D/Paper2D/Source/Paper2D/Private/Terrain/PaperTerrainComponent.cpp b/Engine/Plugins/2D/Paper2D/Source/Paper2D/Private/Terrain/PaperTerrainComponent.cpp
--- a/Engine/Plugins/2D/Paper2D/Source/Paper2D/Private/Terrain/PaperTerrainComponent.cpp
+++ b/Engine/Plugins/2D/Paper2D/Source/Paper2D/Private/Terrain/PaperTerrainComponent.cpp
@@ -64,7 +64,7 @@ protected:
 	TArray<FPaperTerrainSpriteGeometry> DrawingData;
 protected:
 	// FPaperRenderSceneProxy interface
-	virtual void GetDynamicMeshElementsForView(const FSceneView* View, int32 ViewIndex, bool bUseOverrideColor, const FLinearColor& OverrideColor, FMeshElementCollector& Collector) const override;
+	virtual void GetDynamicMeshElementsForView(const FSceneView* View, int32 ViewIndex, FMeshElementCollector& Collector) const override;
 	// End of FPaperRenderSceneProxy interface
 };
 
@@ -81,7 +81,7 @@ FPaperTerrainSceneProxy::FPaperTerrainSceneProxy(const UPaperTerrainComponent* I
 	}
 }
 
-void FPaperTerrainSceneProxy::GetDynamicMeshElementsForView(const FSceneView* View, int32 ViewIndex, bool bUseOverrideColor, const FLinearColor& OverrideColor, FMeshElementCollector& Collector) const
+void FPaperTerrainSceneProxy::GetDynamicMeshElementsForView(const FSceneView* View, int32 ViewIndex, FMeshElementCollector& Collector) const
 {
 	SCOPE_CYCLE_COUNTER(STAT_TerrainSpline_GetDynamicMeshElements);
 
@@ -89,7 +89,7 @@ void FPaperTerrainSceneProxy::GetDynamicMeshElementsForView(const FSceneView* Vi
 	{
 		if (Batch.Material != nullptr)
 		{
-			GetBatchMesh(View, bUseOverrideColor, OverrideColor, Batch.Material, Batch.Records, ViewIndex, Collector);
+			GetBatchMesh(View, Batch.Material, Batch.Records, ViewIndex, Collector);
 		}
 	}
 }
