Paper2D: Fixed a -wshadow issue

diff --git a/Engine/Plugins/2D/Paper2D/Source/Paper2D/Private/PaperRenderSceneProxy.cpp b/Engine/Plugins/2D/Paper2D/Source/Paper2D/Private/PaperRenderSceneProxy.cpp
--- a/Engine/Plugins/2D/Paper2D/Source/Paper2D/Private/PaperRenderSceneProxy.cpp
+++ b/Engine/Plugins/2D/Paper2D/Source/Paper2D/Private/PaperRenderSceneProxy.cpp
@@ -82,15 +82,15 @@ void FPaperSpriteVertexFactory::Init(const FPaperSpriteVertexBuffer* InVertexBuf
 	ENQUEUE_UNIQUE_RENDER_COMMAND_TWOPARAMETER(
 		InitPaperSpriteVertexFactory,
 		FPaperSpriteVertexFactory*,VertexFactory,this,
-		const FPaperSpriteVertexBuffer*,VertexBuffer,InVertexBuffer,
+		const FPaperSpriteVertexBuffer*,VB,InVertexBuffer,
 	{
 		// Initialize the vertex factory's stream components.
 		DataType NewData;
-		NewData.PositionComponent = STRUCTMEMBER_VERTEXSTREAMCOMPONENT(VertexBuffer,FPaperSpriteVertex,Position,VET_Float3);
-		NewData.TangentBasisComponents[0] = STRUCTMEMBER_VERTEXSTREAMCOMPONENT(VertexBuffer,FPaperSpriteVertex,TangentX,VET_PackedNormal);
-		NewData.TangentBasisComponents[1] = STRUCTMEMBER_VERTEXSTREAMCOMPONENT(VertexBuffer,FPaperSpriteVertex,TangentZ,VET_PackedNormal);
-		NewData.ColorComponent = STRUCTMEMBER_VERTEXSTREAMCOMPONENT(VertexBuffer,FPaperSpriteVertex,Color,VET_Color);
-		NewData.TextureCoordinates.Add(FVertexStreamComponent(VertexBuffer, STRUCT_OFFSET(FPaperSpriteVertex,TexCoords), sizeof(FPaperSpriteVertex), VET_Float2));
+		NewData.PositionComponent = STRUCTMEMBER_VERTEXSTREAMCOMPONENT(VB,FPaperSpriteVertex,Position,VET_Float3);
+		NewData.TangentBasisComponents[0] = STRUCTMEMBER_VERTEXSTREAMCOMPONENT(VB,FPaperSpriteVertex,TangentX,VET_PackedNormal);
+		NewData.TangentBasisComponents[1] = STRUCTMEMBER_VERTEXSTREAMCOMPONENT(VB,FPaperSpriteVertex,TangentZ,VET_PackedNormal);
+		NewData.ColorComponent = STRUCTMEMBER_VERTEXSTREAMCOMPONENT(VB,FPaperSpriteVertex,Color,VET_Color);
+		NewData.TextureCoordinates.Add(FVertexStreamComponent(VB, STRUCT_OFFSET(FPaperSpriteVertex,TexCoords), sizeof(FPaperSpriteVertex), VET_Float2));
 		VertexFactory->SetData(NewData);
 	});
 }
