Tile Maps: Added a load-time fix up for corrupted tile maps that contain layers of the wrong size (ones that had a resize that was then undone, prior to CL# 2523792)

diff --git a/Engine/Plugins/2D/Paper2D/Source/Paper2D/Private/PaperTileMap.cpp b/Engine/Plugins/2D/Paper2D/Source/Paper2D/Private/PaperTileMap.cpp
--- a/Engine/Plugins/2D/Paper2D/Source/Paper2D/Private/PaperTileMap.cpp
+++ b/Engine/Plugins/2D/Paper2D/Source/Paper2D/Private/PaperTileMap.cpp
@@ -149,6 +149,15 @@ bool UPaperTileMap::CanEditChange(const UProperty* InProperty) const
 void UPaperTileMap::PostLoad()
 {
 	Super::PostLoad();
+
+	// Make sure that the layers are all of the right size (there was a bug at one point when undoing resizes that could cause the layers to get stuck at a bad size)
+	for (UPaperTileLayer* TileLayer : TileLayers)
+	{
+		TileLayer->ConditionalPostLoad();
+
+		TileLayer->ResizeMap(MapWidth, MapHeight);
+	}
+
 	ValidateSelectedLayerIndex();
 }
 
