Paper2D: More include-what-you-use and forward declaration improvements for faster iteration

diff --git a/Engine/Plugins/2D/Paper2D/Source/Paper2D/Classes/PaperTileMapComponent.h b/Engine/Plugins/2D/Paper2D/Source/Paper2D/Classes/PaperTileMapComponent.h
--- a/Engine/Plugins/2D/Paper2D/Source/Paper2D/Classes/PaperTileMapComponent.h
+++ b/Engine/Plugins/2D/Paper2D/Source/Paper2D/Classes/PaperTileMapComponent.h
@@ -3,6 +3,7 @@
 #pragma once
 
 #include "PaperSprite.h"
+#include "PaperTileMap.h"
 #include "PaperTileLayer.h"
 #include "PaperTileMapComponent.generated.h"
 
@@ -142,7 +143,7 @@ public:
 
 	/** Change the PaperTileMap used by this instance. */
 	UFUNCTION(BlueprintCallable, Category="Sprite")
-	virtual bool SetTileMap(class UPaperTileMap* NewTileMap);
+	virtual bool SetTileMap(UPaperTileMap* NewTileMap);
 
 	// Returns the size of the tile map
 	UFUNCTION(BlueprintCallable, Category="Sprite")
diff --git a/Engine/Plugins/2D/Paper2D/Source/Paper2DEditor/Classes/PaperFlipbookFactory.h b/Engine/Plugins/2D/Paper2D/Source/Paper2DEditor/Classes/PaperFlipbookFactory.h
--- a/Engine/Plugins/2D/Paper2D/Source/Paper2DEditor/Classes/PaperFlipbookFactory.h
+++ b/Engine/Plugins/2D/Paper2D/Source/Paper2DEditor/Classes/PaperFlipbookFactory.h
@@ -1,12 +1,14 @@
 // Copyright 1998-2015 Epic Games, Inc. All Rights Reserved.
 
+#pragma once
+
+#include "PaperFlipbook.h"
+#include "PaperFlipbookFactory.generated.h"
+
 /**
  * Factory for flipbooks
  */
 
-#pragma once
-#include "PaperFlipbookFactory.generated.h"
-
 UCLASS()
 class PAPER2DEDITOR_API UPaperFlipbookFactory : public UFactory
 {
diff --git a/Engine/Plugins/2D/Paper2D/Source/Paper2DEditor/Classes/PaperImporterSettings.h b/Engine/Plugins/2D/Paper2D/Source/Paper2DEditor/Classes/PaperImporterSettings.h
--- a/Engine/Plugins/2D/Paper2D/Source/Paper2DEditor/Classes/PaperImporterSettings.h
+++ b/Engine/Plugins/2D/Paper2D/Source/Paper2DEditor/Classes/PaperImporterSettings.h
@@ -2,6 +2,7 @@
 
 #pragma once
 
+#include "SpriteEditorOnlyTypes.h"
 #include "PaperImporterSettings.generated.h"
 
 class UPaperTileSet;
diff --git a/Engine/Plugins/2D/Paper2D/Source/Paper2DEditor/Private/Atlasing/AtlasAssetTypeActions.cpp b/Engine/Plugins/2D/Paper2D/Source/Paper2DEditor/Private/Atlasing/AtlasAssetTypeActions.cpp
--- a/Engine/Plugins/2D/Paper2D/Source/Paper2DEditor/Private/Atlasing/AtlasAssetTypeActions.cpp
+++ b/Engine/Plugins/2D/Paper2D/Source/Paper2DEditor/Private/Atlasing/AtlasAssetTypeActions.cpp
@@ -4,6 +4,7 @@
 
 #include "AtlasAssetTypeActions.h"
 #include "AssetData.h"
+#include "PaperSpriteAtlas.h"
 
 #define LOCTEXT_NAMESPACE "AssetTypeActions"
 
diff --git a/Engine/Plugins/2D/Paper2D/Source/Paper2DEditor/Private/Atlasing/PaperAtlasGenerator.cpp b/Engine/Plugins/2D/Paper2D/Source/Paper2DEditor/Private/Atlasing/PaperAtlasGenerator.cpp
--- a/Engine/Plugins/2D/Paper2D/Source/Paper2DEditor/Private/Atlasing/PaperAtlasGenerator.cpp
+++ b/Engine/Plugins/2D/Paper2D/Source/Paper2DEditor/Private/Atlasing/PaperAtlasGenerator.cpp
@@ -5,6 +5,7 @@
 
 #include "AssetRegistryModule.h"
 #include "Textures/TextureAtlas.h"
+#include "PaperSprite.h"
 
 //////////////////////////////////////////////////////////////////////////
 // 
diff --git a/Engine/Plugins/2D/Paper2D/Source/Paper2DEditor/Private/Atlasing/PaperAtlasGenerator.h b/Engine/Plugins/2D/Paper2D/Source/Paper2DEditor/Private/Atlasing/PaperAtlasGenerator.h
--- a/Engine/Plugins/2D/Paper2D/Source/Paper2DEditor/Private/Atlasing/PaperAtlasGenerator.h
+++ b/Engine/Plugins/2D/Paper2D/Source/Paper2DEditor/Private/Atlasing/PaperAtlasGenerator.h
@@ -5,9 +5,10 @@
 //////////////////////////////////////////////////////////////////////////
 // FPaperAtlasGenerator
 
+class UPaperSpriteAtlas;
+
 struct FPaperAtlasGenerator
 {
-
 public:
 	static void HandleAssetChangedEvent(UPaperSpriteAtlas* Atlas);
 };
diff --git a/Engine/Plugins/2D/Paper2D/Source/Paper2DEditor/Private/Atlasing/PaperAtlasTextureHelpers.cpp b/Engine/Plugins/2D/Paper2D/Source/Paper2DEditor/Private/Atlasing/PaperAtlasTextureHelpers.cpp
--- a/Engine/Plugins/2D/Paper2D/Source/Paper2DEditor/Private/Atlasing/PaperAtlasTextureHelpers.cpp
+++ b/Engine/Plugins/2D/Paper2D/Source/Paper2DEditor/Private/Atlasing/PaperAtlasTextureHelpers.cpp
@@ -2,6 +2,7 @@
 
 #include "Paper2DEditorPrivatePCH.h"
 #include "PaperAtlasTextureHelpers.h"
+#include "PaperSprite.h"
 
 //////////////////////////////////////////////////////////////////////////
 // 
diff --git a/Engine/Plugins/2D/Paper2D/Source/Paper2DEditor/Private/Atlasing/PaperAtlasTextureHelpers.h b/Engine/Plugins/2D/Paper2D/Source/Paper2DEditor/Private/Atlasing/PaperAtlasTextureHelpers.h
--- a/Engine/Plugins/2D/Paper2D/Source/Paper2DEditor/Private/Atlasing/PaperAtlasTextureHelpers.h
+++ b/Engine/Plugins/2D/Paper2D/Source/Paper2DEditor/Private/Atlasing/PaperAtlasTextureHelpers.h
@@ -2,6 +2,8 @@
 
 #pragma once
 
+#include "PaperSpriteAtlas.h"
+
 struct FPaperAtlasTextureHelpers
 {
 public:
diff --git a/Engine/Plugins/2D/Paper2D/Source/Paper2DEditor/Private/Atlasing/PaperSpriteAtlasFactory.cpp b/Engine/Plugins/2D/Paper2D/Source/Paper2DEditor/Private/Atlasing/PaperSpriteAtlasFactory.cpp
--- a/Engine/Plugins/2D/Paper2D/Source/Paper2DEditor/Private/Atlasing/PaperSpriteAtlasFactory.cpp
+++ b/Engine/Plugins/2D/Paper2D/Source/Paper2DEditor/Private/Atlasing/PaperSpriteAtlasFactory.cpp
@@ -1,7 +1,9 @@
 // Copyright 1998-2015 Epic Games, Inc. All Rights Reserved.
 
 #include "Paper2DEditorPrivatePCH.h"
+#include "PaperSpriteAtlas.h"
 #include "PaperSpriteAtlasFactory.h"
+#include "PaperRuntimeSettings.h"
 
 //////////////////////////////////////////////////////////////////////////
 // UPaperSpriteAtlasFactory
diff --git a/Engine/Plugins/2D/Paper2D/Source/Paper2DEditor/Private/ContentBrowserExtensions/ContentBrowserExtensions.cpp b/Engine/Plugins/2D/Paper2D/Source/Paper2DEditor/Private/ContentBrowserExtensions/ContentBrowserExtensions.cpp
--- a/Engine/Plugins/2D/Paper2D/Source/Paper2DEditor/Private/ContentBrowserExtensions/ContentBrowserExtensions.cpp
+++ b/Engine/Plugins/2D/Paper2D/Source/Paper2DEditor/Private/ContentBrowserExtensions/ContentBrowserExtensions.cpp
@@ -3,6 +3,9 @@
 #include "Paper2DEditorPrivatePCH.h"
 #include "ContentBrowserModule.h"
 #include "ContentBrowserExtensions.h"
+#include "PaperSprite.h"
+#include "PaperFlipbook.h"
+#include "PaperTileSet.h"
 #include "PaperSpriteFactory.h"
 #include "PaperTileSetFactory.h"
 #include "PaperImporterSettings.h"
diff --git a/Engine/Plugins/2D/Paper2D/Source/Paper2DEditor/Private/ExtractSprites/SPaperExtractSpritesDialog.cpp b/Engine/Plugins/2D/Paper2D/Source/Paper2DEditor/Private/ExtractSprites/SPaperExtractSpritesDialog.cpp
--- a/Engine/Plugins/2D/Paper2D/Source/Paper2DEditor/Private/ExtractSprites/SPaperExtractSpritesDialog.cpp
+++ b/Engine/Plugins/2D/Paper2D/Source/Paper2DEditor/Private/ExtractSprites/SPaperExtractSpritesDialog.cpp
@@ -10,6 +10,7 @@
 #include "ContentBrowserModule.h"
 #include "AssetToolsModule.h"
 #include "CanvasTypes.h"
+#include "PaperSprite.h"
 #include "PaperSpriteFactory.h"
 
 #define LOCTEXT_NAMESPACE "PaperEditor"
diff --git a/Engine/Plugins/2D/Paper2D/Source/Paper2DEditor/Private/FlipbookEditor/FlipbookEditor.cpp b/Engine/Plugins/2D/Paper2D/Source/Paper2DEditor/Private/FlipbookEditor/FlipbookEditor.cpp
--- a/Engine/Plugins/2D/Paper2D/Source/Paper2DEditor/Private/FlipbookEditor/FlipbookEditor.cpp
+++ b/Engine/Plugins/2D/Paper2D/Source/Paper2DEditor/Private/FlipbookEditor/FlipbookEditor.cpp
@@ -1,6 +1,7 @@
 // Copyright 1998-2015 Epic Games, Inc. All Rights Reserved.
 
 #include "Paper2DEditorPrivatePCH.h"
+#include "PaperFlipbookComponent.h"
 #include "FlipbookEditor.h"
 #include "SSingleObjectDetailsPanel.h"
 #include "SceneViewport.h"
diff --git a/Engine/Plugins/2D/Paper2D/Source/Paper2DEditor/Private/FlipbookEditor/FlipbookEditor.h b/Engine/Plugins/2D/Paper2D/Source/Paper2DEditor/Private/FlipbookEditor/FlipbookEditor.h
--- a/Engine/Plugins/2D/Paper2D/Source/Paper2DEditor/Private/FlipbookEditor/FlipbookEditor.h
+++ b/Engine/Plugins/2D/Paper2D/Source/Paper2DEditor/Private/FlipbookEditor/FlipbookEditor.h
@@ -5,6 +5,9 @@
 #include "Toolkits/AssetEditorToolkit.h"
 #include "Toolkits/AssetEditorManager.h"
 #include "Editor/EditorWidgets/Public/ITransportControl.h"
+#include "PaperFlipbookComponent.h"
+
+class SFlipbookEditorViewport;
 
 //////////////////////////////////////////////////////////////////////////
 // FFlipbookEditor
@@ -15,8 +18,8 @@ public:
 	FFlipbookEditor();
 
 	// IToolkit interface
-	virtual void RegisterTabSpawners(const TSharedRef<class FTabManager>& TabManager) override;
-	virtual void UnregisterTabSpawners(const TSharedRef<class FTabManager>& TabManager) override;
+	virtual void RegisterTabSpawners(const TSharedRef<FTabManager>& TabManager) override;
+	virtual void UnregisterTabSpawners(const TSharedRef<FTabManager>& TabManager) override;
 	// End of IToolkit interface
 
 	// FAssetEditorToolkit
@@ -33,7 +36,7 @@ public:
 	// End of FSerializableObject interface
 
 public:
-	void InitFlipbookEditor(const EToolkitMode::Type Mode, const TSharedPtr< class IToolkitHost >& InitToolkitHost, class UPaperFlipbook* InitSprite);
+	void InitFlipbookEditor(const EToolkitMode::Type Mode, const TSharedPtr< IToolkitHost >& InitToolkitHost, UPaperFlipbook* InitSprite);
 
 	UPaperFlipbook* GetFlipbookBeingEdited() const { return FlipbookBeingEdited; }
 
@@ -41,7 +44,7 @@ public:
 protected:
 	UPaperFlipbook* FlipbookBeingEdited;
 
-	TSharedPtr<class SFlipbookEditorViewport> ViewportPtr;
+	TSharedPtr<SFlipbookEditorViewport> ViewportPtr;
 
 	// Selection set
 	int32 CurrentSelectedKeyframe;
diff --git a/Engine/Plugins/2D/Paper2D/Source/Paper2DEditor/Private/FlipbookEditor/FlipbookEditorViewportClient.h b/Engine/Plugins/2D/Paper2D/Source/Paper2DEditor/Private/FlipbookEditor/FlipbookEditorViewportClient.h
--- a/Engine/Plugins/2D/Paper2D/Source/Paper2DEditor/Private/FlipbookEditor/FlipbookEditorViewportClient.h
+++ b/Engine/Plugins/2D/Paper2D/Source/Paper2DEditor/Private/FlipbookEditor/FlipbookEditorViewportClient.h
@@ -3,6 +3,7 @@
 #pragma once
 
 #include "Paper2DEditorPrivatePCH.h"
+#include "PaperFlipbookComponent.h"
 #include "FlipbookEditor.h"
 #include "SceneViewport.h"
 #include "PaperEditorViewportClient.h"
diff --git a/Engine/Plugins/2D/Paper2D/Source/Paper2DEditor/Private/FlipbookEditor/SFlipbookTimeline.cpp b/Engine/Plugins/2D/Paper2D/Source/Paper2DEditor/Private/FlipbookEditor/SFlipbookTimeline.cpp
--- a/Engine/Plugins/2D/Paper2D/Source/Paper2DEditor/Private/FlipbookEditor/SFlipbookTimeline.cpp
+++ b/Engine/Plugins/2D/Paper2D/Source/Paper2DEditor/Private/FlipbookEditor/SFlipbookTimeline.cpp
@@ -1,6 +1,7 @@
 // Copyright 1998-2015 Epic Games, Inc. All Rights Reserved.
 
 #include "Paper2DEditorPrivatePCH.h"
+#include "PaperFlipbook.h"
 #include "SFlipbookTimeline.h"
 #include "Editor/UnrealEd/Public/DragAndDrop/AssetDragDropOp.h"
 #include "Editor/UnrealEd/Public/ScopedTransaction.h"
diff --git a/Engine/Plugins/2D/Paper2D/Source/Paper2DEditor/Private/FlipbookEditor/SFlipbookTrackHandle.h b/Engine/Plugins/2D/Paper2D/Source/Paper2DEditor/Private/FlipbookEditor/SFlipbookTrackHandle.h
--- a/Engine/Plugins/2D/Paper2D/Source/Paper2DEditor/Private/FlipbookEditor/SFlipbookTrackHandle.h
+++ b/Engine/Plugins/2D/Paper2D/Source/Paper2DEditor/Private/FlipbookEditor/SFlipbookTrackHandle.h
@@ -2,6 +2,8 @@
 
 #pragma once
 
+#include "PaperFlipbook.h"
+
 //////////////////////////////////////////////////////////////////////////
 // SFlipbookTrackHandle
 
diff --git a/Engine/Plugins/2D/Paper2D/Source/Paper2DEditor/Private/FlipbookEditor/STimelineTrack.h b/Engine/Plugins/2D/Paper2D/Source/Paper2DEditor/Private/FlipbookEditor/STimelineTrack.h
--- a/Engine/Plugins/2D/Paper2D/Source/Paper2DEditor/Private/FlipbookEditor/STimelineTrack.h
+++ b/Engine/Plugins/2D/Paper2D/Source/Paper2DEditor/Private/FlipbookEditor/STimelineTrack.h
@@ -4,6 +4,7 @@
 
 #include "ScopedTransaction.h"
 #include "SFlipbookTimeline.h"
+#include "PaperFlipbook.h"
 
 namespace FFlipbookUIConstants
 {
diff --git a/Engine/Plugins/2D/Paper2D/Source/Paper2DEditor/Private/MeshPainting/MeshPaintSpriteAdapter.cpp b/Engine/Plugins/2D/Paper2D/Source/Paper2DEditor/Private/MeshPainting/MeshPaintSpriteAdapter.cpp
--- a/Engine/Plugins/2D/Paper2D/Source/Paper2DEditor/Private/MeshPainting/MeshPaintSpriteAdapter.cpp
+++ b/Engine/Plugins/2D/Paper2D/Source/Paper2DEditor/Private/MeshPainting/MeshPaintSpriteAdapter.cpp
@@ -3,6 +3,7 @@
 #include "Paper2DEditorPrivatePCH.h"
 #include "MeshPaintSpriteAdapter.h"
 #include "MeshPaintEdMode.h"
+#include "PaperSpriteComponent.h"
 
 //////////////////////////////////////////////////////////////////////////
 // FMeshPaintSpriteAdapter
diff --git a/Engine/Plugins/2D/Paper2D/Source/Paper2DEditor/Private/MeshPainting/MeshPaintSpriteAdapter.h b/Engine/Plugins/2D/Paper2D/Source/Paper2DEditor/Private/MeshPainting/MeshPaintSpriteAdapter.h
--- a/Engine/Plugins/2D/Paper2D/Source/Paper2DEditor/Private/MeshPainting/MeshPaintSpriteAdapter.h
+++ b/Engine/Plugins/2D/Paper2D/Source/Paper2DEditor/Private/MeshPainting/MeshPaintSpriteAdapter.h
@@ -4,6 +4,9 @@
 
 #include "MeshPaintModule.h"
 
+class UMeshComponent;
+class UPaperSpriteComponent;
+
 //////////////////////////////////////////////////////////////////////////
 // FMeshPaintSpriteAdapter
 
@@ -30,5 +33,5 @@ protected:
 class FMeshPaintSpriteAdapterFactory : public IMeshPaintGeometryAdapterFactory
 {
 public:
-	virtual TSharedPtr<IMeshPaintGeometryAdapter> Construct(class UMeshComponent* InComponent, int32 InPaintingMeshLODIndex, int32 InUVChannelIndex) const override;
+	virtual TSharedPtr<IMeshPaintGeometryAdapter> Construct(UMeshComponent* InComponent, int32 InPaintingMeshLODIndex, int32 InUVChannelIndex) const override;
 };
diff --git a/Engine/Plugins/2D/Paper2D/Source/Paper2DEditor/Private/Paper2DEditorModule.cpp b/Engine/Plugins/2D/Paper2D/Source/Paper2DEditor/Private/Paper2DEditorModule.cpp
--- a/Engine/Plugins/2D/Paper2D/Source/Paper2DEditor/Private/Paper2DEditorModule.cpp
+++ b/Engine/Plugins/2D/Paper2D/Source/Paper2DEditor/Private/Paper2DEditorModule.cpp
@@ -46,6 +46,7 @@
 // Atlas support
 #include "Atlasing/AtlasAssetTypeActions.h"
 #include "Atlasing/PaperAtlasGenerator.h"
+#include "PaperSpriteAtlas.h"
 
 // Settings
 #include "PaperRuntimeSettings.h"
diff --git a/Engine/Plugins/2D/Paper2D/Source/Paper2DEditor/Private/PaperEditorShared/SocketEditing.cpp b/Engine/Plugins/2D/Paper2D/Source/Paper2DEditor/Private/PaperEditorShared/SocketEditing.cpp
--- a/Engine/Plugins/2D/Paper2D/Source/Paper2DEditor/Private/PaperEditorShared/SocketEditing.cpp
+++ b/Engine/Plugins/2D/Paper2D/Source/Paper2DEditor/Private/PaperEditorShared/SocketEditing.cpp
@@ -1,6 +1,7 @@
 	// Copyright 1998-2015 Epic Games, Inc. All Rights Reserved.
 
 #include "Paper2DEditorPrivatePCH.h"
+#include "PaperSpriteComponent.h"
 #include "SocketEditing.h"
 #include "SpriteGeometryEditMode.h"
 #include "CanvasItem.h"
diff --git a/Engine/Plugins/2D/Paper2D/Source/Paper2DEditor/Private/PaperFlipbookActorFactory.cpp b/Engine/Plugins/2D/Paper2D/Source/Paper2DEditor/Private/PaperFlipbookActorFactory.cpp
--- a/Engine/Plugins/2D/Paper2D/Source/Paper2DEditor/Private/PaperFlipbookActorFactory.cpp
+++ b/Engine/Plugins/2D/Paper2D/Source/Paper2DEditor/Private/PaperFlipbookActorFactory.cpp
@@ -3,6 +3,7 @@
 #include "Paper2DEditorPrivatePCH.h"
 #include "AssetData.h"
 #include "PaperFlipbookActorFactory.h"
+#include "PaperFlipbookActor.h"
 
 //////////////////////////////////////////////////////////////////////////
 // UPaperFlipbookActorFactory
diff --git a/Engine/Plugins/2D/Paper2D/Source/Paper2DEditor/Private/PaperFlipbookAssetBroker.h b/Engine/Plugins/2D/Paper2D/Source/Paper2DEditor/Private/PaperFlipbookAssetBroker.h
--- a/Engine/Plugins/2D/Paper2D/Source/Paper2DEditor/Private/PaperFlipbookAssetBroker.h
+++ b/Engine/Plugins/2D/Paper2D/Source/Paper2DEditor/Private/PaperFlipbookAssetBroker.h
@@ -3,6 +3,7 @@
 #pragma once
 
 #include "ComponentAssetBroker.h"
+#include "PaperFlipbookComponent.h"
 
 //////////////////////////////////////////////////////////////////////////
 // FPaperFlipbookAssetBroker
diff --git a/Engine/Plugins/2D/Paper2D/Source/Paper2DEditor/Private/PaperFlipbookHelpers.cpp b/Engine/Plugins/2D/Paper2D/Source/Paper2DEditor/Private/PaperFlipbookHelpers.cpp
--- a/Engine/Plugins/2D/Paper2D/Source/Paper2DEditor/Private/PaperFlipbookHelpers.cpp
+++ b/Engine/Plugins/2D/Paper2D/Source/Paper2DEditor/Private/PaperFlipbookHelpers.cpp
@@ -4,6 +4,7 @@
 #include "PaperFlipbookHelpers.h"
 #include "Json.h"
 #include "ObjectTools.h"
+#include "PaperFlipbook.h"
 
 //////////////////////////////////////////////////////////////////////////
 // FPaperFlipbookHelpers
diff --git a/Engine/Plugins/2D/Paper2D/Source/Paper2DEditor/Private/PaperFlipbookThumbnailRenderer.cpp b/Engine/Plugins/2D/Paper2D/Source/Paper2DEditor/Private/PaperFlipbookThumbnailRenderer.cpp
--- a/Engine/Plugins/2D/Paper2D/Source/Paper2DEditor/Private/PaperFlipbookThumbnailRenderer.cpp
+++ b/Engine/Plugins/2D/Paper2D/Source/Paper2DEditor/Private/PaperFlipbookThumbnailRenderer.cpp
@@ -3,6 +3,7 @@
 #include "Paper2DEditorPrivatePCH.h"
 #include "CanvasItem.h"
 #include "PaperFlipbookThumbnailRenderer.h"
+#include "PaperFlipbook.h"
 
 //////////////////////////////////////////////////////////////////////////
 // UPaperFlipbookThumbnailRenderer
diff --git a/Engine/Plugins/2D/Paper2D/Source/Paper2DEditor/Private/PaperImporterSettings.cpp b/Engine/Plugins/2D/Paper2D/Source/Paper2DEditor/Private/PaperImporterSettings.cpp
--- a/Engine/Plugins/2D/Paper2D/Source/Paper2DEditor/Private/PaperImporterSettings.cpp
+++ b/Engine/Plugins/2D/Paper2D/Source/Paper2DEditor/Private/PaperImporterSettings.cpp
@@ -2,6 +2,8 @@
 
 #include "Paper2DEditorPrivatePCH.h"
 #include "PaperImporterSettings.h"
+#include "PaperTileMap.h"
+#include "PaperTileSet.h"
 #include "AlphaBitmap.h"
 #include "TileMapEditing/TileMapEditorSettings.h"
 
diff --git a/Engine/Plugins/2D/Paper2D/Source/Paper2DEditor/Private/PaperSpriteActorFactory.cpp b/Engine/Plugins/2D/Paper2D/Source/Paper2DEditor/Private/PaperSpriteActorFactory.cpp
--- a/Engine/Plugins/2D/Paper2D/Source/Paper2DEditor/Private/PaperSpriteActorFactory.cpp
+++ b/Engine/Plugins/2D/Paper2D/Source/Paper2DEditor/Private/PaperSpriteActorFactory.cpp
@@ -4,6 +4,7 @@
 #include "AssetData.h"
 #include "PhysicsEngine/BodySetup.h"
 #include "PaperSpriteActorFactory.h"
+#include "PaperSpriteActor.h"
 
 //////////////////////////////////////////////////////////////////////////
 // UPaperSpriteActorFactory
diff --git a/Engine/Plugins/2D/Paper2D/Source/Paper2DEditor/Private/PaperSpriteAssetBroker.h b/Engine/Plugins/2D/Paper2D/Source/Paper2DEditor/Private/PaperSpriteAssetBroker.h
--- a/Engine/Plugins/2D/Paper2D/Source/Paper2DEditor/Private/PaperSpriteAssetBroker.h
+++ b/Engine/Plugins/2D/Paper2D/Source/Paper2DEditor/Private/PaperSpriteAssetBroker.h
@@ -4,6 +4,7 @@
 
 #include "ComponentAssetBroker.h"
 #include "PhysicsEngine/BodySetup.h"
+#include "PaperSpriteComponent.h"
 
 //////////////////////////////////////////////////////////////////////////
 // FPaperSpriteAssetBroker
diff --git a/Engine/Plugins/2D/Paper2D/Source/Paper2DEditor/Private/PaperSpriteFactory.cpp b/Engine/Plugins/2D/Paper2D/Source/Paper2DEditor/Private/PaperSpriteFactory.cpp
--- a/Engine/Plugins/2D/Paper2D/Source/Paper2DEditor/Private/PaperSpriteFactory.cpp
+++ b/Engine/Plugins/2D/Paper2D/Source/Paper2DEditor/Private/PaperSpriteFactory.cpp
@@ -3,6 +3,7 @@
 #include "Paper2DEditorPrivatePCH.h"
 #include "PaperImporterSettings.h"
 #include "PaperSpriteFactory.h"
+#include "PaperSprite.h"
 #include "AssetRegistryModule.h"
 #include "PackageTools.h"
 
diff --git a/Engine/Plugins/2D/Paper2D/Source/Paper2DEditor/Private/PaperSpriteThumbnailRenderer.cpp b/Engine/Plugins/2D/Paper2D/Source/Paper2DEditor/Private/PaperSpriteThumbnailRenderer.cpp
--- a/Engine/Plugins/2D/Paper2D/Source/Paper2DEditor/Private/PaperSpriteThumbnailRenderer.cpp
+++ b/Engine/Plugins/2D/Paper2D/Source/Paper2DEditor/Private/PaperSpriteThumbnailRenderer.cpp
@@ -3,6 +3,7 @@
 #include "Paper2DEditorPrivatePCH.h"
 #include "CanvasTypes.h"
 #include "PaperSpriteThumbnailRenderer.h"
+#include "PaperSprite.h"
 
 //////////////////////////////////////////////////////////////////////////
 // UPaperSpriteThumbnailRenderer
diff --git a/Engine/Plugins/2D/Paper2D/Source/Paper2DEditor/Private/PaperTileMapFactory.cpp b/Engine/Plugins/2D/Paper2D/Source/Paper2DEditor/Private/PaperTileMapFactory.cpp
--- a/Engine/Plugins/2D/Paper2D/Source/Paper2DEditor/Private/PaperTileMapFactory.cpp
+++ b/Engine/Plugins/2D/Paper2D/Source/Paper2DEditor/Private/PaperTileMapFactory.cpp
@@ -1,6 +1,7 @@
 // Copyright 1998-2015 Epic Games, Inc. All Rights Reserved.
 
 #include "Paper2DEditorPrivatePCH.h"
+#include "PaperTileMap.h"
 #include "PaperTileMapFactory.h"
 #include "PaperImporterSettings.h"
 
diff --git a/Engine/Plugins/2D/Paper2D/Source/Paper2DEditor/Private/PaperTileMapPromotionFactory.cpp b/Engine/Plugins/2D/Paper2D/Source/Paper2DEditor/Private/PaperTileMapPromotionFactory.cpp
--- a/Engine/Plugins/2D/Paper2D/Source/Paper2DEditor/Private/PaperTileMapPromotionFactory.cpp
+++ b/Engine/Plugins/2D/Paper2D/Source/Paper2DEditor/Private/PaperTileMapPromotionFactory.cpp
@@ -2,6 +2,7 @@
 
 #include "Paper2DEditorPrivatePCH.h"
 #include "PaperTileMapPromotionFactory.h"
+#include "PaperTileMap.h"
 
 #define LOCTEXT_NAMESPACE "Paper2D"
 
diff --git a/Engine/Plugins/2D/Paper2D/Source/Paper2DEditor/Private/PaperTileSetFactory.cpp b/Engine/Plugins/2D/Paper2D/Source/Paper2DEditor/Private/PaperTileSetFactory.cpp
--- a/Engine/Plugins/2D/Paper2D/Source/Paper2DEditor/Private/PaperTileSetFactory.cpp
+++ b/Engine/Plugins/2D/Paper2D/Source/Paper2DEditor/Private/PaperTileSetFactory.cpp
@@ -1,6 +1,7 @@
 // Copyright 1998-2015 Epic Games, Inc. All Rights Reserved.
 
 #include "Paper2DEditorPrivatePCH.h"
+#include "PaperTileSet.h"
 #include "PaperTileSetFactory.h"
 #include "TileSetEditor/TileSetEditorSettings.h"
 
diff --git a/Engine/Plugins/2D/Paper2D/Source/Paper2DEditor/Private/PaperTileSetThumbnailRenderer.cpp b/Engine/Plugins/2D/Paper2D/Source/Paper2DEditor/Private/PaperTileSetThumbnailRenderer.cpp
--- a/Engine/Plugins/2D/Paper2D/Source/Paper2DEditor/Private/PaperTileSetThumbnailRenderer.cpp
+++ b/Engine/Plugins/2D/Paper2D/Source/Paper2DEditor/Private/PaperTileSetThumbnailRenderer.cpp
@@ -3,6 +3,7 @@
 #include "Paper2DEditorPrivatePCH.h"
 #include "CanvasTypes.h"
 #include "PaperTileSetThumbnailRenderer.h"
+#include "PaperTileSet.h"
 
 //////////////////////////////////////////////////////////////////////////
 // UPaperTileSetThumbnailRenderer
diff --git a/Engine/Plugins/2D/Paper2D/Source/Paper2DEditor/Private/SpriteAssetTypeActions.cpp b/Engine/Plugins/2D/Paper2D/Source/Paper2DEditor/Private/SpriteAssetTypeActions.cpp
--- a/Engine/Plugins/2D/Paper2D/Source/Paper2DEditor/Private/SpriteAssetTypeActions.cpp
+++ b/Engine/Plugins/2D/Paper2D/Source/Paper2DEditor/Private/SpriteAssetTypeActions.cpp
@@ -8,6 +8,8 @@
 #include "ContentBrowserModule.h"
 #include "PaperFlipbookHelpers.h"
 #include "PaperFlipbookFactory.h"
+#include "PaperSprite.h"
+#include "PaperFlipbook.h"
 
 #define LOCTEXT_NAMESPACE "AssetTypeActions"
 
diff --git a/Engine/Plugins/2D/Paper2D/Source/Paper2DEditor/Private/SpriteAssetTypeActions.h b/Engine/Plugins/2D/Paper2D/Source/Paper2DEditor/Private/SpriteAssetTypeActions.h
--- a/Engine/Plugins/2D/Paper2D/Source/Paper2DEditor/Private/SpriteAssetTypeActions.h
+++ b/Engine/Plugins/2D/Paper2D/Source/Paper2DEditor/Private/SpriteAssetTypeActions.h
@@ -4,6 +4,8 @@
 
 #include "AssetTypeActions_Base.h"
 
+class UPaperSprite;
+
 class FSpriteAssetTypeActions : public FAssetTypeActions_Base
 {
 public:
diff --git a/Engine/Plugins/2D/Paper2D/Source/Paper2DEditor/Private/SpriteEditor/SSpriteList.cpp b/Engine/Plugins/2D/Paper2D/Source/Paper2DEditor/Private/SpriteEditor/SSpriteList.cpp
--- a/Engine/Plugins/2D/Paper2D/Source/Paper2DEditor/Private/SpriteEditor/SSpriteList.cpp
+++ b/Engine/Plugins/2D/Paper2D/Source/Paper2DEditor/Private/SpriteEditor/SSpriteList.cpp
@@ -5,6 +5,7 @@
 #include "AssetRegistryModule.h"
 #include "ContentBrowserModule.h"
 #include "SpriteEditor.h"
+#include "PaperSprite.h"
 
 //////////////////////////////////////////////////////////////////////////
 // SSpriteList
diff --git a/Engine/Plugins/2D/Paper2D/Source/Paper2DEditor/Private/SpriteEditor/SpriteDetailsCustomization.cpp b/Engine/Plugins/2D/Paper2D/Source/Paper2DEditor/Private/SpriteEditor/SpriteDetailsCustomization.cpp
--- a/Engine/Plugins/2D/Paper2D/Source/Paper2DEditor/Private/SpriteEditor/SpriteDetailsCustomization.cpp
+++ b/Engine/Plugins/2D/Paper2D/Source/Paper2DEditor/Private/SpriteEditor/SpriteDetailsCustomization.cpp
@@ -9,6 +9,7 @@
 #include "PropertyRestriction.h"
 #include "PropertyCustomizationHelpers.h"
 #include "MaterialExpressionSpriteTextureSampler.h"
+#include "PaperRuntimeSettings.h"
 
 #define LOCTEXT_NAMESPACE "SpriteEditor"
 
diff --git a/Engine/Plugins/2D/Paper2D/Source/Paper2DEditor/Private/SpriteEditor/SpriteEditorSelections.h b/Engine/Plugins/2D/Paper2D/Source/Paper2DEditor/Private/SpriteEditor/SpriteEditorSelections.h
--- a/Engine/Plugins/2D/Paper2D/Source/Paper2DEditor/Private/SpriteEditor/SpriteEditorSelections.h
+++ b/Engine/Plugins/2D/Paper2D/Source/Paper2DEditor/Private/SpriteEditor/SpriteEditorSelections.h
@@ -3,6 +3,7 @@
 #pragma once
 
 #include "PaperEditorShared/AssetEditorSelectedItem.h"
+#include "PaperSprite.h"
 
 //////////////////////////////////////////////////////////////////////////
 // FSelectionTypes
diff --git a/Engine/Plugins/2D/Paper2D/Source/Paper2DEditor/Private/SpriteEditor/SpriteEditorViewportClient.h b/Engine/Plugins/2D/Paper2D/Source/Paper2DEditor/Private/SpriteEditor/SpriteEditorViewportClient.h
--- a/Engine/Plugins/2D/Paper2D/Source/Paper2DEditor/Private/SpriteEditor/SpriteEditorViewportClient.h
+++ b/Engine/Plugins/2D/Paper2D/Source/Paper2DEditor/Private/SpriteEditor/SpriteEditorViewportClient.h
@@ -2,7 +2,7 @@
 
 #pragma once
 
-#include "Paper2DEditorPrivatePCH.h"
+#include "PaperSpriteComponent.h"
 #include "SpriteEditor.h"
 #include "SSingleObjectDetailsPanel.h"
 #include "SceneViewport.h"
diff --git a/Engine/Plugins/2D/Paper2D/Source/Paper2DEditor/Private/SpriteEditor/SpritePolygonCollectionCustomization.h b/Engine/Plugins/2D/Paper2D/Source/Paper2DEditor/Private/SpriteEditor/SpritePolygonCollectionCustomization.h
--- a/Engine/Plugins/2D/Paper2D/Source/Paper2DEditor/Private/SpriteEditor/SpritePolygonCollectionCustomization.h
+++ b/Engine/Plugins/2D/Paper2D/Source/Paper2DEditor/Private/SpriteEditor/SpritePolygonCollectionCustomization.h
@@ -3,6 +3,7 @@
 #pragma once
 
 #include "PropertyEditing.h"
+#include "SpriteEditorOnlyTypes.h"
 
 //////////////////////////////////////////////////////////////////////////
 // FSpritePolygonCollectionCustomization
diff --git a/Engine/Plugins/2D/Paper2D/Source/Paper2DEditor/Private/TerrainSplineEditing/TerrainSplineActorFactory.cpp b/Engine/Plugins/2D/Paper2D/Source/Paper2DEditor/Private/TerrainSplineEditing/TerrainSplineActorFactory.cpp
--- a/Engine/Plugins/2D/Paper2D/Source/Paper2DEditor/Private/TerrainSplineEditing/TerrainSplineActorFactory.cpp
+++ b/Engine/Plugins/2D/Paper2D/Source/Paper2DEditor/Private/TerrainSplineEditing/TerrainSplineActorFactory.cpp
@@ -3,6 +3,10 @@
 #include "Paper2DEditorPrivatePCH.h"
 #include "AssetData.h"
 #include "TerrainSplineActorFactory.h"
+#include "PaperTerrainActor.h"
+#include "PaperTerrainMaterial.h"
+#include "PaperTerrainComponent.h"
+#include "PaperRuntimeSettings.h"
 
 //////////////////////////////////////////////////////////////////////////
 // UTerrainSplineActorFactory
diff --git a/Engine/Plugins/2D/Paper2D/Source/Paper2DEditor/Private/TileMapEditing/EdModeTileMap.cpp b/Engine/Plugins/2D/Paper2D/Source/Paper2DEditor/Private/TileMapEditing/EdModeTileMap.cpp
--- a/Engine/Plugins/2D/Paper2D/Source/Paper2DEditor/Private/TileMapEditing/EdModeTileMap.cpp
+++ b/Engine/Plugins/2D/Paper2D/Source/Paper2DEditor/Private/TileMapEditing/EdModeTileMap.cpp
@@ -2,6 +2,7 @@
 
 #include "Paper2DEditorPrivatePCH.h"
 #include "EdModeTileMap.h"
+#include "PaperTileMapComponent.h"
 #include "TileMapEdModeToolkit.h"
 #include "Toolkits/ToolkitManager.h"
 #include "../PaperEditorCommands.h"
diff --git a/Engine/Plugins/2D/Paper2D/Source/Paper2DEditor/Private/TileMapEditing/EdModeTileMap.h b/Engine/Plugins/2D/Paper2D/Source/Paper2DEditor/Private/TileMapEditing/EdModeTileMap.h
--- a/Engine/Plugins/2D/Paper2D/Source/Paper2DEditor/Private/TileMapEditing/EdModeTileMap.h
+++ b/Engine/Plugins/2D/Paper2D/Source/Paper2DEditor/Private/TileMapEditing/EdModeTileMap.h
@@ -1,7 +1,11 @@
 // Copyright 1998-2015 Epic Games, Inc. All Rights Reserved.
 
 #pragma once
-#include "PaperTileSet.h"
+
+class UPaperTileSet;
+class UPaperTileLayer;
+class UPaperTileMap;
+class UPaperTileMapComponent;
 
 namespace ETileMapEditorTool
 {
diff --git a/Engine/Plugins/2D/Paper2D/Source/Paper2DEditor/Private/TileMapEditing/PaperTileMapAssetBroker.h b/Engine/Plugins/2D/Paper2D/Source/Paper2DEditor/Private/TileMapEditing/PaperTileMapAssetBroker.h
--- a/Engine/Plugins/2D/Paper2D/Source/Paper2DEditor/Private/TileMapEditing/PaperTileMapAssetBroker.h
+++ b/Engine/Plugins/2D/Paper2D/Source/Paper2DEditor/Private/TileMapEditing/PaperTileMapAssetBroker.h
@@ -3,6 +3,7 @@
 #pragma once
 
 #include "ComponentAssetBroker.h"
+#include "PaperTileMapComponent.h"
 
 //////////////////////////////////////////////////////////////////////////
 // FPaperTileMapAssetBroker
diff --git a/Engine/Plugins/2D/Paper2D/Source/Paper2DEditor/Private/TileMapEditing/PaperTileMapDetailsCustomization.cpp b/Engine/Plugins/2D/Paper2D/Source/Paper2DEditor/Private/TileMapEditing/PaperTileMapDetailsCustomization.cpp
--- a/Engine/Plugins/2D/Paper2D/Source/Paper2DEditor/Private/TileMapEditing/PaperTileMapDetailsCustomization.cpp
+++ b/Engine/Plugins/2D/Paper2D/Source/Paper2DEditor/Private/TileMapEditing/PaperTileMapDetailsCustomization.cpp
@@ -1,6 +1,10 @@
 // Copyright 1998-2015 Epic Games, Inc. All Rights Reserved.
 
 #include "Paper2DEditorPrivatePCH.h"
+#include "PaperTileMapComponent.h"
+#include "PaperTileMap.h"
+#include "PaperTileLayer.h"
+
 #include "PaperTileMapDetailsCustomization.h"
 #include "DetailLayoutBuilder.h"
 #include "DetailCategoryBuilder.h"
diff --git a/Engine/Plugins/2D/Paper2D/Source/Paper2DEditor/Private/TileMapEditing/STileLayerItem.h b/Engine/Plugins/2D/Paper2D/Source/Paper2DEditor/Private/TileMapEditing/STileLayerItem.h
--- a/Engine/Plugins/2D/Paper2D/Source/Paper2DEditor/Private/TileMapEditing/STileLayerItem.h
+++ b/Engine/Plugins/2D/Paper2D/Source/Paper2DEditor/Private/TileMapEditing/STileLayerItem.h
@@ -2,6 +2,9 @@
 
 #pragma once
 
+#include "PaperTileMap.h"
+#include "PaperTileLayer.h"
+
 //////////////////////////////////////////////////////////////////////////
 // STileLayerItem
 
@@ -17,8 +20,8 @@ public:
 
 protected:
 	int32 MyIndex;
-	class UPaperTileMap* MyMap;
-	class UPaperTileLayer* GetMyLayer() const { return MyMap->TileLayers[MyIndex]; }
+	UPaperTileMap* MyMap;
+	UPaperTileLayer* GetMyLayer() const { return MyMap->TileLayers[MyIndex]; }
 
 	TSharedPtr<SButton> VisibilityButton;
 
diff --git a/Engine/Plugins/2D/Paper2D/Source/Paper2DEditor/Private/TileMapEditing/STileLayerList.cpp b/Engine/Plugins/2D/Paper2D/Source/Paper2DEditor/Private/TileMapEditing/STileLayerList.cpp
--- a/Engine/Plugins/2D/Paper2D/Source/Paper2DEditor/Private/TileMapEditing/STileLayerList.cpp
+++ b/Engine/Plugins/2D/Paper2D/Source/Paper2DEditor/Private/TileMapEditing/STileLayerList.cpp
@@ -1,6 +1,7 @@
 // Copyright 1998-2015 Epic Games, Inc. All Rights Reserved.
 
 #include "Paper2DEditorPrivatePCH.h"
+#include "PaperTileMapComponent.h"
 #include "STileLayerList.h"
 #include "STileLayerItem.h"
 #include "PaperStyle.h"
diff --git a/Engine/Plugins/2D/Paper2D/Source/Paper2DEditor/Private/TileMapEditing/TileMapActorFactory.cpp b/Engine/Plugins/2D/Paper2D/Source/Paper2DEditor/Private/TileMapEditing/TileMapActorFactory.cpp
--- a/Engine/Plugins/2D/Paper2D/Source/Paper2DEditor/Private/TileMapEditing/TileMapActorFactory.cpp
+++ b/Engine/Plugins/2D/Paper2D/Source/Paper2DEditor/Private/TileMapEditing/TileMapActorFactory.cpp
@@ -3,7 +3,10 @@
 #include "Paper2DEditorPrivatePCH.h"
 #include "AssetData.h"
 #include "TileMapActorFactory.h"
+#include "PaperTileMapActor.h"
 #include "PaperImporterSettings.h"
+#include "PaperTileMapComponent.h"
+#include "PaperTileMapActor.h"
 
 //////////////////////////////////////////////////////////////////////////
 // UTileMapActorFactory
diff --git a/Engine/Plugins/2D/Paper2D/Source/Paper2DEditor/Private/TileMapEditing/TileMapAssetImportData.cpp b/Engine/Plugins/2D/Paper2D/Source/Paper2DEditor/Private/TileMapEditing/TileMapAssetImportData.cpp
--- a/Engine/Plugins/2D/Paper2D/Source/Paper2DEditor/Private/TileMapEditing/TileMapAssetImportData.cpp
+++ b/Engine/Plugins/2D/Paper2D/Source/Paper2DEditor/Private/TileMapEditing/TileMapAssetImportData.cpp
@@ -3,6 +3,7 @@
 #include "Paper2DEditorPrivatePCH.h"
 #include "AssetData.h"
 #include "TileMapAssetImportData.h"
+#include "PaperTileMap.h"
 
 //////////////////////////////////////////////////////////////////////////
 // UTileMapAssetImportData
diff --git a/Engine/Plugins/2D/Paper2D/Source/Paper2DEditor/Private/TileMapEditing/TileMapAssetTypeActions.cpp b/Engine/Plugins/2D/Paper2D/Source/Paper2DEditor/Private/TileMapEditing/TileMapAssetTypeActions.cpp
--- a/Engine/Plugins/2D/Paper2D/Source/Paper2DEditor/Private/TileMapEditing/TileMapAssetTypeActions.cpp
+++ b/Engine/Plugins/2D/Paper2D/Source/Paper2DEditor/Private/TileMapEditing/TileMapAssetTypeActions.cpp
@@ -4,6 +4,7 @@
 
 #include "TileMapAssetTypeActions.h"
 #include "TileMapEditor.h"
+#include "PaperTileMap.h"
 
 #define LOCTEXT_NAMESPACE "AssetTypeActions"
 
diff --git a/Engine/Plugins/2D/Paper2D/Source/Paper2DEditor/Private/TileMapEditing/TileMapEdModeToolkit.cpp b/Engine/Plugins/2D/Paper2D/Source/Paper2DEditor/Private/TileMapEditing/TileMapEdModeToolkit.cpp
--- a/Engine/Plugins/2D/Paper2D/Source/Paper2DEditor/Private/TileMapEditing/TileMapEdModeToolkit.cpp
+++ b/Engine/Plugins/2D/Paper2D/Source/Paper2DEditor/Private/TileMapEditing/TileMapEdModeToolkit.cpp
@@ -2,6 +2,9 @@
 
 #include "Paper2DEditorPrivatePCH.h"
 #include "EdModeTileMap.h"
+#include "PaperTileMapComponent.h"
+#include "PaperTileSet.h"
+
 #include "TileMapEdModeToolkit.h"
 #include "TileSetEditor/TileSetSelectorViewport.h"
 #include "SContentReference.h"
diff --git a/Engine/Plugins/2D/Paper2D/Source/Paper2DEditor/Private/TileMapEditing/TileMapEditor.cpp b/Engine/Plugins/2D/Paper2D/Source/Paper2DEditor/Private/TileMapEditing/TileMapEditor.cpp
--- a/Engine/Plugins/2D/Paper2D/Source/Paper2DEditor/Private/TileMapEditing/TileMapEditor.cpp
+++ b/Engine/Plugins/2D/Paper2D/Source/Paper2DEditor/Private/TileMapEditing/TileMapEditor.cpp
@@ -1,6 +1,7 @@
 // Copyright 1998-2015 Epic Games, Inc. All Rights Reserved.
 
 #include "Paper2DEditorPrivatePCH.h"
+#include "PaperTileMapComponent.h"
 #include "TileMapEditor.h"
 #include "SSingleObjectDetailsPanel.h"
 #include "SceneViewport.h"
diff --git a/Engine/Plugins/2D/Paper2D/Source/Paper2DEditor/Private/TileMapEditing/TileMapEditorViewportClient.cpp b/Engine/Plugins/2D/Paper2D/Source/Paper2DEditor/Private/TileMapEditing/TileMapEditorViewportClient.cpp
--- a/Engine/Plugins/2D/Paper2D/Source/Paper2DEditor/Private/TileMapEditing/TileMapEditorViewportClient.cpp
+++ b/Engine/Plugins/2D/Paper2D/Source/Paper2DEditor/Private/TileMapEditing/TileMapEditorViewportClient.cpp
@@ -1,6 +1,7 @@
 // Copyright 1998-2015 Epic Games, Inc. All Rights Reserved.
 
 #include "Paper2DEditorPrivatePCH.h"
+#include "PaperTileMapComponent.h"
 #include "TileMapEditorViewportClient.h"
 #include "SceneViewport.h"
 #include "EdModeTileMap.h"
diff --git a/Engine/Plugins/2D/Paper2D/Source/Paper2DEditor/Private/TileMapEditing/TileMapEditorViewportClient.h b/Engine/Plugins/2D/Paper2D/Source/Paper2DEditor/Private/TileMapEditing/TileMapEditorViewportClient.h
--- a/Engine/Plugins/2D/Paper2D/Source/Paper2DEditor/Private/TileMapEditing/TileMapEditorViewportClient.h
+++ b/Engine/Plugins/2D/Paper2D/Source/Paper2DEditor/Private/TileMapEditing/TileMapEditorViewportClient.h
@@ -10,6 +10,10 @@
 #include "PreviewScene.h"
 #include "ScopedTransaction.h"
 
+class UPaperTileMapComponent;
+class STileMapEditorViewport;
+class FScopedTransaction;
+
 //////////////////////////////////////////////////////////////////////////
 // FTileMapEditorViewportClient
 
@@ -17,7 +21,7 @@ class FTileMapEditorViewportClient : public FPaperEditorViewportClient
 {
 public:
 	/** Constructor */
-	FTileMapEditorViewportClient(TWeakPtr<FTileMapEditor> InTileMapEditor, TWeakPtr<class STileMapEditorViewport> InTileMapEditorViewportPtr);
+	FTileMapEditorViewportClient(TWeakPtr<FTileMapEditor> InTileMapEditor, TWeakPtr<STileMapEditorViewport> InTileMapEditorViewportPtr);
 
 	// FViewportClient interface
 	virtual void Draw(const FSceneView* View, FPrimitiveDrawInterface* PDI) override;
@@ -75,10 +79,10 @@ private:
 	bool bShowTileMapStats;
 
 	// Pointer back to the tile map editor viewport control that owns us
-	TWeakPtr<class STileMapEditorViewport> TileMapEditorViewportPtr;
+	TWeakPtr<STileMapEditorViewport> TileMapEditorViewportPtr;
 
 	// The current transaction for undo/redo
-	class FScopedTransaction* ScopedTransaction;
+	FScopedTransaction* ScopedTransaction;
 
 	// Should we show the sprite pivot?
 	bool bShowPivot;
diff --git a/Engine/Plugins/2D/Paper2D/Source/Paper2DEditor/Private/TileSetAssetTypeActions.cpp b/Engine/Plugins/2D/Paper2D/Source/Paper2DEditor/Private/TileSetAssetTypeActions.cpp
--- a/Engine/Plugins/2D/Paper2D/Source/Paper2DEditor/Private/TileSetAssetTypeActions.cpp
+++ b/Engine/Plugins/2D/Paper2D/Source/Paper2DEditor/Private/TileSetAssetTypeActions.cpp
@@ -3,6 +3,8 @@
 #include "Paper2DEditorPrivatePCH.h"
 #include "TileSetAssetTypeActions.h"
 #include "TileSetEditor.h"
+#include "PaperTileSet.h"
+#include "PaperTileMap.h"
 #include "AssetToolsModule.h"
 #include "ContentBrowserModule.h"
 #include "PaperTileMapFactory.h"
diff --git a/Engine/Plugins/2D/Paper2D/Source/Paper2DEditor/Private/TileSetEditor.cpp b/Engine/Plugins/2D/Paper2D/Source/Paper2DEditor/Private/TileSetEditor.cpp
--- a/Engine/Plugins/2D/Paper2D/Source/Paper2DEditor/Private/TileSetEditor.cpp
+++ b/Engine/Plugins/2D/Paper2D/Source/Paper2DEditor/Private/TileSetEditor.cpp
@@ -2,6 +2,7 @@
 
 #include "Paper2DEditorPrivatePCH.h"
 #include "TileSetEditor.h"
+#include "PaperTileSet.h"
 #include "SSingleObjectDetailsPanel.h"
 #include "SceneViewport.h"
 #include "PaperEditorViewportClient.h"
diff --git a/Engine/Plugins/2D/Paper2D/Source/Paper2DEditor/Private/TileSetEditor/SingleTileEditorViewportClient.cpp b/Engine/Plugins/2D/Paper2D/Source/Paper2DEditor/Private/TileSetEditor/SingleTileEditorViewportClient.cpp
--- a/Engine/Plugins/2D/Paper2D/Source/Paper2DEditor/Private/TileSetEditor/SingleTileEditorViewportClient.cpp
+++ b/Engine/Plugins/2D/Paper2D/Source/Paper2DEditor/Private/TileSetEditor/SingleTileEditorViewportClient.cpp
@@ -2,6 +2,9 @@
 
 #include "Paper2DEditorPrivatePCH.h"
 #include "SingleTileEditorViewportClient.h"
+#include "PaperSpriteComponent.h"
+#include "PaperTileMap.h"
+#include "PaperTileSet.h"
 #include "PaperEditorShared/SpriteGeometryEditCommands.h"
 #include "PaperEditorShared/SpriteGeometryEditMode.h"
 #include "ScopedTransaction.h"
diff --git a/Engine/Plugins/2D/Paper2D/Source/Paper2DEditor/Private/TileSetEditor/SingleTileEditorViewportClient.h b/Engine/Plugins/2D/Paper2D/Source/Paper2DEditor/Private/TileSetEditor/SingleTileEditorViewportClient.h
--- a/Engine/Plugins/2D/Paper2D/Source/Paper2DEditor/Private/TileSetEditor/SingleTileEditorViewportClient.h
+++ b/Engine/Plugins/2D/Paper2D/Source/Paper2DEditor/Private/TileSetEditor/SingleTileEditorViewportClient.h
@@ -6,6 +6,11 @@
 #include "PaperEditorViewportClient.h"
 #include "SpriteEditor/SpriteEditorSelections.h"
 
+class UPaperSpriteComponent;
+class SEditorViewport;
+class FScopedTranaction;
+class UPaperTileSet;
+
 DECLARE_MULTICAST_DELEGATE_TwoParams(FOnSingleTileIndexChanged, int32 /*NewIndex*/, int32 /*OldIndex*/);
 
 //////////////////////////////////////////////////////////////////////////
@@ -22,7 +27,7 @@ public:
 
 	// FEditorViewportClient interface
 	virtual FLinearColor GetBackgroundColor() const override;
-	virtual void TrackingStarted(const struct FInputEventState& InInputState, bool bIsDragging, bool bNudge) override;
+	virtual void TrackingStarted(const FInputEventState& InInputState, bool bIsDragging, bool bNudge) override;
 	virtual void TrackingStopped() override;
 	virtual void DrawCanvas(FViewport& Viewport, FSceneView& View, FCanvas& Canvas) override;
 	// End of FEditorViewportClient interface
@@ -77,10 +82,10 @@ private:
 	bool bShowStats;
 
 	// Pointer back to the sprite editor viewport control that owns us
-	TWeakPtr<class SEditorViewport> SpriteEditorViewportPtr;
+	TWeakPtr<SEditorViewport> SpriteEditorViewportPtr;
 
 	// The current transaction for undo/redo
-	class FScopedTransaction* ScopedTransaction;
+	FScopedTransaction* ScopedTransaction;
 
 	// The preview scene
 	FPreviewScene OwnedPreviewScene;
diff --git a/Engine/Plugins/2D/Paper2D/Source/Paper2DEditor/Private/TileSetEditor/TileSetDetailsCustomization.cpp b/Engine/Plugins/2D/Paper2D/Source/Paper2DEditor/Private/TileSetEditor/TileSetDetailsCustomization.cpp
--- a/Engine/Plugins/2D/Paper2D/Source/Paper2DEditor/Private/TileSetEditor/TileSetDetailsCustomization.cpp
+++ b/Engine/Plugins/2D/Paper2D/Source/Paper2DEditor/Private/TileSetEditor/TileSetDetailsCustomization.cpp
@@ -7,6 +7,7 @@
 #include "PropertyEditing.h"
 #include "ScopedTransaction.h"
 #include "IPropertyUtilities.h"
+#include "PaperTileSet.h"
 
 #define LOCTEXT_NAMESPACE "Paper2D"
 
diff --git a/Engine/Plugins/2D/Paper2D/Source/Paper2DEditor/Private/TileSetEditor/TileSetEditorViewportClient.cpp b/Engine/Plugins/2D/Paper2D/Source/Paper2DEditor/Private/TileSetEditor/TileSetEditorViewportClient.cpp
--- a/Engine/Plugins/2D/Paper2D/Source/Paper2DEditor/Private/TileSetEditor/TileSetEditorViewportClient.cpp
+++ b/Engine/Plugins/2D/Paper2D/Source/Paper2DEditor/Private/TileSetEditor/TileSetEditorViewportClient.cpp
@@ -2,6 +2,7 @@
 
 #include "Paper2DEditorPrivatePCH.h"
 #include "TileSetEditorViewportClient.h"
+#include "PaperTileSet.h"
 
 #define LOCTEXT_NAMESPACE "TileSetEditor"
 
diff --git a/Engine/Plugins/2D/Paper2D/Source/Paper2DEditor/Private/TileSetEditor/TileSetEditorViewportClient.h b/Engine/Plugins/2D/Paper2D/Source/Paper2DEditor/Private/TileSetEditor/TileSetEditorViewportClient.h
--- a/Engine/Plugins/2D/Paper2D/Source/Paper2DEditor/Private/TileSetEditor/TileSetEditorViewportClient.h
+++ b/Engine/Plugins/2D/Paper2D/Source/Paper2DEditor/Private/TileSetEditor/TileSetEditorViewportClient.h
@@ -4,6 +4,8 @@
 
 #include "PaperEditorViewportClient.h"
 
+class UPaperTileSet;
+
 //////////////////////////////////////////////////////////////////////////
 // FTileSetEditorViewportClient
 
diff --git a/Engine/Plugins/2D/Paper2D/Source/Paper2DEditor/Private/TileSetEditor/TileSetSelectorViewport.cpp b/Engine/Plugins/2D/Paper2D/Source/Paper2DEditor/Private/TileSetEditor/TileSetSelectorViewport.cpp
--- a/Engine/Plugins/2D/Paper2D/Source/Paper2DEditor/Private/TileSetEditor/TileSetSelectorViewport.cpp
+++ b/Engine/Plugins/2D/Paper2D/Source/Paper2DEditor/Private/TileSetEditor/TileSetSelectorViewport.cpp
@@ -5,6 +5,7 @@
 #include "TileSetSelectorViewport.h"
 #include "TileSetEditorViewportClient.h"
 #include "TileSetEditorCommands.h"
+#include "PaperTileSet.h"
 
 #define LOCTEXT_NAMESPACE "TileSetEditor"
 
diff --git a/Engine/Plugins/2D/Paper2D/Source/Paper2DEditor/Private/TileSetEditor/TileSheetPaddingFactory.cpp b/Engine/Plugins/2D/Paper2D/Source/Paper2DEditor/Private/TileSetEditor/TileSheetPaddingFactory.cpp
--- a/Engine/Plugins/2D/Paper2D/Source/Paper2DEditor/Private/TileSetEditor/TileSheetPaddingFactory.cpp
+++ b/Engine/Plugins/2D/Paper2D/Source/Paper2DEditor/Private/TileSetEditor/TileSheetPaddingFactory.cpp
@@ -1,6 +1,7 @@
 // Copyright 1998-2015 Epic Games, Inc. All Rights Reserved.
 
 #include "Paper2DEditorPrivatePCH.h"
+#include "PaperTileSet.h"
 #include "TileSheetPaddingFactory.h"
 #include "Atlasing/PaperAtlasTextureHelpers.h"
 
diff --git a/Engine/Plugins/2D/Paper2D/Source/Paper2DEditor/Public/PaperFlipbookHelpers.h b/Engine/Plugins/2D/Paper2D/Source/Paper2DEditor/Public/PaperFlipbookHelpers.h
--- a/Engine/Plugins/2D/Paper2D/Source/Paper2DEditor/Public/PaperFlipbookHelpers.h
+++ b/Engine/Plugins/2D/Paper2D/Source/Paper2DEditor/Public/PaperFlipbookHelpers.h
@@ -2,6 +2,8 @@
 
 #pragma once
 
+class UPaperSprite;
+
 //////////////////////////////////////////////////////////////////////////
 // FPaperFlipbookHelpers
 
diff --git a/Engine/Plugins/2D/Paper2D/Source/PaperSpriteSheetImporter/Private/PaperJsonSpriteSheetImporter.cpp b/Engine/Plugins/2D/Paper2D/Source/PaperSpriteSheetImporter/Private/PaperJsonSpriteSheetImporter.cpp
--- a/Engine/Plugins/2D/Paper2D/Source/PaperSpriteSheetImporter/Private/PaperJsonSpriteSheetImporter.cpp
+++ b/Engine/Plugins/2D/Paper2D/Source/PaperSpriteSheetImporter/Private/PaperJsonSpriteSheetImporter.cpp
@@ -1,7 +1,7 @@
 // Copyright 1998-2015 Epic Games, Inc. All Rights Reserved.
 
 #include "PaperSpriteSheetImporterPrivatePCH.h"
-#include "Paper2DClasses.h"
+#include "PaperSprite.h"
 #include "Json.h"
 #include "PaperJSONHelpers.h"
 #include "PaperJsonSpriteSheetImporter.h"
diff --git a/Engine/Plugins/2D/Paper2D/Source/PaperSpriteSheetImporter/Private/PaperJsonSpriteSheetImporter.h b/Engine/Plugins/2D/Paper2D/Source/PaperSpriteSheetImporter/Private/PaperJsonSpriteSheetImporter.h
--- a/Engine/Plugins/2D/Paper2D/Source/PaperSpriteSheetImporter/Private/PaperJsonSpriteSheetImporter.h
+++ b/Engine/Plugins/2D/Paper2D/Source/PaperSpriteSheetImporter/Private/PaperJsonSpriteSheetImporter.h
@@ -3,7 +3,6 @@
 #pragma once
 
 #include "PaperSpriteSheetImporterPrivatePCH.h"
-#include "Paper2DClasses.h"
 
 class UPaperSpriteSheet;
 
diff --git a/Engine/Plugins/2D/Paper2D/Source/PaperSpriteSheetImporter/Private/PaperSpriteSheetAssetTypeActions.cpp b/Engine/Plugins/2D/Paper2D/Source/PaperSpriteSheetImporter/Private/PaperSpriteSheetAssetTypeActions.cpp
--- a/Engine/Plugins/2D/Paper2D/Source/PaperSpriteSheetImporter/Private/PaperSpriteSheetAssetTypeActions.cpp
+++ b/Engine/Plugins/2D/Paper2D/Source/PaperSpriteSheetImporter/Private/PaperSpriteSheetAssetTypeActions.cpp
@@ -1,7 +1,7 @@
 // Copyright 1998-2015 Epic Games, Inc. All Rights Reserved.
 
 #include "PaperSpriteSheetImporterPrivatePCH.h"
-#include "Paper2DClasses.h"
+#include "PaperFlipbook.h"
 #include "PaperSpriteSheet.h"
 #include "PaperSpriteSheetAssetTypeActions.h"
 #include "AssetToolsModule.h"
diff --git a/Engine/Plugins/2D/Paper2D/Source/PaperSpriteSheetImporter/Private/PaperSpriteSheetImportFactory.cpp b/Engine/Plugins/2D/Paper2D/Source/PaperSpriteSheetImporter/Private/PaperSpriteSheetImportFactory.cpp
--- a/Engine/Plugins/2D/Paper2D/Source/PaperSpriteSheetImporter/Private/PaperSpriteSheetImportFactory.cpp
+++ b/Engine/Plugins/2D/Paper2D/Source/PaperSpriteSheetImporter/Private/PaperSpriteSheetImportFactory.cpp
@@ -1,7 +1,6 @@
 // Copyright 1998-2015 Epic Games, Inc. All Rights Reserved.
 
 #include "PaperSpriteSheetImporterPrivatePCH.h"
-#include "Paper2DClasses.h"
 #include "Json.h"
 #include "PaperJSONHelpers.h"
 #include "AssetToolsModule.h"
diff --git a/Engine/Plugins/2D/Paper2D/Source/PaperSpriteSheetImporter/Private/PaperSpriteSheetImporterPrivatePCH.h b/Engine/Plugins/2D/Paper2D/Source/PaperSpriteSheetImporter/Private/PaperSpriteSheetImporterPrivatePCH.h
--- a/Engine/Plugins/2D/Paper2D/Source/PaperSpriteSheetImporter/Private/PaperSpriteSheetImporterPrivatePCH.h
+++ b/Engine/Plugins/2D/Paper2D/Source/PaperSpriteSheetImporter/Private/PaperSpriteSheetImporterPrivatePCH.h
@@ -3,7 +3,7 @@
 #pragma once
 
 #include "UnrealEd.h"
-#include "Paper2DClasses.h"
+//#include "Paper2DClasses.h"
 
 DECLARE_LOG_CATEGORY_EXTERN(LogPaperSpriteSheetImporter, Verbose, All);
 
diff --git a/Engine/Plugins/2D/Paper2D/Source/PaperTiledImporter/Private/PaperTiledImporterFactory.cpp b/Engine/Plugins/2D/Paper2D/Source/PaperTiledImporter/Private/PaperTiledImporterFactory.cpp
--- a/Engine/Plugins/2D/Paper2D/Source/PaperTiledImporter/Private/PaperTiledImporterFactory.cpp
+++ b/Engine/Plugins/2D/Paper2D/Source/PaperTiledImporter/Private/PaperTiledImporterFactory.cpp
@@ -1,9 +1,8 @@
 // Copyright 1998-2015 Epic Games, Inc. All Rights Reserved.
 
 #include "PaperTiledImporterPrivatePCH.h"
-#include "Paper2DClasses.h"
-#include "Paper2DEditorClasses.h"
 #include "Json.h"
+#include "TileMapAssetImportData.h"
 #include "PaperJSONHelpers.h"
 #include "AssetToolsModule.h"
 #include "AssetRegistryModule.h"
diff --git a/Engine/Plugins/2D/Paper2D/Source/SpriterImporter/Private/SpriterDataModel.cpp b/Engine/Plugins/2D/Paper2D/Source/SpriterImporter/Private/SpriterDataModel.cpp
--- a/Engine/Plugins/2D/Paper2D/Source/SpriterImporter/Private/SpriterDataModel.cpp
+++ b/Engine/Plugins/2D/Paper2D/Source/SpriterImporter/Private/SpriterDataModel.cpp
@@ -1,8 +1,6 @@
 // Copyright 1998-2015 Epic Games, Inc. All Rights Reserved.
 
 #include "SpriterImporterPrivatePCH.h"
-#include "Paper2DClasses.h"
-#include "Paper2DEditorClasses.h"
 #include "Json.h"
 #include "PaperJSONHelpers.h"
 #include "SpriterDataModel.h"
