Paper2D: Instanced sprite component WIP [UE-8031]
Note: UPaperGroupedSpriteComponent is still experimental

diff --git a/Engine/Plugins/2D/Paper2D/Source/Paper2D/Classes/PaperGroupedSpriteComponent.h b/Engine/Plugins/2D/Paper2D/Source/Paper2D/Classes/PaperGroupedSpriteComponent.h
new file mode 100644
--- /dev/null
+++ b/Engine/Plugins/2D/Paper2D/Source/Paper2D/Classes/PaperGroupedSpriteComponent.h
@@ -0,0 +1,156 @@
+// Copyright 1998-2015 Epic Games, Inc. All Rights Reserved.
+
+#pragma once
+
+#include "PaperSprite.h"
+
+#include "PaperGroupedSpriteComponent.generated.h"
+
+class UPaperGroupedSpriteComponent;
+class FGroupedSpriteSceneProxy;
+class FGroupedSpriteComponentDetailsCustomization;
+
+USTRUCT()
+struct FSpriteInstanceData
+{
+public:
+	GENERATED_USTRUCT_BODY()
+
+	UPROPERTY(EditAnywhere, Category=Instances)
+	FMatrix Transform;
+
+	UPROPERTY(EditAnywhere, Category=Instances)
+	UPaperSprite* SourceSprite;
+
+	UPROPERTY(EditAnywhere, Category=Instances)
+	FColor VertexColor;
+
+	UPROPERTY(EditAnywhere, Category=Instances)
+	int32 MaterialIndex;
+
+	bool IsValidInstance() const
+	{
+		return SourceSprite != nullptr;
+	}
+
+public:
+	FSpriteInstanceData()
+		: Transform(FMatrix::Identity)
+		, SourceSprite(nullptr)
+		, VertexColor(FColor::White)
+		, MaterialIndex(INDEX_NONE)
+	{
+	}
+};
+
+/**
+ * A component that handles rendering and collision for many instances of a UPaperSprite asset.
+ *
+ * @see UPrimitiveComponent, UPaperSprite
+ */
+
+UCLASS(ShowCategories=(Mobility), ClassGroup=Paper2D, meta=(BlueprintSpawnableComponent), Experimental)
+class PAPER2D_API UPaperGroupedSpriteComponent : public UMeshComponent
+{
+	GENERATED_UCLASS_BODY()
+
+protected:
+	/** Array of materials used by the instances */
+	UPROPERTY()
+	TArray<UMaterialInterface*> InstanceMaterials;
+
+	/** Array of instances */
+	UPROPERTY(EditAnywhere, DisplayName="Instances", Category=Instances, meta=(MakeEditWidget=true))
+	TArray<FSpriteInstanceData> PerInstanceSpriteData;
+
+	/** Physics representation of the instance bodies */
+	TArray<FBodyInstance*> InstanceBodies;
+
+public:
+
+	/** Add an instance to this component. Transform is given in local space of this component.  */
+	UFUNCTION(BlueprintCallable, Category="Components|Sprite")
+	virtual int32 AddInstance(const FTransform& LocalTransform, UPaperSprite* Sprite, FColor Color = FColor(255, 255, 255));
+
+	/** Add an instance to this component. Transform is given in world space. */
+	UFUNCTION(BlueprintCallable, Category = "Components|Sprite")
+	int32 AddInstanceWorldSpace(const FTransform& WorldTransform, UPaperSprite* Sprite, FColor Color = FColor(255, 255, 255));
+
+	/** Get the transform for the instance specified. Instance is returned in local space of this component unless bWorldSpace is set.  Returns True on success. */
+	UFUNCTION(BlueprintCallable, Category = "Components|Sprite")
+	bool GetInstanceTransform(int32 InstanceIndex, FTransform& OutInstanceTransform, bool bWorldSpace = false) const;
+	
+	/** Update the transform for the instance specified. Instance is given in local space of this component unless bWorldSpace is set.  Returns True on success. */
+	UFUNCTION(BlueprintCallable, Category = "Components|Sprite")
+	virtual bool UpdateInstanceTransform(int32 InstanceIndex, const FTransform& NewInstanceTransform, bool bWorldSpace = false, bool bMarkRenderStateDirty = false);
+
+	/** Remove the instance specified. Returns True on success. */
+	UFUNCTION(BlueprintCallable, Category = "Components|Sprite")
+	virtual bool RemoveInstance(int32 InstanceIndex);
+	
+	/** Clear all instances being rendered by this component */
+	UFUNCTION(BlueprintCallable, Category="Components|Sprite")
+	virtual void ClearInstances();
+	
+	/** Get the number of instances in this component */
+	UFUNCTION(BlueprintCallable, Category = "Components|Sprite")
+	int32 GetInstanceCount() const;
+
+	// UActorComponent interface
+	virtual bool ShouldCreatePhysicsState() const override;
+	virtual void CreatePhysicsState() override;
+	virtual void DestroyPhysicsState() override;
+	virtual void CheckForErrors() override;
+	// End of UActorComponent interface 
+
+	// UPrimitiveComponent interface
+	virtual FPrimitiveSceneProxy* CreateSceneProxy() override;
+	virtual bool CanEditSimulatePhysics() override;
+	virtual FBoxSphereBounds CalcBounds(const FTransform& BoundTransform) const override;
+	virtual void GetUsedTextures(TArray<UTexture*>& OutTextures, EMaterialQualityLevel::Type QualityLevel) override;
+	virtual UMaterialInterface* GetMaterial(int32 MaterialIndex) const override;
+	virtual int32 GetNumMaterials() const override;
+	// End of UPrimitiveComponent interface
+
+	// UObject interface
+#if WITH_EDITOR
+	virtual void PostEditChangeProperty(FPropertyChangedEvent& PropertyChangedEvent) override;
+	virtual void PostEditChangeChainProperty(FPropertyChangedChainEvent& PropertyChangedEvent) override;
+	virtual void PostEditUndo() override;
+#endif
+	// End of UObject interface
+
+	// Returns true if this component references the specified sprite asset
+	bool ContainsSprite(UPaperSprite* SpriteAsset) const;
+
+protected:
+	/**
+	 * Transfers ownership of instance render data to a render thread
+	 * instance render data will be released in scene proxy dtor or on render thread task
+	 */
+	void ReleasePerInstanceRenderData();
+
+	/** Creates body instances for all instances owned by this component */
+	void CreateAllInstanceBodies();
+
+	/** Terminate all body instances owned by this component */
+	void ClearAllInstanceBodies();
+
+	/** Sets up new instance data to sensible defaults, creates physics counterparts if possible */
+	void SetupNewInstanceData(FSpriteInstanceData& InOutNewInstanceData, int32 InInstanceIndex, const FTransform& InInstanceTransform, UPaperSprite* InSprite, const FColor& InColor);
+
+	/** Creates a body instance for the specified instance data if that sprite has defined collision */
+	FBodyInstance* InitInstanceBody(int32 InstanceIndex, const FSpriteInstanceData& InstanceData, FPhysScene* PhysScene);
+
+	/** Invalidates the render and collision state */
+	void RebuildInstances();
+
+	/** Creates the material list from the instances */
+	void RebuildMaterialList();
+
+	/** Adds to the material list from a single sprite */
+	int32 UpdateMaterialList(UPaperSprite* Sprite);
+
+	friend FGroupedSpriteSceneProxy;
+	friend FGroupedSpriteComponentDetailsCustomization;
+};
diff --git a/Engine/Plugins/2D/Paper2D/Source/Paper2D/Private/GroupedSpriteSceneProxy.cpp b/Engine/Plugins/2D/Paper2D/Source/Paper2D/Private/GroupedSpriteSceneProxy.cpp
new file mode 100644
--- /dev/null
+++ b/Engine/Plugins/2D/Paper2D/Source/Paper2D/Private/GroupedSpriteSceneProxy.cpp
@@ -0,0 +1,83 @@
+// Copyright 1998-2015 Epic Games, Inc. All Rights Reserved.
+
+#include "Paper2DPrivatePCH.h"
+#include "GroupedSpriteSceneProxy.h"
+#include "PaperGroupedSpriteComponent.h"
+
+//////////////////////////////////////////////////////////////////////////
+// FGroupedSpriteSceneProxy
+
+FSpriteRenderSection& FGroupedSpriteSceneProxy::FindOrAddSection(FSpriteDrawCallRecord& InBatch, UMaterialInterface* InMaterial)
+{
+	// Check the existing sections, starting with the most recent
+	for (int32 SectionIndex = BatchedSections.Num() - 1; SectionIndex >= 0; --SectionIndex)
+	{
+		FSpriteRenderSection& TestSection = BatchedSections[SectionIndex];
+
+		if (TestSection.Material == InMaterial)
+		{
+			if (TestSection.BaseTexture == InBatch.BaseTexture)
+			{
+				if (TestSection.AdditionalTextures == InBatch.AdditionalTextures)
+				{
+					return TestSection;
+				}
+			}
+		}
+	}
+
+	// Didn't find a matching section, create one
+	FSpriteRenderSection& NewSection = *(new (BatchedSections) FSpriteRenderSection());
+	NewSection.Material = InMaterial;
+	NewSection.BaseTexture = InBatch.BaseTexture;
+	NewSection.AdditionalTextures = InBatch.AdditionalTextures;
+	NewSection.VertexOffset = VertexBuffer.Vertices.Num();
+
+	return NewSection;
+}
+
+FGroupedSpriteSceneProxy::FGroupedSpriteSceneProxy(UPaperGroupedSpriteComponent* InComponent)
+	: FPaperRenderSceneProxy(InComponent)
+	, MyComponent(InComponent)
+{
+	MaterialRelevance = InComponent->GetMaterialRelevance(GetScene().GetFeatureLevel());
+
+	NumInstances = InComponent->PerInstanceSpriteData.Num();
+
+	for (const FSpriteInstanceData InstanceData : InComponent->PerInstanceSpriteData)
+	{
+		if (InstanceData.SourceSprite != nullptr)
+		{
+			FSpriteDrawCallRecord Record;
+			Record.BuildFromSprite(InstanceData.SourceSprite);
+
+			UMaterialInterface* SpriteMaterial = InstanceData.SourceSprite->GetMaterial(0); //@TODO: Should query from the component due to overrides, but that requires an index per item, etc...
+
+			FSpriteRenderSection& Section = FindOrAddSection(Record, SpriteMaterial);
+			
+			
+			const int32 NumNewVerts = Record.RenderVerts.Num();
+			Section.NumVertices += NumNewVerts;
+
+			const FColor VertColor(InstanceData.VertexColor);
+			for (const FVector4& SourceVert : Record.RenderVerts)
+			{
+				const FVector LocalPos((PaperAxisX * SourceVert.X) + (PaperAxisY * SourceVert.Y));
+				const FVector ComponentSpacePos = InstanceData.Transform.TransformPosition(LocalPos);
+				const FVector2D UV(SourceVert.Z, SourceVert.W);
+
+				new (VertexBuffer.Vertices) FPaperSpriteVertex(ComponentSpacePos, UV, VertColor);
+			}
+		}
+	}
+	
+	if (VertexBuffer.Vertices.Num() > 0)
+	{
+		// Init the vertex factory
+		MyVertexFactory.Init(&VertexBuffer);
+
+		// Enqueue initialization of render resources
+		BeginInitResource(&VertexBuffer);
+		BeginInitResource(&MyVertexFactory);
+	}
+}
diff --git a/Engine/Plugins/2D/Paper2D/Source/Paper2D/Private/GroupedSpriteSceneProxy.h b/Engine/Plugins/2D/Paper2D/Source/Paper2D/Private/GroupedSpriteSceneProxy.h
new file mode 100644
--- /dev/null
+++ b/Engine/Plugins/2D/Paper2D/Source/Paper2D/Private/GroupedSpriteSceneProxy.h
@@ -0,0 +1,32 @@
+// Copyright 1998-2015 Epic Games, Inc. All Rights Reserved.
+
+#pragma once
+
+class UPaperGroupedSpriteComponent;
+
+#include "PaperRenderSceneProxy.h"
+
+//////////////////////////////////////////////////////////////////////////
+// FGroupedSpriteSceneProxy
+
+class FGroupedSpriteSceneProxy : public FPaperRenderSceneProxy
+{
+public:
+	FGroupedSpriteSceneProxy(UPaperGroupedSpriteComponent* InComponent);
+
+	// FPrimitiveSceneProxy interface.
+	//virtual void GetDynamicMeshElements(const TArray<const FSceneView*>& Views, const FSceneViewFamily& ViewFamily, uint32 VisibilityMap, FMeshElementCollector& Collector) const override;
+	// End of FPrimitiveSceneProxy interface.
+
+private:
+	const UPaperGroupedSpriteComponent* MyComponent;
+
+	/** Per instance render data, could be shared with component */
+	TSharedPtr<FPerInstanceRenderData, ESPMode::ThreadSafe> PerInstanceRenderData;
+
+	/** Number of instances */
+	int32 NumInstances;
+
+
+	FSpriteRenderSection& FindOrAddSection(FSpriteDrawCallRecord& InBatch, UMaterialInterface* InMaterial);
+};
diff --git a/Engine/Plugins/2D/Paper2D/Source/Paper2D/Private/PaperGroupedSpriteComponent.cpp b/Engine/Plugins/2D/Paper2D/Source/Paper2D/Private/PaperGroupedSpriteComponent.cpp
new file mode 100644
--- /dev/null
+++ b/Engine/Plugins/2D/Paper2D/Source/Paper2D/Private/PaperGroupedSpriteComponent.cpp
@@ -0,0 +1,473 @@
+// Copyright 1998-2015 Epic Games, Inc. All Rights Reserved.
+
+#include "Paper2DPrivatePCH.h"
+#include "PaperSpriteSceneProxy.h"
+#include "PaperSpriteComponent.h"
+
+#include "PaperGroupedSpriteComponent.h"
+#include "GroupedSpriteSceneProxy.h"
+
+#include "PaperSprite.h"
+#include "PhysicsPublic.h"
+
+#include "MessageLog.h"
+#include "UObjectToken.h"
+
+#define LOCTEXT_NAMESPACE "Paper2D"
+
+//////////////////////////////////////////////////////////////////////////
+// UPaperGroupedSpriteComponent
+
+UPaperGroupedSpriteComponent::UPaperGroupedSpriteComponent(const FObjectInitializer& ObjectInitializer)
+	: Super(ObjectInitializer)
+{
+	SetCollisionProfileName(UCollisionProfile::BlockAllDynamic_ProfileName);
+
+	Mobility = EComponentMobility::Movable;
+	BodyInstance.bSimulatePhysics = false;
+}
+
+int32 UPaperGroupedSpriteComponent::AddInstance(const FTransform& LocalTransform, UPaperSprite* Sprite, FColor Color)
+{
+	const int32 NewInstanceIndex = PerInstanceSpriteData.Num();
+
+	FSpriteInstanceData& NewInstanceData = *new (PerInstanceSpriteData) FSpriteInstanceData();
+	SetupNewInstanceData(NewInstanceData, NewInstanceIndex, LocalTransform, Sprite, Color);
+
+	MarkRenderStateDirty();
+
+	UNavigationSystem::UpdateNavOctree(this);
+
+	return NewInstanceIndex;
+}
+
+int32 UPaperGroupedSpriteComponent::AddInstanceWorldSpace(const FTransform& WorldTransform, UPaperSprite* Sprite, FColor Color)
+{
+	// Transform from world space to local space
+	const FTransform RelativeTM = WorldTransform.GetRelativeTransform(ComponentToWorld);
+	return AddInstance(RelativeTM, Sprite, Color);
+}
+
+bool UPaperGroupedSpriteComponent::GetInstanceTransform(int32 InstanceIndex, FTransform& OutInstanceTransform, bool bWorldSpace) const
+{
+	if (!PerInstanceSpriteData.IsValidIndex(InstanceIndex))
+	{
+		return false;
+	}
+
+	const FSpriteInstanceData& InstanceData = PerInstanceSpriteData[InstanceIndex];
+
+	OutInstanceTransform = FTransform(InstanceData.Transform);
+	if (bWorldSpace)
+	{
+		OutInstanceTransform = OutInstanceTransform * ComponentToWorld;
+	}
+
+	return true;
+}
+
+bool UPaperGroupedSpriteComponent::UpdateInstanceTransform(int32 InstanceIndex, const FTransform& NewInstanceTransform, bool bWorldSpace, bool bMarkRenderStateDirty)
+{
+	if (!PerInstanceSpriteData.IsValidIndex(InstanceIndex))
+	{
+		return false;
+	}
+
+	// Request navigation update
+	UNavigationSystem::UpdateNavOctree(this);
+
+	FSpriteInstanceData& InstanceData = PerInstanceSpriteData[InstanceIndex];
+
+	// Render data uses local transform of the instance
+	FTransform LocalTransform = bWorldSpace ? NewInstanceTransform.GetRelativeTransform(ComponentToWorld) : NewInstanceTransform;
+	InstanceData.Transform = LocalTransform.ToMatrixWithScale();
+
+	if (bPhysicsStateCreated)
+	{
+		// Physics uses world transform of the instance
+		const FTransform WorldTransform = bWorldSpace ? NewInstanceTransform : (LocalTransform * ComponentToWorld);
+		if (FBodyInstance* InstanceBodyInstance = InstanceBodies[InstanceIndex])
+		{
+			// Update transform.
+			InstanceBodyInstance->SetBodyTransform(WorldTransform, false);
+			InstanceBodyInstance->UpdateBodyScale(WorldTransform.GetScale3D());
+		}
+	}
+
+	// Request navigation update
+	UNavigationSystem::UpdateNavOctree(this);
+
+	if (bMarkRenderStateDirty)
+	{
+		MarkRenderStateDirty();
+	}
+
+	return true;
+}
+
+bool UPaperGroupedSpriteComponent::RemoveInstance(int32 InstanceIndex)
+{
+	if (!PerInstanceSpriteData.IsValidIndex(InstanceIndex))
+	{
+		return false;
+	}
+
+	// Request navigation update
+	UNavigationSystem::UpdateNavOctree(this);
+
+	// remove instance
+	PerInstanceSpriteData.RemoveAt(InstanceIndex);
+
+	// update the physics state
+	if (bPhysicsStateCreated)
+	{
+		// TODO: it may be possible to instead just update the BodyInstanceIndex for all bodies after the removed instance. 
+		ClearAllInstanceBodies();
+		CreateAllInstanceBodies();
+	}
+
+	// Indicate we need to update render state to reflect changes
+	MarkRenderStateDirty();
+
+	return true;
+}
+
+void UPaperGroupedSpriteComponent::ClearInstances()
+{
+	// Clear all the per-instance data
+	PerInstanceSpriteData.Empty();
+
+	// Release any physics representations
+	ClearAllInstanceBodies();
+
+	// Indicate we need to update render state to reflect changes
+	MarkRenderStateDirty();
+
+	UNavigationSystem::UpdateNavOctree(this);
+}
+
+int32 UPaperGroupedSpriteComponent::GetInstanceCount() const
+{
+	return PerInstanceSpriteData.Num();
+}
+
+bool UPaperGroupedSpriteComponent::ShouldCreatePhysicsState() const
+{
+	return IsRegistered() && (bAlwaysCreatePhysicsState || IsCollisionEnabled());
+}
+
+void UPaperGroupedSpriteComponent::CreatePhysicsState()
+{
+	check(InstanceBodies.Num() == 0);
+
+	// Create all the bodies.
+	CreateAllInstanceBodies();
+
+	USceneComponent::CreatePhysicsState();
+}
+
+void UPaperGroupedSpriteComponent::DestroyPhysicsState()
+{
+	USceneComponent::DestroyPhysicsState();
+
+	// Release all physics representations
+	ClearAllInstanceBodies();
+}
+
+void UPaperGroupedSpriteComponent::CheckForErrors()
+{
+	AActor* Owner = GetOwner();
+
+	bool bAnyInstancesWithNoSprites = false;
+	for (const FSpriteInstanceData& InstanceData : PerInstanceSpriteData)
+	{
+		if (InstanceData.SourceSprite == nullptr)
+		{
+			bAnyInstancesWithNoSprites = true;
+			break;
+		}
+	}
+
+	if (bAnyInstancesWithNoSprites)
+	{
+		FMessageLog("MapCheck").Warning()
+			->AddToken(FUObjectToken::Create(Owner))
+			->AddToken(FTextToken::Create(LOCTEXT("MapCheck_Message_InstanceSpriteComponent_MissingSprite", "One or more instances have no sprite asset set!")));
+	}
+
+	Super::CheckForErrors();
+}
+
+FPrimitiveSceneProxy* UPaperGroupedSpriteComponent::CreateSceneProxy()
+{
+	if (PerInstanceSpriteData.Num() > 0)
+	{
+		return ::new FGroupedSpriteSceneProxy(this);
+	}
+	else
+	{
+		return nullptr;
+	}
+}
+
+bool UPaperGroupedSpriteComponent::CanEditSimulatePhysics()
+{
+	// Simulating for instanced sprite components is never allowed
+	return false;
+}
+
+FBoxSphereBounds UPaperGroupedSpriteComponent::CalcBounds(const FTransform& BoundTransform) const
+{
+	bool bHadAnyBounds = false;
+	FBoxSphereBounds NewBounds(ForceInit);
+
+	if (PerInstanceSpriteData.Num() > 0)
+	{
+		const FMatrix BoundTransformMatrix = BoundTransform.ToMatrixWithScale();
+
+		for (const FSpriteInstanceData& InstanceData : PerInstanceSpriteData)
+		{
+			if (InstanceData.SourceSprite != nullptr)
+			{
+				const FBoxSphereBounds RenderBounds = InstanceData.SourceSprite->GetRenderBounds();
+				const FBoxSphereBounds InstanceBounds = RenderBounds.TransformBy(InstanceData.Transform * BoundTransformMatrix);
+
+				if (bHadAnyBounds)
+				{
+					NewBounds = NewBounds + InstanceBounds;
+				}
+				else
+				{
+					NewBounds = InstanceBounds;
+					bHadAnyBounds = true;
+				}
+			}
+		}
+	}
+
+	return bHadAnyBounds ? NewBounds : FBoxSphereBounds(BoundTransform.GetLocation(), FVector::ZeroVector, 0.f);
+}
+
+#if WITH_EDITOR
+void UPaperGroupedSpriteComponent::PostEditChangeProperty(FPropertyChangedEvent& PropertyChangedEvent)
+{
+	// Rebuild the material array
+	RebuildMaterialList();
+
+	// Rebuild all instances
+	//@TODO: Should try and avoid this except when absolutely necessary
+	RebuildInstances();
+
+	Super::PostEditChangeProperty(PropertyChangedEvent);
+}
+
+void UPaperGroupedSpriteComponent::PostEditChangeChainProperty(FPropertyChangedChainEvent& PropertyChangedEvent)
+{
+	if ((PropertyChangedEvent.Property != nullptr) && (PropertyChangedEvent.Property->GetFName() == GET_MEMBER_NAME_CHECKED(UPaperGroupedSpriteComponent, PerInstanceSpriteData)))
+	{
+		if (PropertyChangedEvent.ChangeType == EPropertyChangeType::ArrayAdd)
+		{
+			const int32 AddedAtIndex = PropertyChangedEvent.GetArrayIndex(PropertyChangedEvent.Property->GetFName().ToString());
+			check(AddedAtIndex != INDEX_NONE);
+			SetupNewInstanceData(PerInstanceSpriteData[AddedAtIndex], AddedAtIndex, FTransform::Identity, nullptr, FColor::White); //@TODO: Need to pull the sprite from somewhere
+		}
+
+		MarkRenderStateDirty();
+	}
+
+	Super::PostEditChangeChainProperty(PropertyChangedEvent);
+}
+
+void UPaperGroupedSpriteComponent::PostEditUndo()
+{
+	Super::PostEditUndo();
+
+	UNavigationSystem::UpdateNavOctree(this);
+}
+#endif
+
+void UPaperGroupedSpriteComponent::CreateAllInstanceBodies()
+{
+	QUICK_SCOPE_CYCLE_COUNTER(STAT_UPaperGroupedSpriteComponent_CreateAllInstanceBodies);
+
+	FPhysScene* PhysScene = GetWorld()->GetPhysicsScene();
+
+	const int32 NumBodies = PerInstanceSpriteData.Num();
+	check(InstanceBodies.Num() == 0);
+	InstanceBodies.SetNumUninitialized(NumBodies);
+
+	TArray<FTransform> Transforms;
+	Transforms.Reserve(NumBodies);
+
+	for (int32 InstanceIndex = 0; InstanceIndex < NumBodies; ++InstanceIndex)
+	{
+		const FSpriteInstanceData& InstanceData = PerInstanceSpriteData[InstanceIndex];
+		InstanceBodies[InstanceIndex] = InitInstanceBody(InstanceIndex, InstanceData, PhysScene);
+	}
+}
+
+void UPaperGroupedSpriteComponent::ClearAllInstanceBodies()
+{
+	QUICK_SCOPE_CYCLE_COUNTER(STAT_UPaperGroupedSpriteComponent_ClearAllInstanceBodies);
+
+	for (FBodyInstance* OldBodyInstance : InstanceBodies)
+	{
+		if (OldBodyInstance != nullptr)
+		{
+			OldBodyInstance->TermBody();
+			delete OldBodyInstance;
+		}
+	}
+
+	InstanceBodies.Empty();
+}
+
+void UPaperGroupedSpriteComponent::SetupNewInstanceData(FSpriteInstanceData& InOutNewInstanceData, int32 InInstanceIndex, const FTransform& InInstanceTransform, UPaperSprite* InSprite, const FColor& InColor)
+{
+	InOutNewInstanceData.Transform = InInstanceTransform.ToMatrixWithScale();
+	InOutNewInstanceData.SourceSprite = InSprite;
+	InOutNewInstanceData.VertexColor = InColor;
+	InOutNewInstanceData.MaterialIndex = UpdateMaterialList(InSprite);
+
+	if (bPhysicsStateCreated && (InSprite != nullptr) && (InSprite->BodySetup != nullptr))
+	{
+		FPhysScene* PhysScene = GetWorld()->GetPhysicsScene();
+		FBodyInstance* NewBodyInstance = InitInstanceBody(InInstanceIndex, InOutNewInstanceData, PhysScene);
+
+		InstanceBodies.Insert(NewBodyInstance, InInstanceIndex);
+	}
+}
+
+FBodyInstance* UPaperGroupedSpriteComponent::InitInstanceBody(int32 InstanceIndex, const FSpriteInstanceData& InstanceData, FPhysScene* PhysScene)
+{
+	FBodyInstance* NewBodyInstance = nullptr;
+
+	if (InstanceData.SourceSprite != nullptr)
+	{
+		if (UBodySetup* BodySetup = InstanceData.SourceSprite->BodySetup)
+		{
+			NewBodyInstance = new FBodyInstance();
+
+			NewBodyInstance->CopyBodyInstancePropertiesFrom(&BodyInstance);
+			NewBodyInstance->InstanceBodyIndex = InstanceIndex; // Set body index 
+			NewBodyInstance->bAutoWeld = false;
+
+			// make sure we never enable bSimulatePhysics for ISMComps
+			NewBodyInstance->bSimulatePhysics = false;
+
+			const FTransform InstanceTransform(FTransform(InstanceData.Transform) * ComponentToWorld);
+			NewBodyInstance->InitBody(BodySetup, InstanceTransform, this, PhysScene);
+		}
+	}
+
+	return NewBodyInstance;
+}
+
+void UPaperGroupedSpriteComponent::RebuildInstances()
+{
+	// update the physics state
+	if (bPhysicsStateCreated)
+	{
+		// TODO: it may be possible to instead just update the BodyInstanceIndex for all bodies after the removed instance. 
+		ClearAllInstanceBodies();
+		CreateAllInstanceBodies();
+	}
+
+	// Indicate we need to update render state to reflect changes
+	MarkRenderStateDirty();
+}
+
+void UPaperGroupedSpriteComponent::RebuildMaterialList()
+{
+	InstanceMaterials.Reset();
+
+	//@TODO: Need to rebuild the OverrideMaterials array and update the per-sprite material indices as part of this process
+
+	for (const FSpriteInstanceData& InstanceData : PerInstanceSpriteData)
+	{
+		UpdateMaterialList(InstanceData.SourceSprite);
+	}
+}
+
+int32 UPaperGroupedSpriteComponent::UpdateMaterialList(UPaperSprite* Sprite)
+{
+	int32 Result = INDEX_NONE;
+
+	if (Sprite != nullptr)
+	{
+		if (UMaterialInterface* SpriteMaterial = Sprite->GetMaterial(0))
+		{
+			Result = InstanceMaterials.AddUnique(SpriteMaterial);
+		}
+	}
+
+	return Result;
+}
+
+void UPaperGroupedSpriteComponent::GetUsedTextures(TArray<UTexture*>& OutTextures, EMaterialQualityLevel::Type QualityLevel)
+{
+	// Get the textures referenced by any sprite instances
+	for (const FSpriteInstanceData& InstanceData : PerInstanceSpriteData)
+	{
+		if (InstanceData.SourceSprite != nullptr)
+		{
+			if (UTexture* BakedTexture = InstanceData.SourceSprite->GetBakedTexture())
+			{
+				OutTextures.AddUnique(BakedTexture);
+			}
+
+			FAdditionalSpriteTextureArray AdditionalTextureList;
+			InstanceData.SourceSprite->GetBakedAdditionalSourceTextures(/*out*/ AdditionalTextureList);
+			for (UTexture* AdditionalTexture : AdditionalTextureList)
+			{
+				if (AdditionalTexture != nullptr)
+				{
+					OutTextures.AddUnique(AdditionalTexture);
+				}
+			}
+		}
+	}
+
+	// Get any textures referenced by our materials
+	Super::GetUsedTextures(OutTextures, QualityLevel);
+}
+
+UMaterialInterface* UPaperGroupedSpriteComponent::GetMaterial(int32 MaterialIndex) const
+{
+	if (OverrideMaterials.IsValidIndex(MaterialIndex))
+	{
+		if (UMaterialInterface* OverrideMaterial = OverrideMaterials[MaterialIndex])
+		{
+			return OverrideMaterial;
+		}
+	}
+
+	if (InstanceMaterials.IsValidIndex(MaterialIndex))
+	{
+		return InstanceMaterials[MaterialIndex];
+	}
+
+	return nullptr;
+}
+
+int32 UPaperGroupedSpriteComponent::GetNumMaterials() const
+{
+	return FMath::Max<int32>(OverrideMaterials.Num(), FMath::Max<int32>(InstanceMaterials.Num(), 1));
+}
+
+bool UPaperGroupedSpriteComponent::ContainsSprite(UPaperSprite* SpriteAsset) const
+{
+	for (const FSpriteInstanceData& InstanceData : PerInstanceSpriteData)
+	{
+		if (InstanceData.SourceSprite == SpriteAsset)
+		{
+			return true;
+		}
+	}
+
+	return false;
+}
+
+//////////////////////////////////////////////////////////////////////////
+
+#undef LOCTEXT_NAMESPACE
diff --git a/Engine/Plugins/2D/Paper2D/Source/Paper2D/Private/PaperSprite.cpp b/Engine/Plugins/2D/Paper2D/Source/Paper2D/Private/PaperSprite.cpp
--- a/Engine/Plugins/2D/Paper2D/Source/Paper2D/Private/PaperSprite.cpp
+++ b/Engine/Plugins/2D/Paper2D/Source/Paper2D/Private/PaperSprite.cpp
@@ -9,6 +9,7 @@
 #include "PaperGeomTools.h"
 #include "PaperSpriteComponent.h"
 #include "PaperFlipbookComponent.h"
+#include "PaperGroupedSpriteComponent.h"
 #include "SpriteDrawCall.h"
 
 #if WITH_EDITOR
@@ -428,31 +429,34 @@ public:
 	FSpriteReregisterContext(UPaperSprite* TargetAsset)
 	{
 		// Look at sprite components
-		for (TObjectIterator<UPaperSpriteComponent> SpriteIt; SpriteIt; ++SpriteIt)
+		for (UPaperSpriteComponent* TestComponent : TObjectRange<UPaperSpriteComponent>())
 		{
-			if (UPaperSpriteComponent* TestComponent = *SpriteIt)
+			if (TestComponent->GetSprite() == TargetAsset)
 			{
-				if (TestComponent->GetSprite() == TargetAsset)
-				{
-					AddComponentToRefresh(TestComponent);
-				}
+				AddComponentToRefresh(TestComponent);
 			}
 		}
 
 		// Look at flipbook components
-		for (TObjectIterator<UPaperFlipbookComponent> FlipbookIt; FlipbookIt; ++FlipbookIt)
+		for (UPaperFlipbookComponent* TestComponent : TObjectRange<UPaperFlipbookComponent>())
 		{
-			if (UPaperFlipbookComponent* TestComponent = *FlipbookIt)
+			if (UPaperFlipbook* Flipbook = TestComponent->GetFlipbook())
 			{
-				if (UPaperFlipbook* Flipbook = TestComponent->GetFlipbook())
+				if (Flipbook->ContainsSprite(TargetAsset))
 				{
-					if (Flipbook->ContainsSprite(TargetAsset))
-					{
-						AddComponentToRefresh(TestComponent);
-					}
+					AddComponentToRefresh(TestComponent);
 				}
 			}
 		}
+
+		// Look at grouped sprite components
+		for (UPaperGroupedSpriteComponent* TestComponent : TObjectRange<UPaperGroupedSpriteComponent>())
+		{
+			if (TestComponent->ContainsSprite(TargetAsset))
+			{
+				AddComponentToRefresh(TestComponent);
+			}
+		}
 	}
 
 protected:
diff --git a/Engine/Plugins/2D/Paper2D/Source/Paper2DEditor/Private/GroupedSpriteDetailsCustomization.cpp b/Engine/Plugins/2D/Paper2D/Source/Paper2DEditor/Private/GroupedSpriteDetailsCustomization.cpp
new file mode 100644
--- /dev/null
+++ b/Engine/Plugins/2D/Paper2D/Source/Paper2DEditor/Private/GroupedSpriteDetailsCustomization.cpp
@@ -0,0 +1,183 @@
+// Copyright 1998-2015 Epic Games, Inc. All Rights Reserved.
+
+#include "Paper2DEditorPrivatePCH.h"
+#include "GroupedSpriteDetailsCustomization.h"
+
+#include "PaperSpriteActor.h"
+#include "PaperSpriteComponent.h"
+#include "PaperGroupedSpriteComponent.h"
+#include "ScopedTransaction.h"
+#include "ILayers.h"
+#include "ComponentReregisterContext.h"
+
+#define LOCTEXT_NAMESPACE "SpriteEditor"
+
+//////////////////////////////////////////////////////////////////////////
+// FGroupedSpriteComponentDetailsCustomization
+
+TSharedRef<IDetailCustomization> FGroupedSpriteComponentDetailsCustomization::MakeInstance()
+{
+	return MakeShareable(new FGroupedSpriteComponentDetailsCustomization());
+}
+
+FGroupedSpriteComponentDetailsCustomization::FGroupedSpriteComponentDetailsCustomization()
+{
+}
+
+void FGroupedSpriteComponentDetailsCustomization::CustomizeDetails(IDetailLayoutBuilder& DetailBuilder)
+{
+	// Create a category so this is displayed early in the properties
+	IDetailCategoryBuilder& SpriteCategory = DetailBuilder.EditCategory("Sprite", FText::GetEmpty(), ECategoryPriority::Important);
+
+	ObjectsBeingCustomized.Empty();
+	DetailBuilder.GetObjectsBeingCustomized(/*out*/ ObjectsBeingCustomized);
+
+	{
+		// Expose split buttons
+		FDetailWidgetRow& SplitRow = SpriteCategory.AddCustomRow(LOCTEXT("SplitSearchText", "Split"))
+			.WholeRowContent()
+			[
+				SNew(SHorizontalBox)
+				+ SHorizontalBox::Slot()
+				.AutoWidth()
+				.Padding(2.0f, 0.0f)
+				.VAlign(VAlign_Center)
+				.HAlign(HAlign_Left)
+				[
+					SNew(SButton)
+					.Text(LOCTEXT("SplitSprites", "Split Sprites"))
+					.ToolTipText(LOCTEXT("SplitSprites_Tooltip", "Splits all sprite instances into separate sprite actors or components"))
+					.OnClicked(this, &FGroupedSpriteComponentDetailsCustomization::SplitSprites)
+				]
+			];
+	}
+}
+
+void FGroupedSpriteComponentDetailsCustomization::BuildHarvestList(const TArray<TWeakObjectPtr<UObject>>& ObjectsToConsider, TSubclassOf<UActorComponent> HarvestClassType, TArray<UActorComponent*>& OutComponentsToHarvest, TArray<AActor*>& OutActorsToDelete)
+{
+	// Determine the components to harvest
+	for (TWeakObjectPtr<UObject> WeakObject : ObjectsToConsider)
+	{
+		UObject* Object = WeakObject.Get();
+		if (Object == nullptr)
+		{
+			continue;
+		}
+
+		if (AActor* SelectedActor = Cast<AActor>(Object))
+		{
+			TArray<UActorComponent*> Components;
+			SelectedActor->GetComponents(Components);
+
+			bool bHadHarvestableComponents = false;
+			for (UActorComponent* Component : Components)
+			{
+				if (Component->IsA(HarvestClassType) && !Component->IsEditorOnly())
+				{
+					bHadHarvestableComponents = true;
+					OutComponentsToHarvest.Add(Component);
+				}
+			}
+
+			if (bHadHarvestableComponents)
+			{
+				OutActorsToDelete.Add(SelectedActor);
+			}
+		}
+		else if (UActorComponent* SelectedComponent = Cast<UActorComponent>(Object))
+		{
+			if (SelectedComponent->IsA(HarvestClassType) && !SelectedComponent->IsEditorOnly())
+			{
+				OutComponentsToHarvest.Add(SelectedComponent);
+				if (AActor* SelectedComponentOwner = SelectedComponent->GetOwner())
+				{
+					OutActorsToDelete.AddUnique(SelectedComponentOwner);
+				}
+			}
+		}
+	}
+}
+
+FReply FGroupedSpriteComponentDetailsCustomization::SplitSprites()
+{
+	TArray<UActorComponent*> ComponentsToHarvest;
+	TSubclassOf<UActorComponent> HarvestClassType = UPaperSpriteComponent::StaticClass();
+	TArray<AActor*> ActorsToDelete;
+	TArray<AActor*> ActorsCreated;
+	
+	FGroupedSpriteComponentDetailsCustomization::BuildHarvestList(ObjectsBeingCustomized, UPaperGroupedSpriteComponent::StaticClass(), /*out*/ ComponentsToHarvest, /*out*/ ActorsToDelete);
+
+	if (ComponentsToHarvest.Num() > 0)
+	{
+		if (UWorld* World = ComponentsToHarvest[0]->GetWorld())
+		{
+			const FScopedTransaction Transaction(LOCTEXT("SplitSprites", "Split sprite instances"));
+
+			// Create an instance from each item of each batch component that we're harvesting
+			for (UActorComponent* SourceComponent : ComponentsToHarvest)
+			{
+				UPaperGroupedSpriteComponent* SourceBatchComponent = CastChecked<UPaperGroupedSpriteComponent>(SourceComponent);
+
+				for (FSpriteInstanceData& InstanceData : SourceBatchComponent->PerInstanceSpriteData)
+				{
+					if (InstanceData.SourceSprite != nullptr)
+					{
+						const FTransform InstanceTransform(FTransform(InstanceData.Transform) * SourceBatchComponent->ComponentToWorld);
+
+						FActorSpawnParameters SpawnParams;
+						SpawnParams.bDeferConstruction = true;
+						if (APaperSpriteActor* SpawnedActor = World->SpawnActor<APaperSpriteActor>(SpawnParams))
+						{
+							UPaperSpriteComponent* SpawnedSpriteComponent = SpawnedActor->GetRenderComponent();
+
+							{
+								FComponentReregisterContext ReregisterContext(SpawnedSpriteComponent);
+
+								SpawnedSpriteComponent->Modify();
+								SpawnedSpriteComponent->SetSpriteColor(InstanceData.VertexColor);
+								SpawnedSpriteComponent->SetSprite(InstanceData.SourceSprite);
+
+								//@TODO: Sprite component material!
+								// if (SourceBatchComponent->GetMaterial() != InstanceData.SourceSprite->GetMaterial(0))
+								// {
+								//SpawnedSpriteComponent->SetMaterial();
+								// }
+							}
+
+							UGameplayStatics::FinishSpawningActor(SpawnedActor, InstanceTransform);
+							ActorsCreated.Add(SpawnedActor);
+						}
+					}
+				}
+			}
+
+			// Delete the existing actor instances
+			for (AActor* ActorToDelete : ActorsToDelete)
+			{
+				// Remove from active selection in editor
+				GEditor->SelectActor(ActorToDelete, /*bSelected=*/ false, /*bNotify=*/ false);
+				GEditor->Layers->DisassociateActorFromLayers(ActorToDelete);
+				World->EditorDestroyActor(ActorToDelete, /*bShouldModifyLevel=*/ true);
+			}
+
+			// Select the new replacements instance
+			GEditor->SelectNone(/*bNoteSelectionChange=*/ false, false, false);
+			for (AActor* SpawnedActor : ActorsCreated)
+			{
+				GEditor->SelectActor(SpawnedActor, /*bSelected=*/ true, /*bNotify=*/ true);
+			}
+			GEditor->NoteSelectionChange();
+		}
+		else
+		{
+			// We're in the BP editor and don't currently support splitting here!
+			UE_LOG(LogPaper2DEditor, Warning, TEXT("Splitting sprites in the Blueprint editor is not currently supported"));
+		}
+	}
+
+	return FReply::Handled();
+}
+
+//////////////////////////////////////////////////////////////////////////
+
+#undef LOCTEXT_NAMESPACE
diff --git a/Engine/Plugins/2D/Paper2D/Source/Paper2DEditor/Private/GroupedSpriteDetailsCustomization.h b/Engine/Plugins/2D/Paper2D/Source/Paper2DEditor/Private/GroupedSpriteDetailsCustomization.h
new file mode 100644
--- /dev/null
+++ b/Engine/Plugins/2D/Paper2D/Source/Paper2DEditor/Private/GroupedSpriteDetailsCustomization.h
@@ -0,0 +1,29 @@
+// Copyright 1998-2015 Epic Games, Inc. All Rights Reserved.
+
+#pragma once
+
+#include "PropertyEditing.h"
+
+//////////////////////////////////////////////////////////////////////////
+// FGroupedSpriteComponentDetailsCustomization
+
+class FGroupedSpriteComponentDetailsCustomization : public IDetailCustomization
+{
+public:
+	/** Makes a new instance of this detail layout class for a specific detail view requesting it */
+	static TSharedRef<IDetailCustomization> MakeInstance();
+
+	// IDetailCustomization interface
+	virtual void CustomizeDetails(IDetailLayoutBuilder& DetailLayout) override;
+	// End of IDetailCustomization interface
+
+	static void BuildHarvestList(const TArray<TWeakObjectPtr<UObject>>& ObjectsToConsider, TSubclassOf<UActorComponent> HarvestClassType, TArray<UActorComponent*>& OutComponentsToHarvest, TArray<AActor*>& OutActorsToDelete);
+
+protected:
+	FGroupedSpriteComponentDetailsCustomization();
+
+	FReply SplitSprites();
+
+protected:
+	TArray<TWeakObjectPtr<UObject>> ObjectsBeingCustomized;
+};
\ No newline at end of file
diff --git a/Engine/Plugins/2D/Paper2D/Source/Paper2DEditor/Private/Paper2DEditorModule.cpp b/Engine/Plugins/2D/Paper2D/Source/Paper2DEditor/Private/Paper2DEditorModule.cpp
--- a/Engine/Plugins/2D/Paper2D/Source/Paper2DEditor/Private/Paper2DEditorModule.cpp
+++ b/Engine/Plugins/2D/Paper2D/Source/Paper2DEditor/Private/Paper2DEditorModule.cpp
@@ -48,6 +48,10 @@
 #include "Atlasing/PaperAtlasGenerator.h"
 #include "PaperSpriteAtlas.h"
 
+// Grouped sprite support
+#include "PaperGroupedSpriteComponent.h"
+#include "GroupedSpriteDetailsCustomization.h"
+
 // Settings
 #include "PaperRuntimeSettings.h"
 #include "ISettingsModule.h"
@@ -146,6 +150,7 @@ public:
 			PropertyModule.RegisterCustomClassLayout(UPaperSprite::StaticClass()->GetFName(), FOnGetDetailCustomizationInstance::CreateStatic(&FSpriteDetailsCustomization::MakeInstance));
 			PropertyModule.RegisterCustomClassLayout(UPaperSpriteComponent::StaticClass()->GetFName(), FOnGetDetailCustomizationInstance::CreateStatic(&FSpriteComponentDetailsCustomization::MakeInstance));
 			PropertyModule.RegisterCustomClassLayout(UPaperFlipbookComponent::StaticClass()->GetFName(), FOnGetDetailCustomizationInstance::CreateStatic(&FFlipbookComponentDetailsCustomization::MakeInstance));
+			PropertyModule.RegisterCustomClassLayout(UPaperGroupedSpriteComponent::StaticClass()->GetFName(), FOnGetDetailCustomizationInstance::CreateStatic(&FGroupedSpriteComponentDetailsCustomization::MakeInstance));
 
 			//@TODO: Struct registration should happen using ::StaticStruct, not by string!!!
 			//PropertyModule.RegisterCustomPropertyTypeLayout( "SpritePolygonCollection", FOnGetPropertyTypeCustomizationInstance::CreateStatic( &FSpritePolygonCollectionCustomization::MakeInstance ) );
diff --git a/Engine/Plugins/2D/Paper2D/Source/Paper2DEditor/Private/SpriteEditor/SpriteComponentDetailsCustomization.cpp b/Engine/Plugins/2D/Paper2D/Source/Paper2DEditor/Private/SpriteEditor/SpriteComponentDetailsCustomization.cpp
--- a/Engine/Plugins/2D/Paper2D/Source/Paper2DEditor/Private/SpriteEditor/SpriteComponentDetailsCustomization.cpp
+++ b/Engine/Plugins/2D/Paper2D/Source/Paper2DEditor/Private/SpriteEditor/SpriteComponentDetailsCustomization.cpp
@@ -1,7 +1,13 @@
 // Copyright 1998-2015 Epic Games, Inc. All Rights Reserved.
 
 #include "Paper2DEditorPrivatePCH.h"
+#include "PaperSpriteComponent.h"
+#include "PaperGroupedSpriteComponent.h"
 #include "SpriteComponentDetailsCustomization.h"
+#include "ScopedTransaction.h"
+#include "ILayers.h"
+#include "GroupedSpriteDetailsCustomization.h"
+#include "ComponentReregisterContext.h"
 
 #define LOCTEXT_NAMESPACE "SpriteEditor"
 
@@ -16,7 +22,113 @@ TSharedRef<IDetailCustomization> FSpriteComponentDetailsCustomization::MakeInsta
 void FSpriteComponentDetailsCustomization::CustomizeDetails(IDetailLayoutBuilder& DetailBuilder)
 {
 	// Create a category so this is displayed early in the properties
-	DetailBuilder.EditCategory("Sprite", FText::GetEmpty(), ECategoryPriority::Important);
+	IDetailCategoryBuilder& SpriteCategory = DetailBuilder.EditCategory("Sprite", FText::GetEmpty(), ECategoryPriority::Important);
+
+	ObjectsBeingCustomized.Empty();
+	DetailBuilder.GetObjectsBeingCustomized(/*out*/ ObjectsBeingCustomized);
+
+	if (ObjectsBeingCustomized.Num() > 1)
+	{
+		// Expose merge buttons
+		FDetailWidgetRow& MergeRow = SpriteCategory.AddCustomRow(LOCTEXT("MergeSearchText", "Merge"))
+			.WholeRowContent()
+			[
+				SNew(SHorizontalBox)
+				+SHorizontalBox::Slot()
+				.AutoWidth()
+				.Padding(2.0f, 0.0f)
+				.VAlign(VAlign_Center)
+				.HAlign(HAlign_Left)
+				[
+					SNew(SButton)
+					.Text(LOCTEXT("MergeSprites", "Merge Sprites"))
+					.ToolTipText(LOCTEXT("MergeSprites_Tooltip", "Merges all selected sprites into entries on a single InstancedSpriteComponent"))
+					.OnClicked(this, &FSpriteComponentDetailsCustomization::MergeSprites)
+				]
+			];
+	}
+}
+
+
+FBox FSpriteComponentDetailsCustomization::ComputeBoundsForComponents(const TArray<UActorComponent*>& ComponentList)
+{
+	FBox BoundsOfList(ForceInit);
+
+	for (const UActorComponent* Component : ComponentList)
+	{
+		if (const USceneComponent* SceneComponent = Cast<const USceneComponent>(Component))
+		{
+			BoundsOfList += SceneComponent->Bounds.GetBox();
+		}
+	}
+
+	return BoundsOfList;
+}
+
+FReply FSpriteComponentDetailsCustomization::MergeSprites()
+{
+	TArray<UActorComponent*> ComponentsToHarvest;
+	TSubclassOf<UActorComponent> HarvestClassType = UPaperSpriteComponent::StaticClass();
+	TArray<AActor*> ActorsToDelete;
+
+	FGroupedSpriteComponentDetailsCustomization::BuildHarvestList(ObjectsBeingCustomized, UPaperSpriteComponent::StaticClass(), /*out*/ ComponentsToHarvest, /*out*/ ActorsToDelete);
+
+	if (ComponentsToHarvest.Num() > 0)
+	{
+		if (UWorld* World = ComponentsToHarvest[0]->GetWorld())
+		{
+			const FBox ComponentBounds = ComputeBoundsForComponents(ComponentsToHarvest);
+			const FTransform MergedWorldTM(ComponentBounds.GetCenter());
+
+			FActorSpawnParameters SpawnParams;
+			SpawnParams.bDeferConstruction = true;
+			if (AActor* SpawnedActor = World->SpawnActor<AActor>(SpawnParams))
+			{
+				const FScopedTransaction Transaction(LOCTEXT("MergeSprites", "Merge sprite instances"));
+
+				// Create a merged sprite component
+				UPaperGroupedSpriteComponent* MergedSpriteComponent = NewObject<UPaperGroupedSpriteComponent>(SpawnedActor);
+
+				// Create an instance from each sprite component that we're harvesting
+				for (UActorComponent* SourceComponent : ComponentsToHarvest)
+				{
+					UPaperSpriteComponent* SourceSpriteComponent = CastChecked<UPaperSpriteComponent>(SourceComponent);
+
+					UPaperSprite* Sprite = SourceSpriteComponent->GetSprite();
+					const FLinearColor SpriteColor = SourceSpriteComponent->GetSpriteColor();
+					const FTransform RelativeSpriteTransform = SourceSpriteComponent->GetComponentTransform().GetRelativeTransform(MergedWorldTM);
+
+					MergedSpriteComponent->AddInstance(RelativeSpriteTransform, Sprite, FColor(SpriteColor));
+				}
+
+				// Finalize the new component
+				SpawnedActor->AddInstanceComponent(MergedSpriteComponent);
+				SpawnedActor->SetRootComponent(MergedSpriteComponent);
+				MergedSpriteComponent->OnComponentCreated();
+				MergedSpriteComponent->RegisterComponent();
+				UGameplayStatics::FinishSpawningActor(SpawnedActor, MergedWorldTM);
+
+				// Delete the existing actor instances
+				for (AActor* ActorToDelete : ActorsToDelete)
+				{
+					// Remove from active selection in editor
+					GEditor->SelectActor(ActorToDelete, /*bSelected=*/ false, /*bNotify=*/ false);
+					GEditor->Layers->DisassociateActorFromLayers(ActorToDelete);
+					World->EditorDestroyActor(ActorToDelete, /*bShouldModifyLevel=*/ true);
+				}
+
+				// Select the new replacement instance
+				GEditor->SelectActor(SpawnedActor, /*bSelected=*/ true, /*bNotify=*/ true);
+			}
+			else
+			{
+				// We're in the BP editor and don't currently support merging here!
+				UE_LOG(LogPaper2DEditor, Warning, TEXT("Merging sprites in the Blueprint editor is not currently supported"));
+			}
+		}
+	}
+
+	return FReply::Handled();
 }
 
 //////////////////////////////////////////////////////////////////////////
diff --git a/Engine/Plugins/2D/Paper2D/Source/Paper2DEditor/Private/SpriteEditor/SpriteComponentDetailsCustomization.h b/Engine/Plugins/2D/Paper2D/Source/Paper2DEditor/Private/SpriteEditor/SpriteComponentDetailsCustomization.h
--- a/Engine/Plugins/2D/Paper2D/Source/Paper2DEditor/Private/SpriteEditor/SpriteComponentDetailsCustomization.h
+++ b/Engine/Plugins/2D/Paper2D/Source/Paper2DEditor/Private/SpriteEditor/SpriteComponentDetailsCustomization.h
@@ -16,4 +16,12 @@ public:
 	// IDetailCustomization interface
 	virtual void CustomizeDetails(IDetailLayoutBuilder& DetailLayout) override;
 	// End of IDetailCustomization interface
+
+protected:
+	FReply MergeSprites();
+
+	static FBox ComputeBoundsForComponents(const TArray<UActorComponent*>& ComponentList);
+
+protected:
+	TArray<TWeakObjectPtr<UObject>> ObjectsBeingCustomized;
 };
