Paper2D: Fixed non-editor build (missing WITH_EDITOR around CheckForErrors)

diff --git a/Engine/Plugins/2D/Paper2D/Source/Paper2D/Classes/PaperGroupedSpriteComponent.h b/Engine/Plugins/2D/Paper2D/Source/Paper2D/Classes/PaperGroupedSpriteComponent.h
--- a/Engine/Plugins/2D/Paper2D/Source/Paper2D/Classes/PaperGroupedSpriteComponent.h
+++ b/Engine/Plugins/2D/Paper2D/Source/Paper2D/Classes/PaperGroupedSpriteComponent.h
@@ -100,7 +100,9 @@ public:
 	virtual bool ShouldCreatePhysicsState() const override;
 	virtual void CreatePhysicsState() override;
 	virtual void DestroyPhysicsState() override;
+#if WITH_EDITOR
 	virtual void CheckForErrors() override;
+#endif
 	// End of UActorComponent interface 
 
 	// UPrimitiveComponent interface
diff --git a/Engine/Plugins/2D/Paper2D/Source/Paper2D/Private/PaperGroupedSpriteComponent.cpp b/Engine/Plugins/2D/Paper2D/Source/Paper2D/Private/PaperGroupedSpriteComponent.cpp
--- a/Engine/Plugins/2D/Paper2D/Source/Paper2D/Private/PaperGroupedSpriteComponent.cpp
+++ b/Engine/Plugins/2D/Paper2D/Source/Paper2D/Private/PaperGroupedSpriteComponent.cpp
@@ -174,6 +174,7 @@ void UPaperGroupedSpriteComponent::DestroyPhysicsState()
 	ClearAllInstanceBodies();
 }
 
+#if WITH_EDITOR
 void UPaperGroupedSpriteComponent::CheckForErrors()
 {
 	AActor* Owner = GetOwner();
@@ -197,6 +198,7 @@ void UPaperGroupedSpriteComponent::CheckForErrors()
 
 	Super::CheckForErrors();
 }
+#endif
 
 FPrimitiveSceneProxy* UPaperGroupedSpriteComponent::CreateSceneProxy()
 {
