Paper2D: Sprite group component improvements
- Merged AddInstance and AddInstanceWorldSpace with a bWorldSpace parameter to match the GetInstanceTransform method
- Changed Color parameters to FLinearColor so they can have literal color picker pins (render data is still stored as a FColor)
- Added an UpdateInstanceColor method

diff --git a/Engine/Plugins/2D/Paper2D/Source/Paper2D/Classes/PaperGroupedSpriteComponent.h b/Engine/Plugins/2D/Paper2D/Source/Paper2D/Classes/PaperGroupedSpriteComponent.h
--- a/Engine/Plugins/2D/Paper2D/Source/Paper2D/Classes/PaperGroupedSpriteComponent.h
+++ b/Engine/Plugins/2D/Paper2D/Source/Paper2D/Classes/PaperGroupedSpriteComponent.h
@@ -68,13 +68,9 @@ protected:
 
 public:
 
-	/** Add an instance to this component. Transform is given in local space of this component.  */
+	/** Add an instance to this component. Transform can be given either in the local space of this component or world space.  */
 	UFUNCTION(BlueprintCallable, Category="Components|Sprite")
-	virtual int32 AddInstance(const FTransform& LocalTransform, UPaperSprite* Sprite, FColor Color = FColor(255, 255, 255));
-
-	/** Add an instance to this component. Transform is given in world space. */
-	UFUNCTION(BlueprintCallable, Category = "Components|Sprite")
-	int32 AddInstanceWorldSpace(const FTransform& WorldTransform, UPaperSprite* Sprite, FColor Color = FColor(255, 255, 255));
+	virtual int32 AddInstance(const FTransform& LocalTransform, UPaperSprite* Sprite, bool bWorldSpace = false, FLinearColor Color = FLinearColor::White);
 
 	/** Get the transform for the instance specified. Instance is returned in local space of this component unless bWorldSpace is set.  Returns True on success. */
 	UFUNCTION(BlueprintCallable, Category = "Components|Sprite")
@@ -82,7 +78,11 @@ public:
 	
 	/** Update the transform for the instance specified. Instance is given in local space of this component unless bWorldSpace is set.  Returns True on success. */
 	UFUNCTION(BlueprintCallable, Category = "Components|Sprite")
-	virtual bool UpdateInstanceTransform(int32 InstanceIndex, const FTransform& NewInstanceTransform, bool bWorldSpace = false, bool bMarkRenderStateDirty = false);
+	virtual bool UpdateInstanceTransform(int32 InstanceIndex, const FTransform& NewInstanceTransform, bool bWorldSpace = false, bool bMarkRenderStateDirty = true);
+
+	/** Update the color for the instance specified. Returns True on success. */
+	UFUNCTION(BlueprintCallable, Category = "Components|Sprite")
+	virtual bool UpdateInstanceColor(int32 InstanceIndex, FLinearColor NewInstanceColor, bool bMarkRenderStateDirty = true);
 
 	/** Remove the instance specified. Returns True on success. */
 	UFUNCTION(BlueprintCallable, Category = "Components|Sprite")
diff --git a/Engine/Plugins/2D/Paper2D/Source/Paper2D/Private/PaperGroupedSpriteComponent.cpp b/Engine/Plugins/2D/Paper2D/Source/Paper2D/Private/PaperGroupedSpriteComponent.cpp
--- a/Engine/Plugins/2D/Paper2D/Source/Paper2D/Private/PaperGroupedSpriteComponent.cpp
+++ b/Engine/Plugins/2D/Paper2D/Source/Paper2D/Private/PaperGroupedSpriteComponent.cpp
@@ -27,10 +27,12 @@ UPaperGroupedSpriteComponent::UPaperGroupedSpriteComponent(const FObjectInitiali
 	BodyInstance.bSimulatePhysics = false;
 }
 
-int32 UPaperGroupedSpriteComponent::AddInstance(const FTransform& LocalTransform, UPaperSprite* Sprite, FColor Color)
+int32 UPaperGroupedSpriteComponent::AddInstance(const FTransform& Transform, UPaperSprite* Sprite, bool bWorldSpace, FLinearColor Color)
 {
 	const int32 NewInstanceIndex = PerInstanceSpriteData.Num();
 
+	const FTransform LocalTransform(bWorldSpace ? Transform.GetRelativeTransform(ComponentToWorld) : Transform);
+
 	FSpriteInstanceData& NewInstanceData = *new (PerInstanceSpriteData) FSpriteInstanceData();
 	SetupNewInstanceData(NewInstanceData, NewInstanceIndex, LocalTransform, Sprite, Color);
 
@@ -41,13 +43,6 @@ int32 UPaperGroupedSpriteComponent::AddInstance(const FTransform& LocalTransform
 	return NewInstanceIndex;
 }
 
-int32 UPaperGroupedSpriteComponent::AddInstanceWorldSpace(const FTransform& WorldTransform, UPaperSprite* Sprite, FColor Color)
-{
-	// Transform from world space to local space
-	const FTransform RelativeTM = WorldTransform.GetRelativeTransform(ComponentToWorld);
-	return AddInstance(RelativeTM, Sprite, Color);
-}
-
 bool UPaperGroupedSpriteComponent::GetInstanceTransform(int32 InstanceIndex, FTransform& OutInstanceTransform, bool bWorldSpace) const
 {
 	if (!PerInstanceSpriteData.IsValidIndex(InstanceIndex))
@@ -105,6 +100,25 @@ bool UPaperGroupedSpriteComponent::UpdateInstanceTransform(int32 InstanceIndex,
 	return true;
 }
 
+bool UPaperGroupedSpriteComponent::UpdateInstanceColor(int32 InstanceIndex, FLinearColor NewInstanceColor, bool bMarkRenderStateDirty)
+{
+	if (PerInstanceSpriteData.IsValidIndex(InstanceIndex))
+	{
+		FSpriteInstanceData& InstanceData = PerInstanceSpriteData[InstanceIndex];
+		InstanceData.VertexColor = NewInstanceColor;
+
+		if (bMarkRenderStateDirty)
+		{
+			MarkRenderStateDirty();
+		}
+		return true;
+	}
+	else
+	{
+		return false;
+	}
+}	
+
 bool UPaperGroupedSpriteComponent::RemoveInstance(int32 InstanceIndex)
 {
 	if (!PerInstanceSpriteData.IsValidIndex(InstanceIndex))
diff --git a/Engine/Plugins/2D/Paper2D/Source/Paper2D/Private/PaperRenderSceneProxy.h b/Engine/Plugins/2D/Paper2D/Source/Paper2D/Private/PaperRenderSceneProxy.h
--- a/Engine/Plugins/2D/Paper2D/Source/Paper2D/Private/PaperRenderSceneProxy.h
+++ b/Engine/Plugins/2D/Paper2D/Source/Paper2D/Private/PaperRenderSceneProxy.h
@@ -160,10 +160,10 @@ public:
 	void SetTransientTextureOverride_RenderThread(const UTexture* InTextureToModifyOverrideFor, UTexture* InOverrideTexture);
 #endif
 
-	class UMaterialInterface* GetMaterial() const
-	{
-		return Material;
-	}
+// 	class UMaterialInterface* GetMaterial() const
+// 	{
+// 		return Material;
+// 	}
 
 protected:
 	virtual void GetDynamicMeshElementsForView(const FSceneView* View, int32 ViewIndex, FMeshElementCollector& Collector) const;
