Paper2D: Added a grouped sprite actor to avoid the need to click on the component to split it, and to provide easier Ctrl+B / Ctrl+E support

diff --git a/Engine/Plugins/2D/Paper2D/Source/Paper2D/Classes/PaperGroupedSpriteActor.h b/Engine/Plugins/2D/Paper2D/Source/Paper2D/Classes/PaperGroupedSpriteActor.h
new file mode 100644
--- /dev/null
+++ b/Engine/Plugins/2D/Paper2D/Source/Paper2D/Classes/PaperGroupedSpriteActor.h
@@ -0,0 +1,32 @@
+// Copyright 1998-2015 Epic Games, Inc. All Rights Reserved.
+
+#pragma once
+
+#include "PaperGroupedSpriteComponent.h"
+#include "PaperGroupedSpriteActor.generated.h"
+
+/**
+ * A group of sprites that will be rendered and culled as a single unit
+ *
+ * This actor is created when you Merge several sprite components together.
+ * it is just a thin wrapper around a UPaperGroupedSpriteComponent.
+ */
+UCLASS(ComponentWrapperClass)
+class PAPER2D_API APaperGroupedSpriteActor : public AActor
+{
+	GENERATED_UCLASS_BODY()
+
+private:
+	UPROPERTY(Category=Sprite, VisibleAnywhere, BlueprintReadOnly, meta=(ExposeFunctionCategories = "Sprite,Rendering,Physics,Components|Sprite", AllowPrivateAccess="true"))
+	UPaperGroupedSpriteComponent* RenderComponent;
+
+public:
+	// AActor interface
+#if WITH_EDITOR
+	virtual bool GetReferencedContentObjects(TArray<UObject*>& Objects) const override;
+#endif
+	// End of AActor interface
+
+	/** Returns RenderComponent subobject **/
+	FORCEINLINE UPaperGroupedSpriteComponent* GetRenderComponent() const { return RenderComponent; }
+};
diff --git a/Engine/Plugins/2D/Paper2D/Source/Paper2D/Classes/PaperGroupedSpriteComponent.h b/Engine/Plugins/2D/Paper2D/Source/Paper2D/Classes/PaperGroupedSpriteComponent.h
--- a/Engine/Plugins/2D/Paper2D/Source/Paper2D/Classes/PaperGroupedSpriteComponent.h
+++ b/Engine/Plugins/2D/Paper2D/Source/Paper2D/Classes/PaperGroupedSpriteComponent.h
@@ -100,6 +100,7 @@ public:
 	virtual bool ShouldCreatePhysicsState() const override;
 	virtual void CreatePhysicsState() override;
 	virtual void DestroyPhysicsState() override;
+	virtual const UObject* AdditionalStatObject() const override;
 #if WITH_EDITOR
 	virtual void CheckForErrors() override;
 #endif
@@ -125,6 +126,9 @@ public:
 	// Returns true if this component references the specified sprite asset
 	bool ContainsSprite(UPaperSprite* SpriteAsset) const;
 
+	// Adds all referenced sprite assets to the specified list
+	void GetReferencedSpriteAssets(TArray<UObject*>& InOutObjects) const;
+
 protected:
 	/**
 	 * Transfers ownership of instance render data to a render thread
diff --git a/Engine/Plugins/2D/Paper2D/Source/Paper2D/Private/PaperGroupedSpriteActor.cpp b/Engine/Plugins/2D/Paper2D/Source/Paper2D/Private/PaperGroupedSpriteActor.cpp
new file mode 100644
--- /dev/null
+++ b/Engine/Plugins/2D/Paper2D/Source/Paper2D/Private/PaperGroupedSpriteActor.cpp
@@ -0,0 +1,32 @@
+// Copyright 1998-2015 Epic Games, Inc. All Rights Reserved.
+
+#include "Paper2DPrivatePCH.h"
+#include "PaperGroupedSpriteActor.h"
+
+#define LOCTEXT_NAMESPACE "Paper2D"
+
+//////////////////////////////////////////////////////////////////////////
+// APaperGroupedSpriteActor
+
+APaperGroupedSpriteActor::APaperGroupedSpriteActor(const FObjectInitializer& ObjectInitializer)
+	: Super(ObjectInitializer)
+{
+	RenderComponent = CreateDefaultSubobject<UPaperGroupedSpriteComponent>(TEXT("RenderComponent"));
+
+	RootComponent = RenderComponent;
+}
+
+#if WITH_EDITOR
+bool APaperGroupedSpriteActor::GetReferencedContentObjects(TArray<UObject*>& Objects) const
+{
+	Super::GetReferencedContentObjects(Objects);
+
+	RenderComponent->GetReferencedSpriteAssets(Objects);
+
+	return true;
+}
+#endif
+
+//////////////////////////////////////////////////////////////////////////
+
+#undef LOCTEXT_NAMESPACE
diff --git a/Engine/Plugins/2D/Paper2D/Source/Paper2D/Private/PaperGroupedSpriteComponent.cpp b/Engine/Plugins/2D/Paper2D/Source/Paper2D/Private/PaperGroupedSpriteComponent.cpp
--- a/Engine/Plugins/2D/Paper2D/Source/Paper2D/Private/PaperGroupedSpriteComponent.cpp
+++ b/Engine/Plugins/2D/Paper2D/Source/Paper2D/Private/PaperGroupedSpriteComponent.cpp
@@ -188,6 +188,19 @@ void UPaperGroupedSpriteComponent::DestroyPhysicsState()
 	ClearAllInstanceBodies();
 }
 
+const UObject* UPaperGroupedSpriteComponent::AdditionalStatObject() const
+{
+	for (const FSpriteInstanceData& InstanceData : PerInstanceSpriteData)
+	{
+		if (InstanceData.SourceSprite != nullptr)
+		{
+			return InstanceData.SourceSprite;
+		}
+	}
+
+	return nullptr;
+}
+
 #if WITH_EDITOR
 void UPaperGroupedSpriteComponent::CheckForErrors()
 {
@@ -484,6 +497,17 @@ bool UPaperGroupedSpriteComponent::ContainsSprite(UPaperSprite* SpriteAsset) con
 	return false;
 }
 
+void UPaperGroupedSpriteComponent::GetReferencedSpriteAssets(TArray<UObject*>& InOutObjects) const
+{
+	for (const FSpriteInstanceData& InstanceData : PerInstanceSpriteData)
+	{
+		if (InstanceData.SourceSprite != nullptr)
+		{
+			InOutObjects.AddUnique(InstanceData.SourceSprite);
+		}
+	}
+}
+
 //////////////////////////////////////////////////////////////////////////
 
 #undef LOCTEXT_NAMESPACE
diff --git a/Engine/Plugins/2D/Paper2D/Source/Paper2DEditor/Private/SpriteEditor/SpriteComponentDetailsCustomization.cpp b/Engine/Plugins/2D/Paper2D/Source/Paper2DEditor/Private/SpriteEditor/SpriteComponentDetailsCustomization.cpp
--- a/Engine/Plugins/2D/Paper2D/Source/Paper2DEditor/Private/SpriteEditor/SpriteComponentDetailsCustomization.cpp
+++ b/Engine/Plugins/2D/Paper2D/Source/Paper2DEditor/Private/SpriteEditor/SpriteComponentDetailsCustomization.cpp
@@ -2,7 +2,7 @@
 
 #include "Paper2DEditorPrivatePCH.h"
 #include "PaperSpriteComponent.h"
-#include "PaperGroupedSpriteComponent.h"
+#include "PaperGroupedSpriteActor.h"
 #include "SpriteComponentDetailsCustomization.h"
 #include "ScopedTransaction.h"
 #include "ILayers.h"
@@ -82,30 +82,29 @@ FReply FSpriteComponentDetailsCustomization::MergeSprites()
 
 			FActorSpawnParameters SpawnParams;
 			SpawnParams.bDeferConstruction = true;
-			if (AActor* SpawnedActor = World->SpawnActor<AActor>(SpawnParams))
+			if (APaperGroupedSpriteActor* SpawnedActor = World->SpawnActor<APaperGroupedSpriteActor>(SpawnParams))
 			{
 				const FScopedTransaction Transaction(LOCTEXT("MergeSprites", "Merge sprite instances"));
 
 				// Create a merged sprite component
-				UPaperGroupedSpriteComponent* MergedSpriteComponent = NewObject<UPaperGroupedSpriteComponent>(SpawnedActor);
-
-				// Create an instance from each sprite component that we're harvesting
-				for (UActorComponent* SourceComponent : ComponentsToHarvest)
 				{
-					UPaperSpriteComponent* SourceSpriteComponent = CastChecked<UPaperSpriteComponent>(SourceComponent);
+					UPaperGroupedSpriteComponent* MergedSpriteComponent = SpawnedActor->GetRenderComponent();
+					FComponentReregisterContext ReregisterContext(MergedSpriteComponent);
+
+					// Create an instance from each sprite component that we're harvesting
+					for (UActorComponent* SourceComponent : ComponentsToHarvest)
+					{
+						UPaperSpriteComponent* SourceSpriteComponent = CastChecked<UPaperSpriteComponent>(SourceComponent);
 
-					UPaperSprite* Sprite = SourceSpriteComponent->GetSprite();
-					const FLinearColor SpriteColor = SourceSpriteComponent->GetSpriteColor();
-					const FTransform RelativeSpriteTransform = SourceSpriteComponent->GetComponentTransform().GetRelativeTransform(MergedWorldTM);
+						UPaperSprite* Sprite = SourceSpriteComponent->GetSprite();
+						const FLinearColor SpriteColor = SourceSpriteComponent->GetSpriteColor();
+						const FTransform RelativeSpriteTransform = SourceSpriteComponent->GetComponentTransform().GetRelativeTransform(MergedWorldTM);
 
-					MergedSpriteComponent->AddInstance(RelativeSpriteTransform, Sprite, FColor(SpriteColor));
+						MergedSpriteComponent->AddInstance(RelativeSpriteTransform, Sprite, /*bWorldSpace=*/ false, SpriteColor);
+					}
 				}
 
 				// Finalize the new component
-				SpawnedActor->AddInstanceComponent(MergedSpriteComponent);
-				SpawnedActor->SetRootComponent(MergedSpriteComponent);
-				MergedSpriteComponent->OnComponentCreated();
-				MergedSpriteComponent->RegisterComponent();
 				UGameplayStatics::FinishSpawningActor(SpawnedActor, MergedWorldTM);
 
 				// Delete the existing actor instances
