Added support for visualizing collision of grouped sprite components

diff --git a/Engine/Plugins/2D/Paper2D/Source/Paper2D/Private/GroupedSpriteSceneProxy.cpp b/Engine/Plugins/2D/Paper2D/Source/Paper2D/Private/GroupedSpriteSceneProxy.cpp
--- a/Engine/Plugins/2D/Paper2D/Source/Paper2D/Private/GroupedSpriteSceneProxy.cpp
+++ b/Engine/Plugins/2D/Paper2D/Source/Paper2D/Private/GroupedSpriteSceneProxy.cpp
@@ -44,8 +44,14 @@ FGroupedSpriteSceneProxy::FGroupedSpriteSceneProxy(UPaperGroupedSpriteComponent*
 
 	NumInstances = InComponent->PerInstanceSpriteData.Num();
 
-	for (const FSpriteInstanceData InstanceData : InComponent->PerInstanceSpriteData)
+	BodySetupTransforms.Reserve(NumInstances);
+	BodySetups.Reserve(NumInstances);
+
+	for (int32 InstanceIndex = 0; InstanceIndex < NumInstances; ++InstanceIndex)
 	{
+		const FSpriteInstanceData InstanceData = InComponent->PerInstanceSpriteData[InstanceIndex];
+
+		TWeakObjectPtr<class UBodySetup> BodySetup;
 		if (InstanceData.SourceSprite != nullptr)
 		{
 			FSpriteDrawCallRecord Record;
@@ -68,7 +74,18 @@ FGroupedSpriteSceneProxy::FGroupedSpriteSceneProxy(UPaperGroupedSpriteComponent*
 
 				new (VertexBuffer.Vertices) FPaperSpriteVertex(ComponentSpacePos, UV, VertColor);
 			}
+
+			if (InComponent->InstanceBodies.IsValidIndex(InstanceIndex))
+			{
+				if (FBodyInstance* BodyInstance = InComponent->InstanceBodies[InstanceIndex])
+				{
+					BodySetup = BodyInstance->BodySetup;
+				}
+			}
 		}
+
+		BodySetupTransforms.Add(InstanceData.Transform);
+		BodySetups.Add(BodySetup);
 	}
 	
 	if (VertexBuffer.Vertices.Num() > 0)
@@ -81,3 +98,27 @@ FGroupedSpriteSceneProxy::FGroupedSpriteSceneProxy(UPaperGroupedSpriteComponent*
 		BeginInitResource(&MyVertexFactory);
 	}
 }
+
+void FGroupedSpriteSceneProxy::DebugDrawCollision(const FSceneView* View, int32 ViewIndex, FMeshElementCollector& Collector, bool bDrawSolid) const
+{
+	for (int32 InstanceIndex = 0; InstanceIndex < NumInstances; ++InstanceIndex)
+	{
+		if (UBodySetup* BodySetup = BodySetups[InstanceIndex].Get())
+		{
+			const FColor CollisionColor = FColor(157, 149, 223, 255);
+			const FMatrix GeomTransform = BodySetupTransforms[InstanceIndex] * GetLocalToWorld();
+			DebugDrawBodySetup(View, ViewIndex, Collector, BodySetup, GeomTransform, CollisionColor, bDrawSolid);
+		}
+	}
+}
+
+void FGroupedSpriteSceneProxy::SetOneBodySetup_RenderThread(int32 InstanceIndex, UBodySetup* NewSetup)
+{
+	BodySetups[InstanceIndex] = NewSetup;
+}
+
+void FGroupedSpriteSceneProxy::SetAllBodySetups_RenderThread(TArray<TWeakObjectPtr<UBodySetup>> Setups)
+{
+	check(Setups.Num() == NumInstances);
+	BodySetups = Setups;
+}
diff --git a/Engine/Plugins/2D/Paper2D/Source/Paper2D/Private/GroupedSpriteSceneProxy.h b/Engine/Plugins/2D/Paper2D/Source/Paper2D/Private/GroupedSpriteSceneProxy.h
--- a/Engine/Plugins/2D/Paper2D/Source/Paper2D/Private/GroupedSpriteSceneProxy.h
+++ b/Engine/Plugins/2D/Paper2D/Source/Paper2D/Private/GroupedSpriteSceneProxy.h
@@ -14,9 +14,13 @@ class FGroupedSpriteSceneProxy : public FPaperRenderSceneProxy
 public:
 	FGroupedSpriteSceneProxy(UPaperGroupedSpriteComponent* InComponent);
 
-	// FPrimitiveSceneProxy interface.
-	//virtual void GetDynamicMeshElements(const TArray<const FSceneView*>& Views, const FSceneViewFamily& ViewFamily, uint32 VisibilityMap, FMeshElementCollector& Collector) const override;
-	// End of FPrimitiveSceneProxy interface.
+	void SetOneBodySetup_RenderThread(int32 InstanceIndex, UBodySetup* NewSetup);
+	void SetAllBodySetups_RenderThread(TArray<TWeakObjectPtr<UBodySetup>> Setups);
+
+protected:
+	// FPaperRenderSceneProxy interface
+	virtual void DebugDrawCollision(const FSceneView* View, int32 ViewIndex, FMeshElementCollector& Collector, bool bDrawSolid) const override;
+	// End of FPaperRenderSceneProxy interface
 
 private:
 	const UPaperGroupedSpriteComponent* MyComponent;
@@ -27,6 +31,9 @@ private:
 	/** Number of instances */
 	int32 NumInstances;
 
+	TArray<FMatrix> BodySetupTransforms;
+	TArray<TWeakObjectPtr<UBodySetup>> BodySetups;
 
+private:
 	FSpriteRenderSection& FindOrAddSection(FSpriteDrawCallRecord& InBatch, UMaterialInterface* InMaterial);
 };
diff --git a/Engine/Plugins/2D/Paper2D/Source/Paper2D/Private/PaperGroupedSpriteComponent.cpp b/Engine/Plugins/2D/Paper2D/Source/Paper2D/Private/PaperGroupedSpriteComponent.cpp
--- a/Engine/Plugins/2D/Paper2D/Source/Paper2D/Private/PaperGroupedSpriteComponent.cpp
+++ b/Engine/Plugins/2D/Paper2D/Source/Paper2D/Private/PaperGroupedSpriteComponent.cpp
@@ -328,10 +328,26 @@ void UPaperGroupedSpriteComponent::CreateAllInstanceBodies()
 	TArray<FTransform> Transforms;
 	Transforms.Reserve(NumBodies);
 
+	TArray<TWeakObjectPtr<UBodySetup>> BodySetups;
+	BodySetups.Reserve(NumBodies);
+
 	for (int32 InstanceIndex = 0; InstanceIndex < NumBodies; ++InstanceIndex)
 	{
 		const FSpriteInstanceData& InstanceData = PerInstanceSpriteData[InstanceIndex];
-		InstanceBodies[InstanceIndex] = InitInstanceBody(InstanceIndex, InstanceData, PhysScene);
+		FBodyInstance* InstanceBody = InitInstanceBody(InstanceIndex, InstanceData, PhysScene);
+		InstanceBodies[InstanceIndex] = InstanceBody;
+		BodySetups.Add(InstanceBody->BodySetup);
+	}
+
+	if (SceneProxy != nullptr)
+	{
+		ENQUEUE_UNIQUE_RENDER_COMMAND_TWOPARAMETER(
+			FSendPaperGroupBodySetups,
+			FGroupedSpriteSceneProxy*, InSceneProxy, (FGroupedSpriteSceneProxy*)SceneProxy,
+			TArray<TWeakObjectPtr<UBodySetup>>, InBodySetups, BodySetups,
+			{
+			InSceneProxy->SetAllBodySetups_RenderThread(InBodySetups);
+		});
 	}
 }
 
diff --git a/Engine/Plugins/2D/Paper2D/Source/Paper2D/Private/PaperRenderSceneProxy.cpp b/Engine/Plugins/2D/Paper2D/Source/Paper2D/Private/PaperRenderSceneProxy.cpp
--- a/Engine/Plugins/2D/Paper2D/Source/Paper2D/Private/PaperRenderSceneProxy.cpp
+++ b/Engine/Plugins/2D/Paper2D/Source/Paper2D/Private/PaperRenderSceneProxy.cpp
@@ -264,7 +264,7 @@ FPaperRenderSceneProxy::FPaperRenderSceneProxy(const UPrimitiveComponent* InComp
 	: FPrimitiveSceneProxy(InComponent)
 	, Material(nullptr)
 	, Owner(InComponent->GetOwner())
-	, BodySetup(const_cast<UPrimitiveComponent*>(InComponent)->GetBodySetup())
+	, MyBodySetup(const_cast<UPrimitiveComponent*>(InComponent)->GetBodySetup())
 	, bCastShadow(InComponent->CastShadow)
 	, CollisionResponse(InComponent->GetCollisionResponseToChannels())
 {
@@ -294,6 +294,47 @@ void FPaperRenderSceneProxy::CreateRenderThreadResources()
 #endif
 }
 
+void FPaperRenderSceneProxy::DebugDrawBodySetup(const FSceneView* View, int32 ViewIndex, FMeshElementCollector& Collector, UBodySetup* BodySetup, const FMatrix& GeomTransformMatrix, const FLinearColor& CollisionColor, bool bDrawSolid) const
+{
+	if (FMath::Abs(GeomTransformMatrix.Determinant()) < SMALL_NUMBER)
+	{
+		// Catch this here or otherwise GeomTransform below will assert
+		// This spams so commented out
+		//UE_LOG(LogStaticMesh, Log, TEXT("Zero scaling not supported (%s)"), *StaticMesh->GetPathName());
+	}
+	else
+	{
+		FTransform GeomTransform(GeomTransformMatrix);
+
+		if (bDrawSolid)
+		{
+			// Make a material for drawing solid collision stuff
+			auto SolidMaterialInstance = new FColoredMaterialRenderProxy(
+				GEngine->ShadedLevelColorationUnlitMaterial->GetRenderProxy(IsSelected(), IsHovered()),
+				WireframeColor
+				);
+
+			Collector.RegisterOneFrameMaterialProxy(SolidMaterialInstance);
+
+			BodySetup->AggGeom.GetAggGeom(GeomTransform, WireframeColor, SolidMaterialInstance, false, true, UseEditorDepthTest(), ViewIndex, Collector);
+		}
+		else
+		{
+			// wireframe
+			BodySetup->AggGeom.GetAggGeom(GeomTransform, GetSelectionColor(CollisionColor, IsSelected(), IsHovered()), nullptr, (Owner == nullptr), false, UseEditorDepthTest(), ViewIndex, Collector);
+		}
+	}
+}
+
+void FPaperRenderSceneProxy::DebugDrawCollision(const FSceneView* View, int32 ViewIndex, FMeshElementCollector& Collector, bool bDrawSolid) const
+{
+	if (MyBodySetup != nullptr)
+	{
+		const FColor CollisionColor = FColor(157, 149, 223, 255);
+		DebugDrawBodySetup(View, ViewIndex, Collector, MyBodySetup, GetLocalToWorld(), CollisionColor, bDrawSolid);
+	}
+}
+
 void FPaperRenderSceneProxy::GetDynamicMeshElements(const TArray<const FSceneView*>& Views, const FSceneViewFamily& ViewFamily, uint32 VisibilityMap, FMeshElementCollector& Collector) const
 {
 	SCOPE_CYCLE_COUNTER(STAT_PaperRenderSceneProxy_GetDynamicMeshElements);
@@ -336,40 +377,9 @@ void FPaperRenderSceneProxy::GetDynamicMeshElements(const TArray<const FSceneVie
 		{
 			if ((bDrawSimpleCollision || bDrawWireframeCollision) && AllowDebugViewmodes())
 			{
-				if (BodySetup)
-				{
-					if (FMath::Abs(GetLocalToWorld().Determinant()) < SMALL_NUMBER)
-					{
-						// Catch this here or otherwise GeomTransform below will assert
-						// This spams so commented out
-						//UE_LOG(LogStaticMesh, Log, TEXT("Zero scaling not supported (%s)"), *StaticMesh->GetPathName());
-					}
-					else
-					{
-						const bool bDrawSolid = !bDrawWireframeCollision;
-
-						if (bDrawSolid)
-						{
-							// Make a material for drawing solid collision stuff
-							auto SolidMaterialInstance = new FColoredMaterialRenderProxy(
-								GEngine->ShadedLevelColorationUnlitMaterial->GetRenderProxy(IsSelected(), IsHovered()),
-								WireframeColor
-								);
-
-							Collector.RegisterOneFrameMaterialProxy(SolidMaterialInstance);
-
-							FTransform GeomTransform(GetLocalToWorld());
-							BodySetup->AggGeom.GetAggGeom(GeomTransform, WireframeColor, SolidMaterialInstance, false, true, UseEditorDepthTest(), ViewIndex, Collector);
-						}
-						else
-						{
-							// wireframe
-							FColor CollisionColor = FColor(157, 149, 223, 255);
-							FTransform GeomTransform(GetLocalToWorld());
-							BodySetup->AggGeom.GetAggGeom(GeomTransform, GetSelectionColor(CollisionColor, IsSelected(), IsHovered()), nullptr, (Owner == nullptr), false, UseEditorDepthTest(), ViewIndex, Collector);
-						}
-					}
-				}
+				const FSceneView* View = Views[ViewIndex];
+				const bool bDrawSolid = !bDrawWireframeCollision;
+				DebugDrawCollision(View, ViewIndex, Collector, bDrawSolid);
 			}
 
 			// Draw bounds
@@ -630,7 +640,7 @@ void FPaperRenderSceneProxy::SetDrawCall_RenderThread(const FSpriteDrawCallRecor
 
 void FPaperRenderSceneProxy::SetBodySetup_RenderThread(UBodySetup* NewSetup)
 {
-	BodySetup = NewSetup;
+	MyBodySetup = NewSetup;
 }
 
 bool FPaperRenderSceneProxy::IsCollisionView(const FEngineShowFlags& EngineShowFlags, bool& bDrawSimpleCollision, bool& bDrawComplexCollision) const
diff --git a/Engine/Plugins/2D/Paper2D/Source/Paper2D/Private/PaperRenderSceneProxy.h b/Engine/Plugins/2D/Paper2D/Source/Paper2D/Private/PaperRenderSceneProxy.h
--- a/Engine/Plugins/2D/Paper2D/Source/Paper2D/Private/PaperRenderSceneProxy.h
+++ b/Engine/Plugins/2D/Paper2D/Source/Paper2D/Private/PaperRenderSceneProxy.h
@@ -160,11 +160,6 @@ public:
 	void SetTransientTextureOverride_RenderThread(const UTexture* InTextureToModifyOverrideFor, UTexture* InOverrideTexture);
 #endif
 
-// 	class UMaterialInterface* GetMaterial() const
-// 	{
-// 		return Material;
-// 	}
-
 protected:
 	virtual void GetDynamicMeshElementsForView(const FSceneView* View, int32 ViewIndex, FMeshElementCollector& Collector) const;
 
@@ -179,6 +174,8 @@ protected:
 
 	void ConvertBatchesToNewStyle(TArray<FSpriteDrawCallRecord>& SourceBatches);
 
+	virtual void DebugDrawCollision(const FSceneView* View, int32 ViewIndex, FMeshElementCollector& Collector, bool bDrawSolid) const;
+	virtual void DebugDrawBodySetup(const FSceneView* View, int32 ViewIndex, FMeshElementCollector& Collector, UBodySetup* BodySetup, const FMatrix& GeomTransform, const FLinearColor& CollisionColor, bool bDrawSolid) const;
 protected:
 	// New style
 	FPaperSpriteVertexBuffer VertexBuffer;
@@ -191,7 +188,7 @@ protected:
 
 	//
 	AActor* Owner;
-	UBodySetup* BodySetup;
+	UBodySetup* MyBodySetup;
 
 	bool bCastShadow;
 
