Added support for nav-mesh generation on grouped sprites and fixed a crash when a sprite has no collision

diff --git a/Engine/Plugins/2D/Paper2D/Source/Paper2D/Classes/PaperGroupedSpriteComponent.h b/Engine/Plugins/2D/Paper2D/Source/Paper2D/Classes/PaperGroupedSpriteComponent.h
--- a/Engine/Plugins/2D/Paper2D/Source/Paper2D/Classes/PaperGroupedSpriteComponent.h
+++ b/Engine/Plugins/2D/Paper2D/Source/Paper2D/Classes/PaperGroupedSpriteComponent.h
@@ -113,6 +113,8 @@ public:
 	virtual void GetUsedTextures(TArray<UTexture*>& OutTextures, EMaterialQualityLevel::Type QualityLevel) override;
 	virtual UMaterialInterface* GetMaterial(int32 MaterialIndex) const override;
 	virtual int32 GetNumMaterials() const override;
+	virtual bool DoCustomNavigableGeometryExport(FNavigableGeometryExport& GeomExport) const override;
+	virtual void GetNavigationData(FNavigationRelevantData& Data) const override;
 	// End of UPrimitiveComponent interface
 
 	// UObject interface
@@ -129,6 +131,9 @@ public:
 	// Adds all referenced sprite assets to the specified list
 	void GetReferencedSpriteAssets(TArray<UObject*>& InOutObjects) const;
 
+	/** Handles request from navigation system to gather instance transforms in a specific area box */
+	void GetNavigationPerInstanceTransforms(const FBox& AreaBox, TArray<FTransform>& OutInstanceTransforms) const;
+
 protected:
 	/**
 	 * Transfers ownership of instance render data to a render thread
diff --git a/Engine/Plugins/2D/Paper2D/Source/Paper2D/Private/PaperGroupedSpriteComponent.cpp b/Engine/Plugins/2D/Paper2D/Source/Paper2D/Private/PaperGroupedSpriteComponent.cpp
--- a/Engine/Plugins/2D/Paper2D/Source/Paper2D/Private/PaperGroupedSpriteComponent.cpp
+++ b/Engine/Plugins/2D/Paper2D/Source/Paper2D/Private/PaperGroupedSpriteComponent.cpp
@@ -6,6 +6,8 @@
 
 #include "PaperGroupedSpriteComponent.h"
 #include "GroupedSpriteSceneProxy.h"
+#include "NavigationSystemHelpers.h"
+#include "NavigationOctree.h"
 
 #include "PaperSprite.h"
 #include "PhysicsPublic.h"
@@ -25,6 +27,7 @@ UPaperGroupedSpriteComponent::UPaperGroupedSpriteComponent(const FObjectInitiali
 
 	Mobility = EComponentMobility::Movable;
 	BodyInstance.bSimulatePhysics = false;
+	bHasCustomNavigableGeometry = EHasCustomNavigableGeometry::Yes;
 }
 
 int32 UPaperGroupedSpriteComponent::AddInstance(const FTransform& Transform, UPaperSprite* Sprite, bool bWorldSpace, FLinearColor Color)
@@ -336,7 +339,7 @@ void UPaperGroupedSpriteComponent::CreateAllInstanceBodies()
 		const FSpriteInstanceData& InstanceData = PerInstanceSpriteData[InstanceIndex];
 		FBodyInstance* InstanceBody = InitInstanceBody(InstanceIndex, InstanceData, PhysScene);
 		InstanceBodies[InstanceIndex] = InstanceBody;
-		BodySetups.Add(InstanceBody->BodySetup);
+		BodySetups.Add((InstanceBody != nullptr) ? InstanceBody->BodySetup : TWeakObjectPtr<UBodySetup>());
 	}
 
 	if (SceneProxy != nullptr)
@@ -500,6 +503,45 @@ int32 UPaperGroupedSpriteComponent::GetNumMaterials() const
 	return FMath::Max<int32>(OverrideMaterials.Num(), FMath::Max<int32>(InstanceMaterials.Num(), 1));
 }
 
+bool UPaperGroupedSpriteComponent::DoCustomNavigableGeometryExport(FNavigableGeometryExport& GeomExport) const
+{
+	for (FBodyInstance* InstanceBody : InstanceBodies)
+	{
+		if (InstanceBody != nullptr)
+		{
+			if (UBodySetup* BodySetup = InstanceBody->BodySetup.Get())
+			{
+				GeomExport.ExportRigidBodySetup(*BodySetup, FTransform::Identity);
+			}
+		}
+
+		// Hook per instance transform delegate
+		GeomExport.SetNavDataPerInstanceTransformDelegate(FNavDataPerInstanceTransformDelegate::CreateUObject(this, &UPaperGroupedSpriteComponent::GetNavigationPerInstanceTransforms));
+	}
+
+	// we don't want "regular" collision export for this component
+	return false;
+}
+
+void UPaperGroupedSpriteComponent::GetNavigationData(FNavigationRelevantData& Data) const
+{
+	// Hook per instance transform delegate
+	Data.NavDataPerInstanceTransformDelegate = FNavDataPerInstanceTransformDelegate::CreateUObject(this, &UPaperGroupedSpriteComponent::GetNavigationPerInstanceTransforms);
+}
+
+void UPaperGroupedSpriteComponent::GetNavigationPerInstanceTransforms(const FBox& AreaBox, TArray<FTransform>& OutInstanceTransforms) const
+{
+	for (const FSpriteInstanceData& InstanceData : PerInstanceSpriteData)
+	{
+		//TODO: Is it worth doing per instance bounds check here ?
+		const FTransform InstanceToComponent(InstanceData.Transform);
+		if (!InstanceToComponent.GetScale3D().IsZero())
+		{
+			OutInstanceTransforms.Add(InstanceToComponent * ComponentToWorld);
+		}
+	}
+}
+
 bool UPaperGroupedSpriteComponent::ContainsSprite(UPaperSprite* SpriteAsset) const
 {
 	for (const FSpriteInstanceData& InstanceData : PerInstanceSpriteData)
