Added support for overriding materials on grouped sprite components

diff --git a/Engine/Plugins/2D/Paper2D/Source/Paper2D/Private/GroupedSpriteSceneProxy.cpp b/Engine/Plugins/2D/Paper2D/Source/Paper2D/Private/GroupedSpriteSceneProxy.cpp
--- a/Engine/Plugins/2D/Paper2D/Source/Paper2D/Private/GroupedSpriteSceneProxy.cpp
+++ b/Engine/Plugins/2D/Paper2D/Source/Paper2D/Private/GroupedSpriteSceneProxy.cpp
@@ -57,7 +57,7 @@ FGroupedSpriteSceneProxy::FGroupedSpriteSceneProxy(UPaperGroupedSpriteComponent*
 			FSpriteDrawCallRecord Record;
 			Record.BuildFromSprite(InstanceData.SourceSprite);
 
-			UMaterialInterface* SpriteMaterial = InstanceData.SourceSprite->GetMaterial(0); //@TODO: Should query from the component due to overrides, but that requires an index per item, etc...
+			UMaterialInterface* SpriteMaterial = InComponent->GetMaterial(InstanceData.MaterialIndex);
 
 			FSpriteRenderSection& Section = FindOrAddSection(Record, SpriteMaterial);
 			
diff --git a/Engine/Plugins/2D/Paper2D/Source/Paper2D/Private/PaperGroupedSpriteComponent.cpp b/Engine/Plugins/2D/Paper2D/Source/Paper2D/Private/PaperGroupedSpriteComponent.cpp
--- a/Engine/Plugins/2D/Paper2D/Source/Paper2D/Private/PaperGroupedSpriteComponent.cpp
+++ b/Engine/Plugins/2D/Paper2D/Source/Paper2D/Private/PaperGroupedSpriteComponent.cpp
@@ -427,13 +427,22 @@ void UPaperGroupedSpriteComponent::RebuildInstances()
 
 void UPaperGroupedSpriteComponent::RebuildMaterialList()
 {
-	InstanceMaterials.Reset();
+	TArray<UMaterialInterface*> OldOverrides(OverrideMaterials);
+	OverrideMaterials.Empty();
 
-	//@TODO: Need to rebuild the OverrideMaterials array and update the per-sprite material indices as part of this process
-
-	for (const FSpriteInstanceData& InstanceData : PerInstanceSpriteData)
+	for (FSpriteInstanceData& InstanceData : PerInstanceSpriteData)
 	{
-		UpdateMaterialList(InstanceData.SourceSprite);
+		const int32 OldMaterialIndex = InstanceData.MaterialIndex;
+		const int32 NewMaterialIndex = UpdateMaterialList(InstanceData.SourceSprite);
+		InstanceData.MaterialIndex = NewMaterialIndex;
+
+		// Remap the overrides
+		UMaterialInterface* OldOverride = OldOverrides.IsValidIndex(OldMaterialIndex) ? OldOverrides[OldMaterialIndex] : nullptr;
+		
+		if ((OldOverride != nullptr) && (NewMaterialIndex >= 0))
+		{
+			SetMaterial(NewMaterialIndex, OldOverride);
+		}
 	}
 }
 
