Paper2D: Prevented material overrides from being discarded when splitting or merging a grouped sprite

diff --git a/Engine/Plugins/2D/Paper2D/Source/Paper2D/Classes/PaperGroupedSpriteComponent.h b/Engine/Plugins/2D/Paper2D/Source/Paper2D/Classes/PaperGroupedSpriteComponent.h
--- a/Engine/Plugins/2D/Paper2D/Source/Paper2D/Classes/PaperGroupedSpriteComponent.h
+++ b/Engine/Plugins/2D/Paper2D/Source/Paper2D/Classes/PaperGroupedSpriteComponent.h
@@ -70,7 +70,9 @@ public:
 
 	/** Add an instance to this component. Transform can be given either in the local space of this component or world space.  */
 	UFUNCTION(BlueprintCallable, Category="Components|Sprite")
-	virtual int32 AddInstance(const FTransform& LocalTransform, UPaperSprite* Sprite, bool bWorldSpace = false, FLinearColor Color = FLinearColor::White);
+	int32 AddInstance(const FTransform& Transform, UPaperSprite* Sprite, bool bWorldSpace = false, FLinearColor Color = FLinearColor::White);
+
+	virtual int32 AddInstanceWithMaterial(const FTransform& Transform, UPaperSprite* Sprite, UMaterialInterface* MaterialOverride = nullptr, bool bWorldSpace = false, FLinearColor Color = FLinearColor::White);
 
 	/** Get the transform for the instance specified. Instance is returned in local space of this component unless bWorldSpace is set.  Returns True on success. */
 	UFUNCTION(BlueprintCallable, Category = "Components|Sprite")
@@ -148,7 +150,7 @@ protected:
 	void ClearAllInstanceBodies();
 
 	/** Sets up new instance data to sensible defaults, creates physics counterparts if possible */
-	void SetupNewInstanceData(FSpriteInstanceData& InOutNewInstanceData, int32 InInstanceIndex, const FTransform& InInstanceTransform, UPaperSprite* InSprite, const FColor& InColor);
+	void SetupNewInstanceData(FSpriteInstanceData& InOutNewInstanceData, int32 InInstanceIndex, const FTransform& InInstanceTransform, UPaperSprite* InSprite, UMaterialInterface* MaterialOverride, const FColor& InColor);
 
 	/** Creates a body instance for the specified instance data if that sprite has defined collision */
 	FBodyInstance* InitInstanceBody(int32 InstanceIndex, const FSpriteInstanceData& InstanceData, FPhysScene* PhysScene);
@@ -160,7 +162,7 @@ protected:
 	void RebuildMaterialList();
 
 	/** Adds to the material list from a single sprite */
-	int32 UpdateMaterialList(UPaperSprite* Sprite);
+	int32 UpdateMaterialList(UPaperSprite* Sprite, UMaterialInterface* MaterialOverride);
 
 	friend FGroupedSpriteSceneProxy;
 	friend FGroupedSpriteComponentDetailsCustomization;
diff --git a/Engine/Plugins/2D/Paper2D/Source/Paper2D/Private/PaperGroupedSpriteComponent.cpp b/Engine/Plugins/2D/Paper2D/Source/Paper2D/Private/PaperGroupedSpriteComponent.cpp
--- a/Engine/Plugins/2D/Paper2D/Source/Paper2D/Private/PaperGroupedSpriteComponent.cpp
+++ b/Engine/Plugins/2D/Paper2D/Source/Paper2D/Private/PaperGroupedSpriteComponent.cpp
@@ -31,13 +31,18 @@ UPaperGroupedSpriteComponent::UPaperGroupedSpriteComponent(const FObjectInitiali
 }
 
 int32 UPaperGroupedSpriteComponent::AddInstance(const FTransform& Transform, UPaperSprite* Sprite, bool bWorldSpace, FLinearColor Color)
+{
+	return AddInstanceWithMaterial(Transform, Sprite, nullptr, bWorldSpace, Color);
+}
+
+int32 UPaperGroupedSpriteComponent::AddInstanceWithMaterial(const FTransform& Transform, UPaperSprite* Sprite, UMaterialInterface* MaterialOverride, bool bWorldSpace, FLinearColor Color)
 {
 	const int32 NewInstanceIndex = PerInstanceSpriteData.Num();
 
 	const FTransform LocalTransform(bWorldSpace ? Transform.GetRelativeTransform(ComponentToWorld) : Transform);
 
-	FSpriteInstanceData& NewInstanceData = *new (PerInstanceSpriteData) FSpriteInstanceData();
-	SetupNewInstanceData(NewInstanceData, NewInstanceIndex, LocalTransform, Sprite, Color);
+	FSpriteInstanceData& NewInstanceData = *new (PerInstanceSpriteData)FSpriteInstanceData();
+	SetupNewInstanceData(NewInstanceData, NewInstanceIndex, LocalTransform, Sprite, MaterialOverride, Color);
 
 	MarkRenderStateDirty();
 
@@ -301,7 +306,7 @@ void UPaperGroupedSpriteComponent::PostEditChangeChainProperty(FPropertyChangedC
 		{
 			const int32 AddedAtIndex = PropertyChangedEvent.GetArrayIndex(PropertyChangedEvent.Property->GetFName().ToString());
 			check(AddedAtIndex != INDEX_NONE);
-			SetupNewInstanceData(PerInstanceSpriteData[AddedAtIndex], AddedAtIndex, FTransform::Identity, nullptr, FColor::White); //@TODO: Need to pull the sprite from somewhere
+			SetupNewInstanceData(PerInstanceSpriteData[AddedAtIndex], AddedAtIndex, FTransform::Identity, nullptr, nullptr, FColor::White); //@TODO: Need to pull the sprite from somewhere
 		}
 
 		MarkRenderStateDirty();
@@ -370,12 +375,12 @@ void UPaperGroupedSpriteComponent::ClearAllInstanceBodies()
 	InstanceBodies.Empty();
 }
 
-void UPaperGroupedSpriteComponent::SetupNewInstanceData(FSpriteInstanceData& InOutNewInstanceData, int32 InInstanceIndex, const FTransform& InInstanceTransform, UPaperSprite* InSprite, const FColor& InColor)
+void UPaperGroupedSpriteComponent::SetupNewInstanceData(FSpriteInstanceData& InOutNewInstanceData, int32 InInstanceIndex, const FTransform& InInstanceTransform, UPaperSprite* InSprite, UMaterialInterface* MaterialOverride, const FColor& InColor)
 {
 	InOutNewInstanceData.Transform = InInstanceTransform.ToMatrixWithScale();
 	InOutNewInstanceData.SourceSprite = InSprite;
 	InOutNewInstanceData.VertexColor = InColor;
-	InOutNewInstanceData.MaterialIndex = UpdateMaterialList(InSprite);
+	InOutNewInstanceData.MaterialIndex = UpdateMaterialList(InSprite, MaterialOverride);
 
 	if (bPhysicsStateCreated && (InSprite != nullptr) && (InSprite->BodySetup != nullptr))
 	{
@@ -433,20 +438,14 @@ void UPaperGroupedSpriteComponent::RebuildMaterialList()
 	for (FSpriteInstanceData& InstanceData : PerInstanceSpriteData)
 	{
 		const int32 OldMaterialIndex = InstanceData.MaterialIndex;
-		const int32 NewMaterialIndex = UpdateMaterialList(InstanceData.SourceSprite);
-		InstanceData.MaterialIndex = NewMaterialIndex;
-
-		// Remap the overrides
 		UMaterialInterface* OldOverride = OldOverrides.IsValidIndex(OldMaterialIndex) ? OldOverrides[OldMaterialIndex] : nullptr;
-		
-		if ((OldOverride != nullptr) && (NewMaterialIndex >= 0))
-		{
-			SetMaterial(NewMaterialIndex, OldOverride);
-		}
+
+		const int32 NewMaterialIndex = UpdateMaterialList(InstanceData.SourceSprite, OldOverride);
+		InstanceData.MaterialIndex = NewMaterialIndex;
 	}
 }
 
-int32 UPaperGroupedSpriteComponent::UpdateMaterialList(UPaperSprite* Sprite)
+int32 UPaperGroupedSpriteComponent::UpdateMaterialList(UPaperSprite* Sprite, UMaterialInterface* MaterialOverride)
 {
 	int32 Result = INDEX_NONE;
 
@@ -456,6 +455,11 @@ int32 UPaperGroupedSpriteComponent::UpdateMaterialList(UPaperSprite* Sprite)
 		{
 			Result = InstanceMaterials.AddUnique(SpriteMaterial);
 		}
+
+		if (MaterialOverride != nullptr)
+		{
+			SetMaterial(Result, MaterialOverride);
+		}
 	}
 
 	return Result;
diff --git a/Engine/Plugins/2D/Paper2D/Source/Paper2DEditor/Private/GroupedSpriteDetailsCustomization.cpp b/Engine/Plugins/2D/Paper2D/Source/Paper2DEditor/Private/GroupedSpriteDetailsCustomization.cpp
--- a/Engine/Plugins/2D/Paper2D/Source/Paper2DEditor/Private/GroupedSpriteDetailsCustomization.cpp
+++ b/Engine/Plugins/2D/Paper2D/Source/Paper2DEditor/Private/GroupedSpriteDetailsCustomization.cpp
@@ -137,11 +137,12 @@ FReply FGroupedSpriteComponentDetailsCustomization::SplitSprites()
 								SpawnedSpriteComponent->SetSpriteColor(InstanceData.VertexColor);
 								SpawnedSpriteComponent->SetSprite(InstanceData.SourceSprite);
 
-								//@TODO: Sprite component material!
-								// if (SourceBatchComponent->GetMaterial() != InstanceData.SourceSprite->GetMaterial(0))
-								// {
-								//SpawnedSpriteComponent->SetMaterial();
-								// }
+								// Apply the material override if there is one
+								UMaterialInterface* InstanceMaterial = SourceBatchComponent->GetMaterial(InstanceData.MaterialIndex);
+								if (InstanceMaterial != InstanceData.SourceSprite->GetMaterial(0))
+								{
+									SpawnedSpriteComponent->SetMaterial(0, InstanceMaterial);
+								}
 							}
 
 							UGameplayStatics::FinishSpawningActor(SpawnedActor, InstanceTransform);
diff --git a/Engine/Plugins/2D/Paper2D/Source/Paper2DEditor/Private/SpriteEditor/SpriteComponentDetailsCustomization.cpp b/Engine/Plugins/2D/Paper2D/Source/Paper2DEditor/Private/SpriteEditor/SpriteComponentDetailsCustomization.cpp
--- a/Engine/Plugins/2D/Paper2D/Source/Paper2DEditor/Private/SpriteEditor/SpriteComponentDetailsCustomization.cpp
+++ b/Engine/Plugins/2D/Paper2D/Source/Paper2DEditor/Private/SpriteEditor/SpriteComponentDetailsCustomization.cpp
@@ -100,7 +100,13 @@ FReply FSpriteComponentDetailsCustomization::MergeSprites()
 						const FLinearColor SpriteColor = SourceSpriteComponent->GetSpriteColor();
 						const FTransform RelativeSpriteTransform = SourceSpriteComponent->GetComponentTransform().GetRelativeTransform(MergedWorldTM);
 
-						MergedSpriteComponent->AddInstance(RelativeSpriteTransform, Sprite, /*bWorldSpace=*/ false, SpriteColor);
+						UMaterialInterface* OverrideMaterial = SourceSpriteComponent->GetMaterial(0);
+						if ((Sprite == nullptr) || (OverrideMaterial == Sprite->GetMaterial(0)))
+						{
+							OverrideMaterial = nullptr;
+						}
+
+						MergedSpriteComponent->AddInstanceWithMaterial(RelativeSpriteTransform, Sprite, OverrideMaterial, /*bWorldSpace=*/ false, SpriteColor);
 					}
 				}
 
