Paper2D: Added actor labelling when splitting apart a sprite group so they get a nicer name based on the asset

diff --git a/Engine/Plugins/2D/Paper2D/Source/Paper2DEditor/Private/GroupedSpriteDetailsCustomization.cpp b/Engine/Plugins/2D/Paper2D/Source/Paper2DEditor/Private/GroupedSpriteDetailsCustomization.cpp
--- a/Engine/Plugins/2D/Paper2D/Source/Paper2DEditor/Private/GroupedSpriteDetailsCustomization.cpp
+++ b/Engine/Plugins/2D/Paper2D/Source/Paper2DEditor/Private/GroupedSpriteDetailsCustomization.cpp
@@ -147,6 +147,12 @@ FReply FGroupedSpriteComponentDetailsCustomization::SplitSprites()
 
 							UGameplayStatics::FinishSpawningActor(SpawnedActor, InstanceTransform);
 							ActorsCreated.Add(SpawnedActor);
+
+							// Give it a good name
+							if (InstanceData.SourceSprite != nullptr)
+							{
+								FActorLabelUtilities::SetActorLabelUnique(SpawnedActor, InstanceData.SourceSprite->GetName());
+							}
 						}
 					}
 				}
