Grouped Sprites: Added context menu options to the level editor to merge/split grouped sprites that are available even if the selection is heterogeneous as long as it contains at least one actor with a sprite component or grouped sprite component
- Moved merge and split code into a new FPaperGroupedSpriteUtilities class

diff --git a/Engine/Plugins/2D/Paper2D/Source/Paper2D/Classes/PaperGroupedSpriteComponent.h b/Engine/Plugins/2D/Paper2D/Source/Paper2D/Classes/PaperGroupedSpriteComponent.h
--- a/Engine/Plugins/2D/Paper2D/Source/Paper2D/Classes/PaperGroupedSpriteComponent.h
+++ b/Engine/Plugins/2D/Paper2D/Source/Paper2D/Classes/PaperGroupedSpriteComponent.h
@@ -2,14 +2,12 @@
 
 #pragma once
 
+class FGroupedSpriteSceneProxy;
+
 #include "PaperSprite.h"
 
 #include "PaperGroupedSpriteComponent.generated.h"
 
-class UPaperGroupedSpriteComponent;
-class FGroupedSpriteSceneProxy;
-class FGroupedSpriteComponentDetailsCustomization;
-
 USTRUCT()
 struct FSpriteInstanceData
 {
@@ -136,6 +134,11 @@ public:
 	/** Handles request from navigation system to gather instance transforms in a specific area box */
 	void GetNavigationPerInstanceTransforms(const FBox& AreaBox, TArray<FTransform>& OutInstanceTransforms) const;
 
+	const TArray<FSpriteInstanceData>& GetPerInstanceSpriteData() const
+	{
+		return PerInstanceSpriteData;
+	}
+
 protected:
 	/**
 	 * Transfers ownership of instance render data to a render thread
@@ -165,5 +168,4 @@ protected:
 	int32 UpdateMaterialList(UPaperSprite* Sprite, UMaterialInterface* MaterialOverride);
 
 	friend FGroupedSpriteSceneProxy;
-	friend FGroupedSpriteComponentDetailsCustomization;
 };
diff --git a/Engine/Plugins/2D/Paper2D/Source/Paper2DEditor/Private/GroupedSpriteDetailsCustomization.cpp b/Engine/Plugins/2D/Paper2D/Source/Paper2DEditor/Private/GroupedSpriteDetailsCustomization.cpp
deleted file mode 100644
diff --git a/Engine/Plugins/2D/Paper2D/Source/Paper2DEditor/Private/GroupedSpriteDetailsCustomization.h b/Engine/Plugins/2D/Paper2D/Source/Paper2DEditor/Private/GroupedSpriteDetailsCustomization.h
deleted file mode 100644
diff --git a/Engine/Plugins/2D/Paper2D/Source/Paper2DEditor/Private/GroupedSprites/GroupedSpriteDetailsCustomization.cpp b/Engine/Plugins/2D/Paper2D/Source/Paper2DEditor/Private/GroupedSprites/GroupedSpriteDetailsCustomization.cpp
new file mode 100644
--- /dev/null
+++ b/Engine/Plugins/2D/Paper2D/Source/Paper2DEditor/Private/GroupedSprites/GroupedSpriteDetailsCustomization.cpp
@@ -0,0 +1,66 @@
+// Copyright 1998-2015 Epic Games, Inc. All Rights Reserved.
+
+#include "Paper2DEditorPrivatePCH.h"
+#include "GroupedSpriteDetailsCustomization.h"
+
+#include "PaperSpriteActor.h"
+#include "PaperSpriteComponent.h"
+#include "PaperGroupedSpriteComponent.h"
+#include "PaperGroupedSpriteUtilities.h"
+
+#define LOCTEXT_NAMESPACE "SpriteEditor"
+
+//////////////////////////////////////////////////////////////////////////
+// FGroupedSpriteComponentDetailsCustomization
+
+TSharedRef<IDetailCustomization> FGroupedSpriteComponentDetailsCustomization::MakeInstance()
+{
+	return MakeShareable(new FGroupedSpriteComponentDetailsCustomization());
+}
+
+FGroupedSpriteComponentDetailsCustomization::FGroupedSpriteComponentDetailsCustomization()
+{
+}
+
+void FGroupedSpriteComponentDetailsCustomization::CustomizeDetails(IDetailLayoutBuilder& DetailBuilder)
+{
+	// Create a category so this is displayed early in the properties
+	IDetailCategoryBuilder& SpriteCategory = DetailBuilder.EditCategory("Sprite", FText::GetEmpty(), ECategoryPriority::Important);
+
+	ObjectsBeingCustomized.Empty();
+	DetailBuilder.GetObjectsBeingCustomized(/*out*/ ObjectsBeingCustomized);
+
+	{
+		// Expose split buttons
+		FDetailWidgetRow& SplitRow = SpriteCategory.AddCustomRow(LOCTEXT("SplitSearchText", "Split"))
+			.WholeRowContent()
+			[
+				SNew(SHorizontalBox)
+				+ SHorizontalBox::Slot()
+				.AutoWidth()
+				.Padding(2.0f, 0.0f)
+				.VAlign(VAlign_Center)
+				.HAlign(HAlign_Left)
+				[
+					SNew(SButton)
+					.Text(LOCTEXT("SplitSprites", "Split Sprites"))
+					.ToolTipText(LOCTEXT("SplitSprites_Tooltip", "Splits all sprite instances into separate sprite actors or components"))
+					.OnClicked(this, &FGroupedSpriteComponentDetailsCustomization::SplitSprites)
+				]
+			];
+	}
+}
+
+FReply FGroupedSpriteComponentDetailsCustomization::SplitSprites()
+{
+	TArray<UObject*> StrongObjects;
+	CopyFromWeakArray(StrongObjects, ObjectsBeingCustomized);
+
+	FPaperGroupedSpriteUtilities::SplitSprites(StrongObjects);
+
+	return FReply::Handled();
+}
+
+//////////////////////////////////////////////////////////////////////////
+
+#undef LOCTEXT_NAMESPACE
diff --git a/Engine/Plugins/2D/Paper2D/Source/Paper2DEditor/Private/GroupedSprites/GroupedSpriteDetailsCustomization.h b/Engine/Plugins/2D/Paper2D/Source/Paper2DEditor/Private/GroupedSprites/GroupedSpriteDetailsCustomization.h
new file mode 100644
--- /dev/null
+++ b/Engine/Plugins/2D/Paper2D/Source/Paper2DEditor/Private/GroupedSprites/GroupedSpriteDetailsCustomization.h
@@ -0,0 +1,27 @@
+// Copyright 1998-2015 Epic Games, Inc. All Rights Reserved.
+
+#pragma once
+
+#include "PropertyEditing.h"
+
+//////////////////////////////////////////////////////////////////////////
+// FGroupedSpriteComponentDetailsCustomization
+
+class FGroupedSpriteComponentDetailsCustomization : public IDetailCustomization
+{
+public:
+	/** Makes a new instance of this detail layout class for a specific detail view requesting it */
+	static TSharedRef<IDetailCustomization> MakeInstance();
+
+	// IDetailCustomization interface
+	virtual void CustomizeDetails(IDetailLayoutBuilder& DetailLayout) override;
+	// End of IDetailCustomization interface
+
+protected:
+	FGroupedSpriteComponentDetailsCustomization();
+
+	FReply SplitSprites();
+
+protected:
+	TArray<TWeakObjectPtr<UObject>> ObjectsBeingCustomized;
+};
\ No newline at end of file
diff --git a/Engine/Plugins/2D/Paper2D/Source/Paper2DEditor/Private/GroupedSprites/PaperGroupedSpriteUtilities.cpp b/Engine/Plugins/2D/Paper2D/Source/Paper2DEditor/Private/GroupedSprites/PaperGroupedSpriteUtilities.cpp
new file mode 100644
--- /dev/null
+++ b/Engine/Plugins/2D/Paper2D/Source/Paper2DEditor/Private/GroupedSprites/PaperGroupedSpriteUtilities.cpp
@@ -0,0 +1,232 @@
+// Copyright 1998-2015 Epic Games, Inc. All Rights Reserved.
+
+#include "Paper2DEditorPrivatePCH.h"
+#include "AssetData.h"
+#include "PaperGroupedSpriteActor.h"
+#include "PaperGroupedSpriteUtilities.h"
+#include "ScopedTransaction.h"
+#include "ILayers.h"
+#include "ComponentReregisterContext.h"
+
+#define LOCTEXT_NAMESPACE "SpriteEditor"
+
+//////////////////////////////////////////////////////////////////////////
+// FPaperGroupedSpriteUtilities
+
+void FPaperGroupedSpriteUtilities::BuildHarvestList(const TArray<UObject*>& ObjectsToConsider, TSubclassOf<UActorComponent> HarvestClassType, TArray<UActorComponent*>& OutComponentsToHarvest, TArray<AActor*>& OutActorsToDelete)
+{
+	// Determine the components to harvest
+	for (UObject* Object : ObjectsToConsider)
+	{
+		if (Object == nullptr)
+		{
+			continue;
+		}
+
+		if (AActor* SelectedActor = Cast<AActor>(Object))
+		{
+			TArray<UActorComponent*> Components;
+			SelectedActor->GetComponents(Components);
+
+			bool bHadHarvestableComponents = false;
+			for (UActorComponent* Component : Components)
+			{
+				if (Component->IsA(HarvestClassType) && !Component->IsEditorOnly())
+				{
+					bHadHarvestableComponents = true;
+					OutComponentsToHarvest.Add(Component);
+				}
+			}
+
+			if (bHadHarvestableComponents)
+			{
+				OutActorsToDelete.Add(SelectedActor);
+			}
+		}
+		else if (UActorComponent* SelectedComponent = Cast<UActorComponent>(Object))
+		{
+			if (SelectedComponent->IsA(HarvestClassType) && !SelectedComponent->IsEditorOnly())
+			{
+				OutComponentsToHarvest.Add(SelectedComponent);
+				if (AActor* SelectedComponentOwner = SelectedComponent->GetOwner())
+				{
+					OutActorsToDelete.AddUnique(SelectedComponentOwner);
+				}
+			}
+		}
+	}
+}
+
+FBox FPaperGroupedSpriteUtilities::ComputeBoundsForComponents(const TArray<UActorComponent*>& ComponentList)
+{
+	FBox BoundsOfList(ForceInit);
+
+	for (const UActorComponent* Component : ComponentList)
+	{
+		if (const USceneComponent* SceneComponent = Cast<const USceneComponent>(Component))
+		{
+			BoundsOfList += SceneComponent->Bounds.GetBox();
+		}
+	}
+
+	return BoundsOfList;
+}
+
+void FPaperGroupedSpriteUtilities::SplitSprites(const TArray<UObject*>& InObjectList)
+{
+	TArray<UActorComponent*> ComponentsToHarvest;
+	TSubclassOf<UActorComponent> HarvestClassType = UPaperSpriteComponent::StaticClass();
+	TArray<AActor*> ActorsToDelete;
+	TArray<AActor*> ActorsCreated;
+
+	FPaperGroupedSpriteUtilities::BuildHarvestList(InObjectList, UPaperGroupedSpriteComponent::StaticClass(), /*out*/ ComponentsToHarvest, /*out*/ ActorsToDelete);
+
+	if (ComponentsToHarvest.Num() > 0)
+	{
+		if (UWorld* World = ComponentsToHarvest[0]->GetWorld())
+		{
+			const FScopedTransaction Transaction(LOCTEXT("SplitSprites", "Split sprite instances"));
+
+			// Create an instance from each item of each batch component that we're harvesting
+			for (UActorComponent* SourceComponent : ComponentsToHarvest)
+			{
+				UPaperGroupedSpriteComponent* SourceBatchComponent = CastChecked<UPaperGroupedSpriteComponent>(SourceComponent);
+
+				for (const FSpriteInstanceData& InstanceData : SourceBatchComponent->GetPerInstanceSpriteData())
+				{
+					if (InstanceData.SourceSprite != nullptr)
+					{
+						const FTransform InstanceTransform(FTransform(InstanceData.Transform) * SourceBatchComponent->ComponentToWorld);
+
+						FActorSpawnParameters SpawnParams;
+						SpawnParams.bDeferConstruction = true;
+						if (APaperSpriteActor* SpawnedActor = World->SpawnActor<APaperSpriteActor>(SpawnParams))
+						{
+							UPaperSpriteComponent* SpawnedSpriteComponent = SpawnedActor->GetRenderComponent();
+
+							{
+								FComponentReregisterContext ReregisterContext(SpawnedSpriteComponent);
+
+								SpawnedSpriteComponent->Modify();
+								SpawnedSpriteComponent->SetSpriteColor(InstanceData.VertexColor);
+								SpawnedSpriteComponent->SetSprite(InstanceData.SourceSprite);
+
+								// Apply the material override if there is one
+								UMaterialInterface* InstanceMaterial = SourceBatchComponent->GetMaterial(InstanceData.MaterialIndex);
+								if (InstanceMaterial != InstanceData.SourceSprite->GetMaterial(0))
+								{
+									SpawnedSpriteComponent->SetMaterial(0, InstanceMaterial);
+								}
+							}
+
+							UGameplayStatics::FinishSpawningActor(SpawnedActor, InstanceTransform);
+							ActorsCreated.Add(SpawnedActor);
+
+							// Give it a good name
+							if (InstanceData.SourceSprite != nullptr)
+							{
+								FActorLabelUtilities::SetActorLabelUnique(SpawnedActor, InstanceData.SourceSprite->GetName());
+							}
+						}
+					}
+				}
+			}
+
+			// Delete the existing actor instances
+			for (AActor* ActorToDelete : ActorsToDelete)
+			{
+				// Remove from active selection in editor
+				GEditor->SelectActor(ActorToDelete, /*bSelected=*/ false, /*bNotify=*/ false);
+				GEditor->Layers->DisassociateActorFromLayers(ActorToDelete);
+				World->EditorDestroyActor(ActorToDelete, /*bShouldModifyLevel=*/ true);
+			}
+
+			// Select the new replacements instance
+			GEditor->SelectNone(/*bNoteSelectionChange=*/ false, false, false);
+			for (AActor* SpawnedActor : ActorsCreated)
+			{
+				GEditor->SelectActor(SpawnedActor, /*bSelected=*/ true, /*bNotify=*/ true);
+			}
+			GEditor->NoteSelectionChange();
+		}
+		else
+		{
+			// We're in the BP editor and don't currently support splitting here!
+			UE_LOG(LogPaper2DEditor, Warning, TEXT("Splitting sprites in the Blueprint editor is not currently supported"));
+		}
+	}
+}
+
+void FPaperGroupedSpriteUtilities::MergeSprites(const TArray<UObject*>& InObjectList)
+{
+	TArray<UActorComponent*> ComponentsToHarvest;
+	TSubclassOf<UActorComponent> HarvestClassType = UPaperSpriteComponent::StaticClass();
+	TArray<AActor*> ActorsToDelete;
+
+	FPaperGroupedSpriteUtilities::BuildHarvestList(InObjectList, UPaperSpriteComponent::StaticClass(), /*out*/ ComponentsToHarvest, /*out*/ ActorsToDelete);
+
+	if (ComponentsToHarvest.Num() > 0)
+	{
+		if (UWorld* World = ComponentsToHarvest[0]->GetWorld())
+		{
+			const FBox ComponentBounds = ComputeBoundsForComponents(ComponentsToHarvest);
+			const FTransform MergedWorldTM(ComponentBounds.GetCenter());
+
+			FActorSpawnParameters SpawnParams;
+			SpawnParams.bDeferConstruction = true;
+
+			if (APaperGroupedSpriteActor* SpawnedActor = World->SpawnActor<APaperGroupedSpriteActor>(SpawnParams))
+			{
+				const FScopedTransaction Transaction(LOCTEXT("MergeSprites", "Merge sprite instances"));
+
+				// Create a merged sprite component
+				{
+					UPaperGroupedSpriteComponent* MergedSpriteComponent = SpawnedActor->GetRenderComponent();
+					FComponentReregisterContext ReregisterContext(MergedSpriteComponent);
+
+					// Create an instance from each sprite component that we're harvesting
+					for (UActorComponent* SourceComponent : ComponentsToHarvest)
+					{
+						UPaperSpriteComponent* SourceSpriteComponent = CastChecked<UPaperSpriteComponent>(SourceComponent);
+
+						UPaperSprite* Sprite = SourceSpriteComponent->GetSprite();
+						const FLinearColor SpriteColor = SourceSpriteComponent->GetSpriteColor();
+						const FTransform RelativeSpriteTransform = SourceSpriteComponent->GetComponentTransform().GetRelativeTransform(MergedWorldTM);
+
+						UMaterialInterface* OverrideMaterial = SourceSpriteComponent->GetMaterial(0);
+						if ((Sprite == nullptr) || (OverrideMaterial == Sprite->GetMaterial(0)))
+						{
+							OverrideMaterial = nullptr;
+						}
+
+						MergedSpriteComponent->AddInstanceWithMaterial(RelativeSpriteTransform, Sprite, OverrideMaterial, /*bWorldSpace=*/ false, SpriteColor);
+					}
+				}
+
+				// Finalize the new component
+				UGameplayStatics::FinishSpawningActor(SpawnedActor, MergedWorldTM);
+
+				// Delete the existing actor instances
+				for (AActor* ActorToDelete : ActorsToDelete)
+				{
+					// Remove from active selection in editor
+					GEditor->SelectActor(ActorToDelete, /*bSelected=*/ false, /*bNotify=*/ false);
+					GEditor->Layers->DisassociateActorFromLayers(ActorToDelete);
+					World->EditorDestroyActor(ActorToDelete, /*bShouldModifyLevel=*/ true);
+				}
+
+				// Select the new replacement instance
+				GEditor->SelectActor(SpawnedActor, /*bSelected=*/ true, /*bNotify=*/ true);
+			}
+			else
+			{
+				// We're in the BP editor and don't currently support merging here!
+				UE_LOG(LogPaper2DEditor, Warning, TEXT("Merging sprites in the Blueprint editor is not currently supported"));
+			}
+		}
+	}
+}
+
+//////////////////////////////////////////////////////////////////////////
+
+#undef LOCTEXT_NAMESPACE
diff --git a/Engine/Plugins/2D/Paper2D/Source/Paper2DEditor/Private/GroupedSprites/PaperGroupedSpriteUtilities.h b/Engine/Plugins/2D/Paper2D/Source/Paper2DEditor/Private/GroupedSprites/PaperGroupedSpriteUtilities.h
new file mode 100644
--- /dev/null
+++ b/Engine/Plugins/2D/Paper2D/Source/Paper2DEditor/Private/GroupedSprites/PaperGroupedSpriteUtilities.h
@@ -0,0 +1,18 @@
+// Copyright 1998-2015 Epic Games, Inc. All Rights Reserved.
+
+#pragma once
+
+class APaperGroupedSpriteActor;
+
+class FPaperGroupedSpriteUtilities
+{
+public:
+	static void BuildHarvestList(const TArray<UObject*>& ObjectsToConsider, TSubclassOf<UActorComponent> HarvestClassType, TArray<UActorComponent*>& OutComponentsToHarvest, TArray<AActor*>& OutActorsToDelete);
+
+	// Computes the enclosing bounding box of the specified components (using their individual bounds)
+	static FBox ComputeBoundsForComponents(const TArray<UActorComponent*>& ComponentList);
+
+	static void SplitSprites(const TArray<UObject*>& InObjectList);
+
+	static void MergeSprites(const TArray<UObject*>& InObjectList);
+};
diff --git a/Engine/Plugins/2D/Paper2D/Source/Paper2DEditor/Private/LevelEditorMenuExtensions/Paper2DLevelEditorExtensions.cpp b/Engine/Plugins/2D/Paper2D/Source/Paper2DEditor/Private/LevelEditorMenuExtensions/Paper2DLevelEditorExtensions.cpp
new file mode 100644
--- /dev/null
+++ b/Engine/Plugins/2D/Paper2D/Source/Paper2DEditor/Private/LevelEditorMenuExtensions/Paper2DLevelEditorExtensions.cpp
@@ -0,0 +1,147 @@
+// Copyright 1998-2015 Epic Games, Inc. All Rights Reserved.
+
+#include "Paper2DEditorPrivatePCH.h"
+#include "Paper2DLevelEditorExtensions.h"
+#include "LevelEditor.h"
+
+#include "PaperSpriteComponent.h"
+#include "PaperGroupedSpriteComponent.h"
+#include "GroupedSprites/PaperGroupedSpriteUtilities.h"
+
+#define LOCTEXT_NAMESPACE "Paper2D"
+
+//////////////////////////////////////////////////////////////////////////
+
+FLevelEditorModule::FLevelViewportMenuExtender_SelectedActors LevelEditorMenuExtenderDelegate;
+FDelegateHandle LevelEditorExtenderDelegateHandle;
+
+//////////////////////////////////////////////////////////////////////////
+// FPaperLevelEditorMenuExtensions_Impl
+
+class FPaperLevelEditorMenuExtensions_Impl
+{
+public:
+	static void MergeSprites()
+	{
+		// Create an array of actors to consider
+		TArray<UObject*> SelectedActors;
+		GEditor->GetSelectedActors()->GetSelectedObjects(AActor::StaticClass(), /*out*/ SelectedActors);
+
+		// Merge them
+		FPaperGroupedSpriteUtilities::MergeSprites(SelectedActors);
+	}
+
+	static void SplitSprites()
+	{
+		// Create an array of actors to consider
+		TArray<UObject*> SelectedActors;
+		GEditor->GetSelectedActors()->GetSelectedObjects(AActor::StaticClass(), /*out*/ SelectedActors);
+
+		// Split them
+		FPaperGroupedSpriteUtilities::SplitSprites(SelectedActors);
+	}
+
+	static void CreateSpriteActionsMenuEntries(FMenuBuilder& MenuBuilder, bool bCanMerge, bool bCanSplit)
+	{
+		MenuBuilder.BeginSection("Paper2D", LOCTEXT("Paper2DLevelEditorHeading", "Paper2D"));
+
+		if (bCanMerge)
+		{
+			FUIAction Action_MergeSprites(FExecuteAction::CreateStatic(&FPaperLevelEditorMenuExtensions_Impl::MergeSprites));
+
+			MenuBuilder.AddMenuEntry(
+				LOCTEXT("MenuExtensionMergeSprites", "Merge Sprites"),
+				LOCTEXT("MenuExtensionMergeSprites_Tooltip", "Replaces all of the selected actors that contain a sprite component with a single grouped sprite component"),
+				FSlateIcon(),
+				Action_MergeSprites,
+				NAME_None,
+				EUserInterfaceActionType::Button);
+		}
+
+		if (bCanSplit)
+		{
+			FUIAction Action_SplitSprites(FExecuteAction::CreateStatic(&FPaperLevelEditorMenuExtensions_Impl::SplitSprites));
+
+			MenuBuilder.AddMenuEntry(
+				LOCTEXT("MenuExtensionSplitSprites", "Split Sprites"),
+				LOCTEXT("MenuExtensionSplitSprites_Tooltip", "Replaces all of the selected actors that contain a grouped sprite component with many individual sprite components, one per element"),
+				FSlateIcon(),
+				Action_SplitSprites,
+				NAME_None,
+				EUserInterfaceActionType::Button);
+		}
+
+		MenuBuilder.EndSection();
+	}
+
+	static TSharedRef<FExtender> OnExtendLevelEditorMenu(const TSharedRef<FUICommandList> CommandList, TArray<AActor*> SelectedActors)
+	{
+		TSharedRef<FExtender> Extender(new FExtender());
+
+		// Run thru the actors to determine if any meet our criteria
+		bool bCanMergeSprites = false;
+		bool bCanSplitSprites = false;
+
+		for (AActor* Actor : SelectedActors)
+		{
+			TInlineComponentArray<UActorComponent*> ActorComponents;
+			Actor->GetComponents(ActorComponents);
+
+			for (UActorComponent* Component : ActorComponents)
+			{
+				if (Component->IsA(UPaperSpriteComponent::StaticClass()))
+				{
+					bCanMergeSprites = true;
+				}
+				else if (Component->IsA(UPaperGroupedSpriteComponent::StaticClass()))
+				{
+					bCanSplitSprites = true;
+				}
+			}
+
+		}
+
+		bCanMergeSprites = bCanMergeSprites && (SelectedActors.Num() > 1);
+
+		if (bCanMergeSprites || bCanSplitSprites)
+		{
+			// Add the sprite actions sub-menu extender
+			Extender->AddMenuExtension(
+				"ActorType",
+				EExtensionHook::Before,
+				nullptr,
+				FMenuExtensionDelegate::CreateStatic(&FPaperLevelEditorMenuExtensions_Impl::CreateSpriteActionsMenuEntries, /*bCanMerge=*/ bCanMergeSprites, /*bCanSplit=*/ bCanSplitSprites));
+		}
+
+		return Extender;
+	}
+};
+
+//////////////////////////////////////////////////////////////////////////
+// FPaperLevelEditorMenuExtensions
+
+void FPaperLevelEditorMenuExtensions::InstallHooks()
+{
+	LevelEditorMenuExtenderDelegate = FLevelEditorModule::FLevelViewportMenuExtender_SelectedActors::CreateStatic(&FPaperLevelEditorMenuExtensions_Impl::OnExtendLevelEditorMenu);
+
+	FLevelEditorModule& LevelEditorModule = FModuleManager::Get().LoadModuleChecked<FLevelEditorModule>("LevelEditor");
+	auto& MenuExtenders = LevelEditorModule.GetAllLevelViewportContextMenuExtenders();
+	MenuExtenders.Add(LevelEditorMenuExtenderDelegate);
+	LevelEditorExtenderDelegateHandle = MenuExtenders.Last().GetHandle();
+}
+
+void FPaperLevelEditorMenuExtensions::RemoveHooks()
+{
+	// Remove level viewport context menu extenders
+	if (FModuleManager::Get().IsModuleLoaded("LevelEditor"))
+	{
+		FLevelEditorModule& LevelEditorModule = FModuleManager::LoadModuleChecked<FLevelEditorModule>("LevelEditor");
+		LevelEditorModule.GetAllLevelViewportContextMenuExtenders().RemoveAll([&](const FLevelEditorModule::FLevelViewportMenuExtender_SelectedActors& Delegate) {
+			return Delegate.GetHandle() == LevelEditorExtenderDelegateHandle;
+		});
+	}
+}
+
+//////////////////////////////////////////////////////////////////////////
+
+#undef LOCTEXT_NAMESPACE
\ No newline at end of file
diff --git a/Engine/Plugins/2D/Paper2D/Source/Paper2DEditor/Private/LevelEditorMenuExtensions/Paper2DLevelEditorExtensions.h b/Engine/Plugins/2D/Paper2D/Source/Paper2DEditor/Private/LevelEditorMenuExtensions/Paper2DLevelEditorExtensions.h
new file mode 100644
--- /dev/null
+++ b/Engine/Plugins/2D/Paper2D/Source/Paper2DEditor/Private/LevelEditorMenuExtensions/Paper2DLevelEditorExtensions.h
@@ -0,0 +1,11 @@
+// Copyright 1998-2015 Epic Games, Inc. All Rights Reserved.
+
+#pragma once
+
+// Integrate Paper2D actions into the level editor context menu
+class FPaperLevelEditorMenuExtensions
+{
+public:
+	static void InstallHooks();
+	static void RemoveHooks();
+};
\ No newline at end of file
diff --git a/Engine/Plugins/2D/Paper2D/Source/Paper2DEditor/Private/Paper2DEditorModule.cpp b/Engine/Plugins/2D/Paper2D/Source/Paper2DEditor/Private/Paper2DEditorModule.cpp
--- a/Engine/Plugins/2D/Paper2D/Source/Paper2DEditor/Private/Paper2DEditorModule.cpp
+++ b/Engine/Plugins/2D/Paper2D/Source/Paper2DEditor/Private/Paper2DEditorModule.cpp
@@ -12,6 +12,7 @@
 #include "AssetEditorToolkit.h"
 #include "ModuleManager.h"
 #include "ContentBrowserExtensions/ContentBrowserExtensions.h"
+#include "LevelEditorMenuExtensions/Paper2DLevelEditorExtensions.h"
 #include "PaperImporterSettings.h"
 
 // Sprite support
@@ -50,7 +51,7 @@
 
 // Grouped sprite support
 #include "PaperGroupedSpriteComponent.h"
-#include "GroupedSpriteDetailsCustomization.h"
+#include "GroupedSprites/GroupedSpriteDetailsCustomization.h"
 
 // Settings
 #include "PaperRuntimeSettings.h"
@@ -183,8 +184,9 @@ public:
 			FSlateIcon(),
 			false);
 
-		// Integrate Paper2D actions associated with existing engine types (e.g., Texture2D) into the content browser
+		// Integrate Paper2D actions into existing editor context menus
 		FPaperContentBrowserExtensions::InstallHooks();
+		FPaperLevelEditorMenuExtensions::InstallHooks();
 
 		// Register with the mesh paint module
 		if (IMeshPaintModule* MeshPaintModule = FModuleManager::LoadModulePtr<IMeshPaintModule>("MeshPaint"))
@@ -216,6 +218,7 @@ public:
 				SpriteMeshPaintAdapterFactory.Reset();
 			}
 
+			FPaperLevelEditorMenuExtensions::RemoveHooks();
 			FPaperContentBrowserExtensions::RemoveHooks();
 
 			FComponentAssetBrokerage::UnregisterBroker(PaperTileMapBroker);
diff --git a/Engine/Plugins/2D/Paper2D/Source/Paper2DEditor/Private/SpriteEditor/SpriteComponentDetailsCustomization.cpp b/Engine/Plugins/2D/Paper2D/Source/Paper2DEditor/Private/SpriteEditor/SpriteComponentDetailsCustomization.cpp
--- a/Engine/Plugins/2D/Paper2D/Source/Paper2DEditor/Private/SpriteEditor/SpriteComponentDetailsCustomization.cpp
+++ b/Engine/Plugins/2D/Paper2D/Source/Paper2DEditor/Private/SpriteEditor/SpriteComponentDetailsCustomization.cpp
@@ -4,10 +4,7 @@
 #include "PaperSpriteComponent.h"
 #include "PaperGroupedSpriteActor.h"
 #include "SpriteComponentDetailsCustomization.h"
-#include "ScopedTransaction.h"
-#include "ILayers.h"
-#include "GroupedSpriteDetailsCustomization.h"
-#include "ComponentReregisterContext.h"
+#include "GroupedSprites/PaperGroupedSpriteUtilities.h"
 
 #define LOCTEXT_NAMESPACE "SpriteEditor"
 
@@ -42,96 +39,19 @@ void FSpriteComponentDetailsCustomization::CustomizeDetails(IDetailLayoutBuilder
 				[
 					SNew(SButton)
 					.Text(LOCTEXT("MergeSprites", "Merge Sprites"))
-					.ToolTipText(LOCTEXT("MergeSprites_Tooltip", "Merges all selected sprites into entries on a single InstancedSpriteComponent"))
+					.ToolTipText(LOCTEXT("MergeSprites_Tooltip", "Merges all selected sprite components into entries on a single grouped sprite component"))
 					.OnClicked(this, &FSpriteComponentDetailsCustomization::MergeSprites)
 				]
 			];
 	}
 }
 
-
-FBox FSpriteComponentDetailsCustomization::ComputeBoundsForComponents(const TArray<UActorComponent*>& ComponentList)
-{
-	FBox BoundsOfList(ForceInit);
-
-	for (const UActorComponent* Component : ComponentList)
-	{
-		if (const USceneComponent* SceneComponent = Cast<const USceneComponent>(Component))
-		{
-			BoundsOfList += SceneComponent->Bounds.GetBox();
-		}
-	}
-
-	return BoundsOfList;
-}
-
 FReply FSpriteComponentDetailsCustomization::MergeSprites()
 {
-	TArray<UActorComponent*> ComponentsToHarvest;
-	TSubclassOf<UActorComponent> HarvestClassType = UPaperSpriteComponent::StaticClass();
-	TArray<AActor*> ActorsToDelete;
-
-	FGroupedSpriteComponentDetailsCustomization::BuildHarvestList(ObjectsBeingCustomized, UPaperSpriteComponent::StaticClass(), /*out*/ ComponentsToHarvest, /*out*/ ActorsToDelete);
-
-	if (ComponentsToHarvest.Num() > 0)
-	{
-		if (UWorld* World = ComponentsToHarvest[0]->GetWorld())
-		{
-			const FBox ComponentBounds = ComputeBoundsForComponents(ComponentsToHarvest);
-			const FTransform MergedWorldTM(ComponentBounds.GetCenter());
-
-			FActorSpawnParameters SpawnParams;
-			SpawnParams.bDeferConstruction = true;
-			if (APaperGroupedSpriteActor* SpawnedActor = World->SpawnActor<APaperGroupedSpriteActor>(SpawnParams))
-			{
-				const FScopedTransaction Transaction(LOCTEXT("MergeSprites", "Merge sprite instances"));
-
-				// Create a merged sprite component
-				{
-					UPaperGroupedSpriteComponent* MergedSpriteComponent = SpawnedActor->GetRenderComponent();
-					FComponentReregisterContext ReregisterContext(MergedSpriteComponent);
-
-					// Create an instance from each sprite component that we're harvesting
-					for (UActorComponent* SourceComponent : ComponentsToHarvest)
-					{
-						UPaperSpriteComponent* SourceSpriteComponent = CastChecked<UPaperSpriteComponent>(SourceComponent);
-
-						UPaperSprite* Sprite = SourceSpriteComponent->GetSprite();
-						const FLinearColor SpriteColor = SourceSpriteComponent->GetSpriteColor();
-						const FTransform RelativeSpriteTransform = SourceSpriteComponent->GetComponentTransform().GetRelativeTransform(MergedWorldTM);
+	TArray<UObject*> StrongObjects;
+	CopyFromWeakArray(StrongObjects, ObjectsBeingCustomized);
 
-						UMaterialInterface* OverrideMaterial = SourceSpriteComponent->GetMaterial(0);
-						if ((Sprite == nullptr) || (OverrideMaterial == Sprite->GetMaterial(0)))
-						{
-							OverrideMaterial = nullptr;
-						}
-
-						MergedSpriteComponent->AddInstanceWithMaterial(RelativeSpriteTransform, Sprite, OverrideMaterial, /*bWorldSpace=*/ false, SpriteColor);
-					}
-				}
-
-				// Finalize the new component
-				UGameplayStatics::FinishSpawningActor(SpawnedActor, MergedWorldTM);
-
-				// Delete the existing actor instances
-				for (AActor* ActorToDelete : ActorsToDelete)
-				{
-					// Remove from active selection in editor
-					GEditor->SelectActor(ActorToDelete, /*bSelected=*/ false, /*bNotify=*/ false);
-					GEditor->Layers->DisassociateActorFromLayers(ActorToDelete);
-					World->EditorDestroyActor(ActorToDelete, /*bShouldModifyLevel=*/ true);
-				}
-
-				// Select the new replacement instance
-				GEditor->SelectActor(SpawnedActor, /*bSelected=*/ true, /*bNotify=*/ true);
-			}
-			else
-			{
-				// We're in the BP editor and don't currently support merging here!
-				UE_LOG(LogPaper2DEditor, Warning, TEXT("Merging sprites in the Blueprint editor is not currently supported"));
-			}
-		}
-	}
+	FPaperGroupedSpriteUtilities::MergeSprites(StrongObjects);
 
 	return FReply::Handled();
 }
diff --git a/Engine/Plugins/2D/Paper2D/Source/Paper2DEditor/Private/SpriteEditor/SpriteComponentDetailsCustomization.h b/Engine/Plugins/2D/Paper2D/Source/Paper2DEditor/Private/SpriteEditor/SpriteComponentDetailsCustomization.h
--- a/Engine/Plugins/2D/Paper2D/Source/Paper2DEditor/Private/SpriteEditor/SpriteComponentDetailsCustomization.h
+++ b/Engine/Plugins/2D/Paper2D/Source/Paper2DEditor/Private/SpriteEditor/SpriteComponentDetailsCustomization.h
@@ -20,8 +20,6 @@ public:
 protected:
 	FReply MergeSprites();
 
-	static FBox ComputeBoundsForComponents(const TArray<UActorComponent*>& ComponentList);
-
 protected:
 	TArray<TWeakObjectPtr<UObject>> ObjectsBeingCustomized;
 };
