Paper2D: Promoted UPaperFlipbookComponent from EarlyAccessPreview to fully supported

diff --git a/Engine/Plugins/2D/Paper2D/Source/Paper2D/Classes/PaperFlipbookComponent.h b/Engine/Plugins/2D/Paper2D/Source/Paper2D/Classes/PaperFlipbookComponent.h
--- a/Engine/Plugins/2D/Paper2D/Source/Paper2D/Classes/PaperFlipbookComponent.h
+++ b/Engine/Plugins/2D/Paper2D/Source/Paper2D/Classes/PaperFlipbookComponent.h
@@ -8,7 +8,7 @@
 // Event for a non-looping flipbook finishing play
 DECLARE_DYNAMIC_MULTICAST_DELEGATE(FFlipbookFinishedPlaySignature);
 
-UCLASS(ShowCategories=(Mobility, ComponentReplication), ClassGroup=Paper2D, EarlyAccessPreview, meta=(BlueprintSpawnableComponent))
+UCLASS(ShowCategories=(Mobility, ComponentReplication), ClassGroup=Paper2D, meta=(BlueprintSpawnableComponent))
 class PAPER2D_API UPaperFlipbookComponent : public UMeshComponent
 {
 	GENERATED_UCLASS_BODY()
diff --git a/Engine/Plugins/2D/Paper2D/Source/Paper2DEditor/Private/FlipbookEditor/FlipbookEditor.cpp b/Engine/Plugins/2D/Paper2D/Source/Paper2DEditor/Private/FlipbookEditor/FlipbookEditor.cpp
--- a/Engine/Plugins/2D/Paper2D/Source/Paper2DEditor/Private/FlipbookEditor/FlipbookEditor.cpp
+++ b/Engine/Plugins/2D/Paper2D/Source/Paper2DEditor/Private/FlipbookEditor/FlipbookEditor.cpp
@@ -247,25 +247,7 @@ TSharedRef<SDockTab> FFlipbookEditor::SpawnTab_Viewport(const FSpawnTabArgs& Arg
 			
 			+SVerticalBox::Slot()
 			[
-				SNew(SOverlay)
-
-				// The flipbook editor viewport
-				+SOverlay::Slot()
-				[
-					ViewportPtr.ToSharedRef()
-				]
-
-				// Bottom-right corner text indicating the preview nature of the sprite editor
-				+SOverlay::Slot()
-				.Padding(10)
-				.VAlign(VAlign_Bottom)
-				.HAlign(HAlign_Right)
-				[
-					SNew(STextBlock)
-					.Visibility( EVisibility::HitTestInvisible )
-					.TextStyle( FEditorStyle::Get(), "Graph.CornerText" )
-					.Text(LOCTEXT("FlipbookEditorViewportExperimentalWarning", "Early access preview"))
-				]
+				ViewportPtr.ToSharedRef()
 			]
 
 			+SVerticalBox::Slot()
