Flipbook Editor: Increased the default display size of a key frame (64->120) to make it easier to read keyframe names

diff --git a/Engine/Plugins/2D/Paper2D/Source/Paper2DEditor/Private/FlipbookEditor/SFlipbookTimeline.cpp b/Engine/Plugins/2D/Paper2D/Source/Paper2DEditor/Private/FlipbookEditor/SFlipbookTimeline.cpp
--- a/Engine/Plugins/2D/Paper2D/Source/Paper2DEditor/Private/FlipbookEditor/SFlipbookTimeline.cpp
+++ b/Engine/Plugins/2D/Paper2D/Source/Paper2DEditor/Private/FlipbookEditor/SFlipbookTimeline.cpp
@@ -27,7 +27,7 @@ void SFlipbookTimeline::Construct(const FArguments& InArgs, TSharedPtr<FUIComman
 	OnSelectionChanged = InArgs._OnSelectionChanged;
 	CommandList = InCommandList;
 
-	SlateUnitsPerFrame = 64.0f;
+	SlateUnitsPerFrame = 120.0f;
 
 	BackgroundPerFrameSlices = SNew(SHorizontalBox);
 
