Paper2D: Added a sort command in the details customization for grouped sprites to aid in working with translucent sprites

diff --git a/Engine/Plugins/2D/Paper2D/Source/Paper2D/Classes/PaperGroupedSpriteComponent.h b/Engine/Plugins/2D/Paper2D/Source/Paper2D/Classes/PaperGroupedSpriteComponent.h
--- a/Engine/Plugins/2D/Paper2D/Source/Paper2D/Classes/PaperGroupedSpriteComponent.h
+++ b/Engine/Plugins/2D/Paper2D/Source/Paper2D/Classes/PaperGroupedSpriteComponent.h
@@ -96,6 +96,10 @@ public:
 	UFUNCTION(BlueprintCallable, Category = "Components|Sprite")
 	int32 GetInstanceCount() const;
 
+	/** Sort all instances by their world space position along the specified axis */
+	UFUNCTION(BlueprintCallable, Category = "Components|Sprite")
+	void SortInstancesAlongAxis(FVector WorldSpaceSortAxis);
+
 	// UActorComponent interface
 	virtual bool ShouldCreatePhysicsState() const override;
 	virtual void CreatePhysicsState() override;
diff --git a/Engine/Plugins/2D/Paper2D/Source/Paper2D/Private/PaperGroupedSpriteComponent.cpp b/Engine/Plugins/2D/Paper2D/Source/Paper2D/Private/PaperGroupedSpriteComponent.cpp
--- a/Engine/Plugins/2D/Paper2D/Source/Paper2D/Private/PaperGroupedSpriteComponent.cpp
+++ b/Engine/Plugins/2D/Paper2D/Source/Paper2D/Private/PaperGroupedSpriteComponent.cpp
@@ -421,7 +421,6 @@ void UPaperGroupedSpriteComponent::RebuildInstances()
 	// update the physics state
 	if (bPhysicsStateCreated)
 	{
-		// TODO: it may be possible to instead just update the BodyInstanceIndex for all bodies after the removed instance. 
 		ClearAllInstanceBodies();
 		CreateAllInstanceBodies();
 	}
@@ -579,6 +578,45 @@ void UPaperGroupedSpriteComponent::GetReferencedSpriteAssets(TArray<UObject*>& I
 	}
 }
 
+void UPaperGroupedSpriteComponent::SortInstancesAlongAxis(FVector WorldSpaceSortAxis)
+{
+	struct FSortStruct
+	{
+		int32 OldIndex;
+		float SortKey;
+
+		FSortStruct(int32 InIndex, float InDepth)
+			: OldIndex(InIndex)
+			, SortKey(InDepth)
+		{
+		}
+	};
+
+	// Figure out the sort order
+	TArray<FSortStruct> SortArray;
+	SortArray.Empty(PerInstanceSpriteData.Num());
+	for (int32 Index = 0; Index < PerInstanceSpriteData.Num(); ++Index)
+	{
+		const FVector InstanceWorldPos = ComponentToWorld.TransformPosition(PerInstanceSpriteData[Index].Transform.GetOrigin());
+		const float SortKey = FVector::DotProduct(InstanceWorldPos, WorldSpaceSortAxis);
+		SortArray.Emplace(Index, SortKey);
+	}
+
+	SortArray.Sort([](const FSortStruct& LHS, const FSortStruct& RHS)  { return LHS.SortKey > RHS.SortKey; });
+
+	// Reorganize the array to match
+	TArray<FSpriteInstanceData> OldData(PerInstanceSpriteData);
+	PerInstanceSpriteData.Reset();
+
+	for (const FSortStruct& SortedItem : SortArray)
+	{
+		PerInstanceSpriteData.Add(OldData[SortedItem.OldIndex]);
+	}
+
+	// Rebuild, as the rendering scene proxy and body setup orderings are both out of date
+	RebuildInstances();
+}
+
 //////////////////////////////////////////////////////////////////////////
 
 #undef LOCTEXT_NAMESPACE
diff --git a/Engine/Plugins/2D/Paper2D/Source/Paper2DEditor/Private/GroupedSprites/GroupedSpriteDetailsCustomization.cpp b/Engine/Plugins/2D/Paper2D/Source/Paper2DEditor/Private/GroupedSprites/GroupedSpriteDetailsCustomization.cpp
--- a/Engine/Plugins/2D/Paper2D/Source/Paper2DEditor/Private/GroupedSprites/GroupedSpriteDetailsCustomization.cpp
+++ b/Engine/Plugins/2D/Paper2D/Source/Paper2DEditor/Private/GroupedSprites/GroupedSpriteDetailsCustomization.cpp
@@ -7,6 +7,9 @@
 #include "PaperSpriteComponent.h"
 #include "PaperGroupedSpriteComponent.h"
 #include "PaperGroupedSpriteUtilities.h"
+#include "ScopedTransaction.h"
+
+#include "Engine/RendererSettings.h"
 
 #define LOCTEXT_NAMESPACE "SpriteEditor"
 
@@ -30,25 +33,55 @@ void FGroupedSpriteComponentDetailsCustomization::CustomizeDetails(IDetailLayout
 	ObjectsBeingCustomized.Empty();
 	DetailBuilder.GetObjectsBeingCustomized(/*out*/ ObjectsBeingCustomized);
 
-	{
-		// Expose split buttons
-		FDetailWidgetRow& SplitRow = SpriteCategory.AddCustomRow(LOCTEXT("SplitSearchText", "Split"))
-			.WholeRowContent()
+
+	TSharedRef<SWrapBox> ButtonBox = SNew(SWrapBox).UseAllottedWidth(true);
+
+	const float MinButtonSize = 100.0f;
+	const FMargin ButtonPadding(0.0f, 2.0f, 2.0f, 0.0f);
+
+	// Split button
+	ButtonBox->AddSlot()
+	.Padding(ButtonPadding)
+	[
+		SNew(SBox)
+		.MinDesiredWidth(MinButtonSize)
+		[
+			SNew(SButton)
+			.VAlign(VAlign_Center)
+			.HAlign(HAlign_Center)
+			.Text(LOCTEXT("SplitSprites", "Split Sprites"))
+			.ToolTipText(LOCTEXT("SplitSprites_Tooltip", "Splits all sprite instances into separate sprite actors or components"))
+			.OnClicked(this, &FGroupedSpriteComponentDetailsCustomization::SplitSprites)
+		]
+	];
+
+	// Sort button
+	ButtonBox->AddSlot()
+	.Padding(ButtonPadding)
+	[
+		SNew(SBox)
+		.MinDesiredWidth(MinButtonSize)
+		[
+			SNew(SButton)
+			.VAlign(VAlign_Center)
+			.HAlign(HAlign_Center)
+			.Text(LOCTEXT("SortSprites", "Sort Sprites"))
+			.ToolTipText(LOCTEXT("SortSprites_Tooltip", "Sorts all sprite instances according to the Translucency Sort Axis in the Rendering project settings"))
+			.OnClicked(this, &FGroupedSpriteComponentDetailsCustomization::SortSprites)
+		]
+	];
+
+	// Add the action buttons
+	FDetailWidgetRow& GroupActionsRow = SpriteCategory.AddCustomRow(LOCTEXT("GroupActionsSearchText", "Split Sort"))
+		.WholeRowContent()
+		[
+			SNew(SHorizontalBox)
+			+SHorizontalBox::Slot()
+			.FillWidth(1.0f)
 			[
-				SNew(SHorizontalBox)
-				+ SHorizontalBox::Slot()
-				.AutoWidth()
-				.Padding(2.0f, 0.0f)
-				.VAlign(VAlign_Center)
-				.HAlign(HAlign_Left)
-				[
-					SNew(SButton)
-					.Text(LOCTEXT("SplitSprites", "Split Sprites"))
-					.ToolTipText(LOCTEXT("SplitSprites_Tooltip", "Splits all sprite instances into separate sprite actors or components"))
-					.OnClicked(this, &FGroupedSpriteComponentDetailsCustomization::SplitSprites)
-				]
-			];
-	}
+				ButtonBox
+			]
+		];
 }
 
 FReply FGroupedSpriteComponentDetailsCustomization::SplitSprites()
@@ -61,6 +94,29 @@ FReply FGroupedSpriteComponentDetailsCustomization::SplitSprites()
 	return FReply::Handled();
 }
 
+FReply FGroupedSpriteComponentDetailsCustomization::SortSprites()
+{
+	TArray<UObject*> StrongObjects;
+	CopyFromWeakArray(StrongObjects, ObjectsBeingCustomized);
+
+	TArray<UActorComponent*> ComponentList;
+	TArray<AActor*> IgnoredList;
+	FPaperGroupedSpriteUtilities::BuildHarvestList(StrongObjects, UPaperGroupedSpriteComponent::StaticClass(), /*out*/ ComponentList, /*out*/ IgnoredList);
+
+	const FVector SortAxis = GetDefault<URendererSettings>()->TranslucentSortAxis;
+
+	const FScopedTransaction Transaction(LOCTEXT("SortSprites", "Sort instances in group"));
+	for (UActorComponent* Component : ComponentList)
+	{
+		UPaperGroupedSpriteComponent* GroupedComponent = CastChecked<UPaperGroupedSpriteComponent>(Component);
+		GroupedComponent->SortInstancesAlongAxis(SortAxis);
+	}
+
+	GEditor->RedrawLevelEditingViewports(true);
+
+	return FReply::Handled();
+}
+
 //////////////////////////////////////////////////////////////////////////
 
 #undef LOCTEXT_NAMESPACE
diff --git a/Engine/Plugins/2D/Paper2D/Source/Paper2DEditor/Private/GroupedSprites/GroupedSpriteDetailsCustomization.h b/Engine/Plugins/2D/Paper2D/Source/Paper2DEditor/Private/GroupedSprites/GroupedSpriteDetailsCustomization.h
--- a/Engine/Plugins/2D/Paper2D/Source/Paper2DEditor/Private/GroupedSprites/GroupedSpriteDetailsCustomization.h
+++ b/Engine/Plugins/2D/Paper2D/Source/Paper2DEditor/Private/GroupedSprites/GroupedSpriteDetailsCustomization.h
@@ -21,6 +21,7 @@ protected:
 	FGroupedSpriteComponentDetailsCustomization();
 
 	FReply SplitSprites();
+	FReply SortSprites();
 
 protected:
 	TArray<TWeakObjectPtr<UObject>> ObjectsBeingCustomized;
diff --git a/Engine/Plugins/2D/Paper2D/Source/Paper2DEditor/Private/GroupedSprites/PaperGroupedSpriteUtilities.cpp b/Engine/Plugins/2D/Paper2D/Source/Paper2DEditor/Private/GroupedSprites/PaperGroupedSpriteUtilities.cpp
--- a/Engine/Plugins/2D/Paper2D/Source/Paper2DEditor/Private/GroupedSprites/PaperGroupedSpriteUtilities.cpp
+++ b/Engine/Plugins/2D/Paper2D/Source/Paper2DEditor/Private/GroupedSprites/PaperGroupedSpriteUtilities.cpp
@@ -155,6 +155,8 @@ void FPaperGroupedSpriteUtilities::SplitSprites(const TArray<UObject*>& InObject
 			UE_LOG(LogPaper2DEditor, Warning, TEXT("Splitting sprites in the Blueprint editor is not currently supported"));
 		}
 	}
+
+	GEditor->RedrawLevelEditingViewports(true);
 }
 
 void FPaperGroupedSpriteUtilities::MergeSprites(const TArray<UObject*>& InObjectList)
@@ -225,6 +227,8 @@ void FPaperGroupedSpriteUtilities::MergeSprites(const TArray<UObject*>& InObject
 			}
 		}
 	}
+
+	GEditor->RedrawLevelEditingViewports(true);
 }
 
 //////////////////////////////////////////////////////////////////////////
