Paper2D: Promoted UPaperGroupedSpriteComponent from experimental to early access preview [UE-8031]

diff --git a/Engine/Plugins/2D/Paper2D/Source/Paper2D/Classes/PaperGroupedSpriteComponent.h b/Engine/Plugins/2D/Paper2D/Source/Paper2D/Classes/PaperGroupedSpriteComponent.h
--- a/Engine/Plugins/2D/Paper2D/Source/Paper2D/Classes/PaperGroupedSpriteComponent.h
+++ b/Engine/Plugins/2D/Paper2D/Source/Paper2D/Classes/PaperGroupedSpriteComponent.h
@@ -42,12 +42,12 @@ public:
 };
 
 /**
- * A component that handles rendering and collision for many instances of a UPaperSprite asset.
+ * A component that handles rendering and collision for many instances of one or more UPaperSprite assets.
  *
  * @see UPrimitiveComponent, UPaperSprite
  */
 
-UCLASS(ShowCategories=(Mobility), ClassGroup=Paper2D, meta=(BlueprintSpawnableComponent), Experimental)
+UCLASS(ShowCategories=(Mobility), ClassGroup=Paper2D, meta=(BlueprintSpawnableComponent), EarlyAccessPreview)
 class PAPER2D_API UPaperGroupedSpriteComponent : public UMeshComponent
 {
 	GENERATED_UCLASS_BODY()
