Fixing crash on startup when running commandlets

diff --git a/Engine/Plugins/2D/Paper2D/Source/Paper2DEditor/Private/Paper2DEditorModule.cpp b/Engine/Plugins/2D/Paper2D/Source/Paper2DEditor/Private/Paper2DEditorModule.cpp
--- a/Engine/Plugins/2D/Paper2D/Source/Paper2DEditor/Private/Paper2DEditorModule.cpp
+++ b/Engine/Plugins/2D/Paper2D/Source/Paper2DEditor/Private/Paper2DEditorModule.cpp
@@ -185,9 +185,11 @@ public:
 			false);
 
 		// Integrate Paper2D actions into existing editor context menus
-		FPaperContentBrowserExtensions::InstallHooks();
-		FPaperLevelEditorMenuExtensions::InstallHooks();
-
+		if (!IsRunningCommandlet())
+		{
+			FPaperContentBrowserExtensions::InstallHooks();
+			FPaperLevelEditorMenuExtensions::InstallHooks();
+		}
 		// Register with the mesh paint module
 		if (IMeshPaintModule* MeshPaintModule = FModuleManager::LoadModulePtr<IMeshPaintModule>("MeshPaint"))
 		{
