minor cleanup

diff --git a/Engine/Plugins/2D/Paper2D/Source/Paper2D/Private/Rendering/PaperBatchSceneProxy.cpp b/Engine/Plugins/2D/Paper2D/Source/Paper2D/Private/Rendering/PaperBatchSceneProxy.cpp
--- a/Engine/Plugins/2D/Paper2D/Source/Paper2D/Private/Rendering/PaperBatchSceneProxy.cpp
+++ b/Engine/Plugins/2D/Paper2D/Source/Paper2D/Private/Rendering/PaperBatchSceneProxy.cpp
@@ -27,7 +27,6 @@ FPrimitiveViewRelevance FPaperBatchSceneProxy::GetViewRelevance(const FSceneView
 	Result.bRenderInMainPass = ShouldRenderInMainPass();
 
 	Result.bMaskedRelevance = true;
-	//Result.bNormalTranslucencyRelevance = true;
 	Result.bDynamicRelevance = true;
 	Result.bOpaqueRelevance = true;
 
