Engine - Reducing code copy/paste.  Now unifying all checkboard texture creation under one function in FImageUtils.

diff --git a/Engine/Plugins/2D/Paper2D/Source/Paper2DEditor/Private/PaperEditorViewportClient.cpp b/Engine/Plugins/2D/Paper2D/Source/Paper2DEditor/Private/PaperEditorViewportClient.cpp
--- a/Engine/Plugins/2D/Paper2D/Source/Paper2DEditor/Private/PaperEditorViewportClient.cpp
+++ b/Engine/Plugins/2D/Paper2D/Source/Paper2DEditor/Private/PaperEditorViewportClient.cpp
@@ -6,6 +6,7 @@
 #include "CanvasTypes.h"
 #include "SEditorViewport.h"
 #include "PreviewScene.h"
+#include "ImageUtils.h"
 
 //////////////////////////////////////////////////////////////////////////
 // FPaperEditorViewportClient
@@ -52,7 +53,7 @@ FPaperEditorViewportClient::FPaperEditorViewportClient(const TWeakPtr<SEditorVie
 	bDeferZoomToSpriteIsInstant = true;
 
 	// Get the correct general direction of the perspective mode; the distance doesn't matter much as we've queued up a deferred zoom that will calculate a much better distance
- 	SetInitialViewTransform(LVT_Perspective, -100.0f * PaperAxisZ, PaperAxisZ.Rotation(), 0.0f);
+	SetInitialViewTransform(LVT_Perspective, -100.0f * PaperAxisZ, PaperAxisZ.Rotation(), 0.0f);
 }
 
 FPaperEditorViewportClient::~FPaperEditorViewportClient()
@@ -128,46 +129,9 @@ void FPaperEditorViewportClient::ModifyCheckerboardTextureColors()
 
 void FPaperEditorViewportClient::SetupCheckerboardTexture(const FColor& ColorOne, const FColor& ColorTwo, int32 CheckerSize)
 {
-// 	GConfig->GetInt(TEXT("TextureProperties"), TEXT("CheckerboardCheckerPixelNum"), CheckerSize, GEditorPerProjectIni);
-// 	GConfig->GetColor(TEXT("TextureProperties"), TEXT("CheckerboardColorOne"), CheckerColorOne, GEditorPerProjectIni);
-// 	GConfig->GetColor(TEXT("TextureProperties"), TEXT("CheckerboardColorTwo"), CheckerColorTwo, GEditorPerProjectIni);
-	
-// 	CheckerColorOne = FColor(128, 128, 128);
-// 	CheckerColorTwo = FColor(64, 64, 64);
-// 	bIsCheckeredBackgroundFill = false;
-// 	CheckerSize = 32;
-	CheckerSize = FMath::RoundUpToPowerOfTwo(CheckerSize);
-	const int32 HalfPixelNum = CheckerSize >> 1;
-
 	if (CheckerboardTexture == nullptr)
 	{
-		// Create the texture
-		CheckerboardTexture = UTexture2D::CreateTransient(CheckerSize, CheckerSize, PF_B8G8R8A8);
-
-		// Lock the checkerboard texture so it can be modified
-		FColor* MipData = static_cast<FColor*>(CheckerboardTexture->PlatformData->Mips[0].BulkData.Lock(LOCK_READ_WRITE));
-
-		// Fill in the colors in a checkerboard pattern
-		for (int32 RowNum = 0; RowNum < CheckerSize; ++RowNum)
-		{
-			for (int32 ColNum = 0; ColNum < CheckerSize; ++ColNum)
-			{
-				FColor& CurColor = MipData[(ColNum + (RowNum * CheckerSize))];
-
-				if (ColNum < HalfPixelNum)
-				{
-					CurColor = (RowNum < HalfPixelNum) ? ColorOne: ColorTwo;
-				}
-				else
-				{
-					CurColor = (RowNum < HalfPixelNum) ? ColorTwo: ColorOne;
-				}
-			}
-		}
-
-		// Unlock the texture
-		CheckerboardTexture->PlatformData->Mips[0].BulkData.Unlock();
-		CheckerboardTexture->UpdateResource();
+		CheckerboardTexture = FImageUtils::CreateCheckerboardTexture(ColorOne, ColorTwo, CheckerSize);
 	}
 }
 
diff --git a/Engine/Source/Editor/TextureEditor/Private/Models/TextureEditorViewportClient.cpp b/Engine/Source/Editor/TextureEditor/Private/Models/TextureEditorViewportClient.cpp
--- a/Engine/Source/Editor/TextureEditor/Private/Models/TextureEditorViewportClient.cpp
+++ b/Engine/Source/Editor/TextureEditor/Private/Models/TextureEditorViewportClient.cpp
@@ -6,6 +6,7 @@
 #include "Engine/TextureCube.h"
 #include "Engine/TextureRenderTarget2D.h"
 #include "Engine/TextureRenderTargetCube.h"
+#include "ImageUtils.h"
 
 
 /* FTextureEditorViewportClient structors
@@ -161,36 +162,10 @@ void FTextureEditorViewportClient::AddReferencedObjects(FReferenceCollector& Col
 
 void FTextureEditorViewportClient::ModifyCheckerboardTextureColors()
 {
-	const UTextureEditorSettings& Settings = *GetDefault<UTextureEditorSettings>();
-	const int32 HalfPixelNum = Settings.CheckerSize >> 1;
-
 	DestroyCheckerboardTexture();
-	SetupCheckerboardTexture();
-
-	// Lock the checkerboard texture so it can be modified
-	FColor* MipData = static_cast<FColor*>(CheckerboardTexture->PlatformData->Mips[0].BulkData.Lock(LOCK_READ_WRITE));
-
-	// Fill in the colors in a checkerboard pattern
-	for (int32 RowNum = 0; RowNum < Settings.CheckerSize; ++RowNum)
-	{
-		for (int32 ColNum = 0; ColNum < Settings.CheckerSize; ++ColNum)
-		{
-			FColor& CurColor = MipData[(ColNum + (RowNum * Settings.CheckerSize))];
-
-			if (ColNum < HalfPixelNum)
-			{
-				CurColor = (RowNum < HalfPixelNum)? Settings.CheckerColorOne: Settings.CheckerColorTwo;
-			}
-			else
-			{
-				CurColor = (RowNum < HalfPixelNum)? Settings.CheckerColorTwo: Settings.CheckerColorOne;
-			}
-		}
-	}
 
-	// Unlock the texture
-	CheckerboardTexture->PlatformData->Mips[0].BulkData.Unlock();
-	CheckerboardTexture->UpdateResource();
+	const UTextureEditorSettings& Settings = *GetDefault<UTextureEditorSettings>();
+	CheckerboardTexture = FImageUtils::CreateCheckerboardTexture(Settings.CheckerColorOne, Settings.CheckerColorTwo, Settings.CheckerSize);
 }
 
 
@@ -311,17 +286,6 @@ FVector2D FTextureEditorViewportClient::GetViewportScrollBarPositions() const
 	return Positions;
 }
 
-
-void FTextureEditorViewportClient::SetupCheckerboardTexture()
-{
-	if (!CheckerboardTexture)
-	{
-		const UTextureEditorSettings& Settings = *GetDefault<UTextureEditorSettings>();
-		CheckerboardTexture = UTexture2D::CreateTransient(Settings.CheckerSize, Settings.CheckerSize, PF_B8G8R8A8);
-	}
-}
-
-
 void FTextureEditorViewportClient::DestroyCheckerboardTexture()
 {
 	if (CheckerboardTexture)
diff --git a/Engine/Source/Editor/TextureEditor/Private/Models/TextureEditorViewportClient.h b/Engine/Source/Editor/TextureEditor/Private/Models/TextureEditorViewportClient.h
--- a/Engine/Source/Editor/TextureEditor/Private/Models/TextureEditorViewportClient.h
+++ b/Engine/Source/Editor/TextureEditor/Private/Models/TextureEditorViewportClient.h
@@ -37,9 +37,6 @@ private:
 	/** Returns the positions of the scrollbars relative to the Texture textures */
 	FVector2D GetViewportScrollBarPositions() const;
 
-	/** Initialize the checkerboard texture for the texture preview, if necessary */
-	void SetupCheckerboardTexture();
-
 	/** Destroy the checkerboard texture if one exists */
 	void DestroyCheckerboardTexture();
 
diff --git a/Engine/Source/Editor/UnrealEd/Private/ThumbnailManager.cpp b/Engine/Source/Editor/UnrealEd/Private/ThumbnailManager.cpp
--- a/Engine/Source/Editor/UnrealEd/Private/ThumbnailManager.cpp
+++ b/Engine/Source/Editor/UnrealEd/Private/ThumbnailManager.cpp
@@ -252,38 +252,7 @@ void UThumbnailManager::SetupCheckerboardTexture()
 		return;
 	}
 
-	const FColor ColorOne = FColor(128, 128, 128);
-	const FColor ColorTwo = FColor(64, 64, 64);
-	const int32 CheckerSize = 32;
-	const int32 HalfPixelNum = CheckerSize >> 1;
-
-	// Create the texture
-	CheckerboardTexture = UTexture2D::CreateTransient(CheckerSize, CheckerSize, PF_B8G8R8A8);
-
-	// Lock the checkerboard texture so it can be modified
-	FColor* MipData = static_cast<FColor*>(CheckerboardTexture->PlatformData->Mips[0].BulkData.Lock(LOCK_READ_WRITE));
-
-	// Fill in the colors in a checkerboard pattern
-	for (int32 RowNum = 0; RowNum < CheckerSize; ++RowNum)
-	{
-		for (int32 ColNum = 0; ColNum < CheckerSize; ++ColNum)
-		{
-			FColor& CurColor = MipData[(ColNum + (RowNum * CheckerSize))];
-
-			if (ColNum < HalfPixelNum)
-			{
-				CurColor = (RowNum < HalfPixelNum)? ColorOne: ColorTwo;
-			}
-			else
-			{
-				CurColor = (RowNum < HalfPixelNum)? ColorTwo: ColorOne;
-			}
-		}
-	}
-
-	// Unlock the texture
-	CheckerboardTexture->PlatformData->Mips[0].BulkData.Unlock();
-	CheckerboardTexture->UpdateResource();
+	CheckerboardTexture = FImageUtils::CreateCheckerboardTexture(FColor(128, 128, 128), FColor(64, 64, 64), 32);
 }
 
 bool UThumbnailManager::CaptureProjectThumbnail(FViewport* Viewport, const FString& OutputFilename, bool bUseSCCIfPossible)
diff --git a/Engine/Source/Runtime/Engine/Private/ImageUtils.cpp b/Engine/Source/Runtime/Engine/Private/ImageUtils.cpp
--- a/Engine/Source/Runtime/Engine/Private/ImageUtils.cpp
+++ b/Engine/Source/Runtime/Engine/Private/ImageUtils.cpp
@@ -255,4 +255,40 @@ void FImageUtils::CompressImageArray( int32 ImageWidth, int32 ImageHeight, TArra
 	TempThumbnail.CompressImageData();
 	TArray<uint8>& CompressedByteArray = TempThumbnail.AccessCompressedImageData();
 	DstData = TempThumbnail.AccessCompressedImageData();
-}
\ No newline at end of file
+}
+
+UTexture2D* FImageUtils::CreateCheckerboardTexture(FColor ColorOne, FColor ColorTwo, int32 CheckerSize)
+{
+	CheckerSize = FMath::RoundUpToPowerOfTwo(CheckerSize);
+	const int32 HalfPixelNum = CheckerSize >> 1;
+
+	// Create the texture
+	UTexture2D* CheckerboardTexture = UTexture2D::CreateTransient(CheckerSize, CheckerSize, PF_B8G8R8A8);
+
+	// Lock the checkerboard texture so it can be modified
+	FColor* MipData = static_cast<FColor*>( CheckerboardTexture->PlatformData->Mips[0].BulkData.Lock(LOCK_READ_WRITE) );
+
+	// Fill in the colors in a checkerboard pattern
+	for ( int32 RowNum = 0; RowNum < CheckerSize; ++RowNum )
+	{
+		for ( int32 ColNum = 0; ColNum < CheckerSize; ++ColNum )
+		{
+			FColor& CurColor = MipData[( ColNum + ( RowNum * CheckerSize ) )];
+
+			if ( ColNum < HalfPixelNum )
+			{
+				CurColor = ( RowNum < HalfPixelNum ) ? ColorOne : ColorTwo;
+			}
+			else
+			{
+				CurColor = ( RowNum < HalfPixelNum ) ? ColorTwo : ColorOne;
+			}
+		}
+	}
+
+	// Unlock the texture
+	CheckerboardTexture->PlatformData->Mips[0].BulkData.Unlock();
+	CheckerboardTexture->UpdateResource();
+
+	return CheckerboardTexture;
+}
diff --git a/Engine/Source/Runtime/Engine/Public/ImageUtils.h b/Engine/Source/Runtime/Engine/Public/ImageUtils.h
--- a/Engine/Source/Runtime/Engine/Public/ImageUtils.h
+++ b/Engine/Source/Runtime/Engine/Public/ImageUtils.h
@@ -4,8 +4,7 @@
 ImageUtils.h: Image utility functions.
 =============================================================================*/
 
-#ifndef __IMAGEUTILS_H__
-#define __IMAGEUTILS_H__
+#pragma once
 
 /**
  *	Parameters used for creating a Texture2D frmo a simple color buffer.
@@ -92,7 +91,14 @@ public:
 	 *
 	 */
 	ENGINE_API static void CompressImageArray( int32 ImageWidth, int32 ImageHeight, TArray<FColor> &SrcData, TArray<uint8> &DstData );
-};
-
-#endif
 
+	/**
+	 * Creates a new UTexture2D with a checkerboard pattern.
+	 *
+	 * @param ColorOne		The color of half of the squares.
+	 * @param ColorTwo		The color of the other half of the squares.
+	 * @param CheckerSize	The size in pixels of each checker square.
+	 *
+	 */
+	ENGINE_API static UTexture2D* CreateCheckerboardTexture(FColor ColorOne = FColor(64, 64, 64), FColor ColorTwo = FColor(128, 128, 128), int32 CheckerSize = 32);
+};
