Tile Maps: Fixed a crash when undoing a resize [UE-16060]

diff --git a/Engine/Plugins/2D/Paper2D/Source/Paper2D/Private/PaperTileMap.cpp b/Engine/Plugins/2D/Paper2D/Source/Paper2D/Private/PaperTileMap.cpp
--- a/Engine/Plugins/2D/Paper2D/Source/Paper2D/Private/PaperTileMap.cpp
+++ b/Engine/Plugins/2D/Paper2D/Source/Paper2D/Private/PaperTileMap.cpp
@@ -121,6 +121,18 @@ void UPaperTileMap::PostEditChangeProperty(FPropertyChangedEvent& PropertyChange
 	{
 		ResizeMap(MapWidth, MapHeight, /*bForceResize=*/ true);
 	}
+	else
+	{
+		// Make sure that the layers are all of the right size
+		for (UPaperTileLayer* TileLayer : TileLayers)
+		{
+			if ((TileLayer->GetLayerWidth() != MapWidth) || (TileLayer->GetLayerHeight() != MapHeight))
+			{
+				TileLayer->Modify();
+				TileLayer->ResizeMap(MapWidth, MapHeight);
+			}
+		}
+	}
 
 	if (!IsTemplate())
 	{
diff --git a/Engine/Plugins/2D/Paper2D/Source/Paper2D/Private/PaperTileMapComponent.cpp b/Engine/Plugins/2D/Paper2D/Source/Paper2D/Private/PaperTileMapComponent.cpp
--- a/Engine/Plugins/2D/Paper2D/Source/Paper2D/Private/PaperTileMapComponent.cpp
+++ b/Engine/Plugins/2D/Paper2D/Source/Paper2D/Private/PaperTileMapComponent.cpp
@@ -309,6 +309,9 @@ void UPaperTileMapComponent::RebuildRenderData(TArray<FSpriteRenderSection>& Sec
 		FSpriteRenderSection* CurrentBatch = nullptr;
 		FVector CurrentDestinationOrigin;
 		const FPaperTileInfo* CurrentCellPtr = Layer->PRIVATE_GetAllocatedCells();
+		check(Layer->GetLayerWidth() == TileMap->MapWidth);
+		check(Layer->GetLayerHeight() == TileMap->MapHeight);
+
 
 		for (int32 Y = 0; Y < TileMap->MapHeight; ++Y)
 		{
