Tile Maps: Fixed an issue where any edits to the UPaperTileMap could not be undone

diff --git a/Engine/Plugins/2D/Paper2D/Source/Paper2D/Private/PaperTileMap.cpp b/Engine/Plugins/2D/Paper2D/Source/Paper2D/Private/PaperTileMap.cpp
--- a/Engine/Plugins/2D/Paper2D/Source/Paper2D/Private/PaperTileMap.cpp
+++ b/Engine/Plugins/2D/Paper2D/Source/Paper2D/Private/PaperTileMap.cpp
@@ -45,6 +45,8 @@ void UPaperTileMap::PreEditChange(UProperty* PropertyAboutToChange)
 		// Subtract out the hex side length; we'll add it back (along with any changes) in PostEditChangeProperty
 		TileHeight -= HexSideLength;
 	}
+
+	Super::PreEditChange(PropertyAboutToChange);
 }
 #endif
 
