Paper2D: Fixed undo on a merge sprites action not deleting the merged sprite (transaction was too tightly scoped) [UE-16017]

diff --git a/Engine/Plugins/2D/Paper2D/Source/Paper2DEditor/Private/GroupedSprites/PaperGroupedSpriteUtilities.cpp b/Engine/Plugins/2D/Paper2D/Source/Paper2DEditor/Private/GroupedSprites/PaperGroupedSpriteUtilities.cpp
--- a/Engine/Plugins/2D/Paper2D/Source/Paper2DEditor/Private/GroupedSprites/PaperGroupedSpriteUtilities.cpp
+++ b/Engine/Plugins/2D/Paper2D/Source/Paper2DEditor/Private/GroupedSprites/PaperGroupedSpriteUtilities.cpp
@@ -177,13 +177,15 @@ void FPaperGroupedSpriteUtilities::MergeSprites(const TArray<UObject*>& InObject
 			FActorSpawnParameters SpawnParams;
 			SpawnParams.bDeferConstruction = true;
 
+			const FScopedTransaction Transaction(LOCTEXT("MergeSprites", "Merge sprite instances"));
+
 			if (APaperGroupedSpriteActor* SpawnedActor = World->SpawnActor<APaperGroupedSpriteActor>(SpawnParams))
 			{
-				const FScopedTransaction Transaction(LOCTEXT("MergeSprites", "Merge sprite instances"));
-
+				SpawnedActor->Modify();
 				// Create a merged sprite component
 				{
 					UPaperGroupedSpriteComponent* MergedSpriteComponent = SpawnedActor->GetRenderComponent();
+					MergedSpriteComponent->Modify();
 					FComponentReregisterContext ReregisterContext(MergedSpriteComponent);
 
 					// Create an instance from each sprite component that we're harvesting
