Paper2D: Fixed problems with tangent generation in sprites and tile maps [UE-16082]

diff --git a/Engine/Plugins/2D/Paper2D/Source/Paper2D/Private/PaperCustomVersion.h b/Engine/Plugins/2D/Paper2D/Source/Paper2D/Private/PaperCustomVersion.h
--- a/Engine/Plugins/2D/Paper2D/Source/Paper2D/Private/PaperCustomVersion.h
+++ b/Engine/Plugins/2D/Paper2D/Source/Paper2D/Private/PaperCustomVersion.h
@@ -55,6 +55,9 @@ struct FPaperCustomVersion
 		// Change Tile Set margin and padding to allow non-uniform values (also converts TileWidth/TileHeight into TileSize and makes members private)
 		AllowNonUniformPaddingInTileSets = 15,
 
+		// Fixed the normals for sprites that aren't two sided
+		FixNormalsForFrontFace = 16,
+
 		// -----<new versions can be added above this line>-------------------------------------------------
 		VersionPlusOne,
 		LatestVersion = VersionPlusOne - 1
diff --git a/Engine/Plugins/2D/Paper2D/Source/Paper2D/Private/PaperRenderSceneProxy.cpp b/Engine/Plugins/2D/Paper2D/Source/Paper2D/Private/PaperRenderSceneProxy.cpp
--- a/Engine/Plugins/2D/Paper2D/Source/Paper2D/Private/PaperRenderSceneProxy.cpp
+++ b/Engine/Plugins/2D/Paper2D/Source/Paper2D/Private/PaperRenderSceneProxy.cpp
@@ -13,6 +13,8 @@
 #include "Engine/Engine.h"
 #include "SpriteDrawCall.h"
 
+static TAutoConsoleVariable<int32> CVarDrawSpritesAsTwoSided(TEXT("r.Paper2D.DrawTwoSided"), 1, TEXT("Draw sprites as two sided."));
+
 DECLARE_CYCLE_STAT(TEXT("Get Batch Mesh"), STAT_PaperRender_GetBatchMesh, STATGROUP_Paper2D);
 DECLARE_CYCLE_STAT(TEXT("Get New Batch Meshes"), STAT_PaperRender_GetNewBatchMeshes, STATGROUP_Paper2D);
 DECLARE_CYCLE_STAT(TEXT("Convert Batches"), STAT_PaperRender_ConvertBatches, STATGROUP_Paper2D);
@@ -27,7 +29,7 @@ FPackedNormal FPaperSpriteVertex::PackedNormalZ(FVector(0.0f, -1.0f, 0.0f));
 void FPaperSpriteVertex::SetTangentsFromPaperAxes()
 {
 	PackedNormalX = PaperAxisX;
-	PackedNormalZ = PaperAxisZ;
+	PackedNormalZ = -PaperAxisZ;
 	// store determinant of basis in w component of normal vector
 	PackedNormalZ.Vector.W = (GetBasisDeterminantSign(PaperAxisX, PaperAxisY, PaperAxisZ) < 0.0f) ? 0 : 255;
 }
@@ -269,6 +271,8 @@ FPaperRenderSceneProxy::FPaperRenderSceneProxy(const UPrimitiveComponent* InComp
 	, CollisionResponse(InComponent->GetCollisionResponseToChannels())
 {
 	WireframeColor = FLinearColor::White;
+
+	bDrawTwoSided = CVarDrawSpritesAsTwoSided.GetValueOnGameThread() != 0;
 }
 
 FPaperRenderSceneProxy::~FPaperRenderSceneProxy()
@@ -467,7 +471,7 @@ void FPaperRenderSceneProxy::GetNewBatchMeshes(const FSceneView* View, int32 Vie
 			Mesh.CastShadow = bCastShadow;
 			Mesh.DepthPriorityGroup = DPG;
 			Mesh.Type = PT_TriangleList;
-			Mesh.bDisableBackfaceCulling = true;
+			Mesh.bDisableBackfaceCulling = bDrawTwoSided;
 			Mesh.MaterialRenderProxy = TextureOverrideMaterialProxy;
 
 			// Set up the FMeshBatchElement.
@@ -561,7 +565,7 @@ void FPaperRenderSceneProxy::GetBatchMesh(const FSceneView* View, class UMateria
 			Mesh.CastShadow = bCastShadow;
 			Mesh.DepthPriorityGroup = DPG;
 			Mesh.Type = PT_TriangleList;
-			Mesh.bDisableBackfaceCulling = true;
+			Mesh.bDisableBackfaceCulling = bDrawTwoSided;
 
 			// Set up the FMeshBatchElement.
 			FMeshBatchElement& BatchElement = Mesh.Elements[0];
diff --git a/Engine/Plugins/2D/Paper2D/Source/Paper2D/Private/PaperRenderSceneProxy.h b/Engine/Plugins/2D/Paper2D/Source/Paper2D/Private/PaperRenderSceneProxy.h
--- a/Engine/Plugins/2D/Paper2D/Source/Paper2D/Private/PaperRenderSceneProxy.h
+++ b/Engine/Plugins/2D/Paper2D/Source/Paper2D/Private/PaperRenderSceneProxy.h
@@ -190,6 +190,7 @@ protected:
 	AActor* Owner;
 	UBodySetup* MyBodySetup;
 
+	bool bDrawTwoSided;
 	bool bCastShadow;
 
 	// The view relevance for the associated material
diff --git a/Engine/Plugins/2D/Paper2D/Source/Paper2D/Private/PaperSprite.cpp b/Engine/Plugins/2D/Paper2D/Source/Paper2D/Private/PaperSprite.cpp
--- a/Engine/Plugins/2D/Paper2D/Source/Paper2D/Private/PaperSprite.cpp
+++ b/Engine/Plugins/2D/Paper2D/Source/Paper2D/Private/PaperSprite.cpp
@@ -968,7 +968,14 @@ void UPaperSprite::RebuildRenderData()
 
 		new (BakedRenderData) FVector4(PivotSpacePos.X * UnitsPerPixel, PivotSpacePos.Y * UnitsPerPixel, UV.X * InverseWidth, UV.Y * InverseHeight);
 	}
+
 	check((BakedRenderData.Num() % 3) == 0);
+
+	// Swap the generated vertices so they end up in counterclockwise order
+	for (int32 SVT = 0; SVT < TriangluatedPoints.Num(); SVT += 3)
+	{
+		Swap(BakedRenderData[SVT + 2], BakedRenderData[SVT + 0]);
+	}
 }
 
 void UPaperSprite::FindTextureBoundingBox(float AlphaThreshold, /*out*/ FVector2D& OutBoxPosition, /*out*/ FVector2D& OutBoxSize)
@@ -1778,6 +1785,11 @@ void UPaperSprite::PostLoad()
 		bSnapPivotToPixelGrid = false;
 	}
 
+	if (PaperVer < FPaperCustomVersion::FixTangentGenerationForFrontFace)
+	{
+		bRebuildRenderData = true;
+	}
+
 	if (PaperVer < FPaperCustomVersion::AddPixelsPerUnrealUnit)
 	{
 		PixelsPerUnrealUnit = 1.0f;
diff --git a/Engine/Plugins/2D/Paper2D/Source/Paper2D/Private/PaperTileMapComponent.cpp b/Engine/Plugins/2D/Paper2D/Source/Paper2D/Private/PaperTileMapComponent.cpp
--- a/Engine/Plugins/2D/Paper2D/Source/Paper2D/Private/PaperTileMapComponent.cpp
+++ b/Engine/Plugins/2D/Paper2D/Source/Paper2D/Private/PaperTileMapComponent.cpp
@@ -427,13 +427,13 @@ void UPaperTileMapComponent::RebuildRenderData(TArray<FSpriteRenderSection>& Sec
 					const FVector4 TopRight(WX0 + TileSizeWithFlip.X, WY0, UValues[UVIndex2], VValues[UVIndex2]);
 					const FVector4 TopLeft(WX0, WY0, UValues[UVIndex3], VValues[UVIndex3]);
 
-					CurrentBatch->AddVertex(BottomLeft.X, BottomLeft.Y, BottomLeft.Z, BottomLeft.W, CurrentDestinationOrigin, DrawColor, Vertices);
-					CurrentBatch->AddVertex(TopRight.X, TopRight.Y, TopRight.Z, TopRight.W, CurrentDestinationOrigin, DrawColor, Vertices);
 					CurrentBatch->AddVertex(BottomRight.X, BottomRight.Y, BottomRight.Z, BottomRight.W, CurrentDestinationOrigin, DrawColor, Vertices);
-
+					CurrentBatch->AddVertex(TopRight.X, TopRight.Y, TopRight.Z, TopRight.W, CurrentDestinationOrigin, DrawColor, Vertices);
 					CurrentBatch->AddVertex(BottomLeft.X, BottomLeft.Y, BottomLeft.Z, BottomLeft.W, CurrentDestinationOrigin, DrawColor, Vertices);
-					CurrentBatch->AddVertex(TopLeft.X, TopLeft.Y, TopLeft.Z, TopLeft.W, CurrentDestinationOrigin, DrawColor, Vertices);
+
 					CurrentBatch->AddVertex(TopRight.X, TopRight.Y, TopRight.Z, TopRight.W, CurrentDestinationOrigin, DrawColor, Vertices);
+					CurrentBatch->AddVertex(TopLeft.X, TopLeft.Y, TopLeft.Z, TopLeft.W, CurrentDestinationOrigin, DrawColor, Vertices);
+					CurrentBatch->AddVertex(BottomLeft.X, BottomLeft.Y, BottomLeft.Z, BottomLeft.W, CurrentDestinationOrigin, DrawColor, Vertices);
 				}
 			}
 		}
