Paper2D: Checking in missing file

diff --git a/Engine/Plugins/2D/Paper2D/Source/Paper2D/Private/PaperCustomVersion.h b/Engine/Plugins/2D/Paper2D/Source/Paper2D/Private/PaperCustomVersion.h
--- a/Engine/Plugins/2D/Paper2D/Source/Paper2D/Private/PaperCustomVersion.h
+++ b/Engine/Plugins/2D/Paper2D/Source/Paper2D/Private/PaperCustomVersion.h
@@ -55,8 +55,8 @@ struct FPaperCustomVersion
 		// Change Tile Set margin and padding to allow non-uniform values (also converts TileWidth/TileHeight into TileSize and makes members private)
 		AllowNonUniformPaddingInTileSets = 15,
 
-		// Fixed the normals for sprites that aren't two sided
-		FixNormalsForFrontFace = 16,
+		// Fixed the tangent generation for sprites, tile maps, etc... being incorrect for the front face
+		FixTangentGenerationForFrontFace = 16,
 
 		// -----<new versions can be added above this line>-------------------------------------------------
 		VersionPlusOne,
