Paper2D: Added bSilent param to various methods in the spritesheet importer to allow it to be used for a FactoryCanImport method without generating log warnings

diff --git a/Engine/Plugins/2D/Paper2D/Source/PaperSpriteSheetImporter/Private/PaperJsonSpriteSheetImporter.cpp b/Engine/Plugins/2D/Paper2D/Source/PaperSpriteSheetImporter/Private/PaperJsonSpriteSheetImporter.cpp
--- a/Engine/Plugins/2D/Paper2D/Source/PaperSpriteSheetImporter/Private/PaperJsonSpriteSheetImporter.cpp
+++ b/Engine/Plugins/2D/Paper2D/Source/PaperSpriteSheetImporter/Private/PaperJsonSpriteSheetImporter.cpp
@@ -12,7 +12,7 @@
 #include "PaperSpriteSheet.h"
 #include "PaperImporterSettings.h"
 
-TSharedPtr<FJsonObject> ParseJSON(const FString& FileContents, const FString& NameForErrors)
+TSharedPtr<FJsonObject> ParseJSON(const FString& FileContents, const FString& NameForErrors, bool bSilent)
 {
 	// Load the file up (JSON format)
 	if (!FileContents.IsEmpty())
@@ -27,18 +27,24 @@ TSharedPtr<FJsonObject> ParseJSON(const FString& FileContents, const FString& Na
 		}
 		else
 		{
-			UE_LOG(LogPaperSpriteSheetImporter, Warning, TEXT("Failed to parse sprite descriptor file '%s'.  Error: '%s'"), *NameForErrors, *Reader->GetErrorMessage());
+			if (!bSilent)
+			{
+				UE_LOG(LogPaperSpriteSheetImporter, Warning, TEXT("Failed to parse sprite descriptor file '%s'.  Error: '%s'"), *NameForErrors, *Reader->GetErrorMessage());
+			}
 			return nullptr;
 		}
 	}
 	else
 	{
-		UE_LOG(LogPaperSpriteSheetImporter, Warning, TEXT("Sprite descriptor file '%s' was empty.  This sprite cannot be imported."), *NameForErrors);
+		if (!bSilent)
+		{
+			UE_LOG(LogPaperSpriteSheetImporter, Warning, TEXT("Sprite descriptor file '%s' was empty.  This sprite cannot be imported."), *NameForErrors);
+		}
 		return nullptr;
 	}
 }
 
-TSharedPtr<FJsonObject> ParseJSON(FArchive* const Stream, const FString& NameForErrors)
+TSharedPtr<FJsonObject> ParseJSON(FArchive* const Stream, const FString& NameForErrors, bool bSilent)
 {
 	const TSharedRef< TJsonReader<> >& Reader = TJsonReaderFactory<>::Create(Stream);
 
@@ -50,12 +56,15 @@ TSharedPtr<FJsonObject> ParseJSON(FArchive* const Stream, const FString& NameFor
 	}
 	else
 	{
-		UE_LOG(LogPaperSpriteSheetImporter, Warning, TEXT("Failed to parse sprite descriptor file '%s'.  Error: '%s'"), *NameForErrors, *Reader->GetErrorMessage());
+		if (!bSilent)
+		{
+			UE_LOG(LogPaperSpriteSheetImporter, Warning, TEXT("Failed to parse sprite descriptor file '%s'.  Error: '%s'"), *NameForErrors, *Reader->GetErrorMessage());
+		}
 		return nullptr;
 	}
 }
 
-bool ParseMetaBlock(const FString& NameForErrors, TSharedPtr<FJsonObject>& SpriteDescriptorObject, FString& OutImage)
+bool ParseMetaBlock(const FString& NameForErrors, TSharedPtr<FJsonObject>& SpriteDescriptorObject, FString& OutImage, bool bSilent)
 {
 	bool bLoadedSuccessfully = true;
 	TSharedPtr<FJsonObject> MetaBlock = FPaperJSONHelpers::ReadObject(SpriteDescriptorObject, TEXT("meta"));
@@ -78,35 +87,50 @@ bool ParseMetaBlock(const FString& NameForErrors, TSharedPtr<FJsonObject>& Sprit
 		if (AppName.StartsWith(FlashPrefix) || AppName.StartsWith(TexturePackerPrefix))
 		{
 			// Cool, we (mostly) know how to handle these sorts of files!
-			UE_LOG(LogPaperSpriteSheetImporter, Log, TEXT("Parsing sprite sheet exported from '%s'"), *AppName);
+			if (!bSilent)
+			{
+				UE_LOG(LogPaperSpriteSheetImporter, Log, TEXT("Parsing sprite sheet exported from '%s'"), *AppName);
+			}
 		}
 		else if (!AppName.IsEmpty())
 		{
 			// It's got an app tag inside a meta block, so we'll take a crack at it
-			UE_LOG(LogPaperSpriteSheetImporter, Warning, TEXT("Unexpected 'app' named '%s' while parsing sprite descriptor file '%s'.  Parsing will continue but the format may not be fully supported"), *AppName, *NameForErrors);
+			if (!bSilent)
+			{
+				UE_LOG(LogPaperSpriteSheetImporter, Warning, TEXT("Unexpected 'app' named '%s' while parsing sprite descriptor file '%s'.  Parsing will continue but the format may not be fully supported"), *AppName, *NameForErrors);
+			}
 		}
 		else
 		{
 			// Probably not a sprite sheet
-			UE_LOG(LogPaperSpriteSheetImporter, Warning, TEXT("Failed to parse sprite descriptor file '%s'.  Expected 'app' key indicating the exporter (might not be a sprite sheet)"), *NameForErrors);
+			if (!bSilent)
+			{
+				UE_LOG(LogPaperSpriteSheetImporter, Warning, TEXT("Failed to parse sprite descriptor file '%s'.  Expected 'app' key indicating the exporter (might not be a sprite sheet)"), *NameForErrors);
+			}
 			bLoadedSuccessfully = false;
 		}
 
 		if (OutImage.IsEmpty())
 		{
-			UE_LOG(LogPaperSpriteSheetImporter, Warning, TEXT("Failed to parse sprite descriptor file '%s'.  Expected valid 'image' tag"), *NameForErrors);
+			if (!bSilent)
+			{
+				UE_LOG(LogPaperSpriteSheetImporter, Warning, TEXT("Failed to parse sprite descriptor file '%s'.  Expected valid 'image' tag"), *NameForErrors);
+			}
 			bLoadedSuccessfully = false;
 		}
 	}
 	else
 	{
-		UE_LOG(LogPaperSpriteSheetImporter, Warning, TEXT("Failed to parse sprite descriptor file '%s'.  Missing meta block"), *NameForErrors);
+		if (!bSilent)
+		{
+			UE_LOG(LogPaperSpriteSheetImporter, Warning, TEXT("Failed to parse sprite descriptor file '%s'.  Missing meta block"), *NameForErrors);
+		}
 		bLoadedSuccessfully = false;
 	}
 	return bLoadedSuccessfully;
 }
 
-static bool ParseFrame(TSharedPtr<FJsonObject> &FrameData, FSpriteFrame &OutFrame)
+static bool ParseFrame(TSharedPtr<FJsonObject>& FrameData, FSpriteFrame& OutFrame)
 {
 	bool bReadFrameSuccessfully = true;
 	// An example frame:
@@ -349,9 +373,9 @@ void FPaperJsonSpriteSheetImporter::SetReimportData(const TArray<FString>& Exist
 	bIsReimporting = true;
 }
 
-bool FPaperJsonSpriteSheetImporter::Import(TSharedPtr<FJsonObject> SpriteDescriptorObject, const FString& NameForErrors)
+bool FPaperJsonSpriteSheetImporter::Import(TSharedPtr<FJsonObject> SpriteDescriptorObject, const FString& NameForErrors, bool bSilent)
 {
-	bool bLoadedSuccessfully = ParseMetaBlock(NameForErrors, SpriteDescriptorObject, /*out*/ ImageName);
+	bool bLoadedSuccessfully = ParseMetaBlock(NameForErrors, SpriteDescriptorObject, /*out*/ ImageName, bSilent);
 	if (bLoadedSuccessfully)
 	{
 		TSharedPtr<FJsonObject> ObjectFrameBlock = FPaperJSONHelpers::ReadObject(SpriteDescriptorObject, TEXT("frames"));
@@ -371,35 +395,38 @@ bool FPaperJsonSpriteSheetImporter::Import(TSharedPtr<FJsonObject> SpriteDescrip
 			}
 			else
 			{
-				UE_LOG(LogPaperSpriteSheetImporter, Warning, TEXT("Failed to parse sprite descriptor file '%s'.  Missing frames block"), *NameForErrors);
+				if (!bSilent)
+				{
+					UE_LOG(LogPaperSpriteSheetImporter, Warning, TEXT("Failed to parse sprite descriptor file '%s'.  Missing frames block"), *NameForErrors);
+				}
 				bLoadedSuccessfully = false;
 			}
 		}
 
 		if (bLoadedSuccessfully && (Frames.Num() == 0))
 		{
-			UE_LOG(LogPaperSpriteSheetImporter, Warning, TEXT("Failed to parse sprite descriptor file '%s'.  No frames loaded"), *NameForErrors);
+			if (!bSilent)
+			{
+				UE_LOG(LogPaperSpriteSheetImporter, Warning, TEXT("Failed to parse sprite descriptor file '%s'.  No frames loaded"), *NameForErrors);
+			}
 			bLoadedSuccessfully = false;
 		}
 	}
 	return bLoadedSuccessfully;
 }
 
-bool FPaperJsonSpriteSheetImporter::ImportFromString(const FString& FileContents, const FString& NameForErrors)
+bool FPaperJsonSpriteSheetImporter::ImportFromString(const FString& FileContents, const FString& NameForErrors, bool bSilent)
 {
-	TSharedPtr<FJsonObject> SpriteDescriptorObject = ParseJSON(FileContents, NameForErrors);
-	return SpriteDescriptorObject.IsValid() &&
-		Import(SpriteDescriptorObject, NameForErrors);
+	TSharedPtr<FJsonObject> SpriteDescriptorObject = ParseJSON(FileContents, NameForErrors, bSilent);
+	return SpriteDescriptorObject.IsValid() && Import(SpriteDescriptorObject, NameForErrors, bSilent);
 }
 
-bool FPaperJsonSpriteSheetImporter::ImportFromArchive(FArchive* Archive, const FString& NameForErrors)
+bool FPaperJsonSpriteSheetImporter::ImportFromArchive(FArchive* Archive, const FString& NameForErrors, bool bSilent)
 {
-	TSharedPtr<FJsonObject> SpriteDescriptorObject = ParseJSON(Archive, NameForErrors);
-	return SpriteDescriptorObject.IsValid() &&
-		Import(SpriteDescriptorObject, NameForErrors);
+	TSharedPtr<FJsonObject> SpriteDescriptorObject = ParseJSON(Archive, NameForErrors, bSilent);
+	return SpriteDescriptorObject.IsValid() && Import(SpriteDescriptorObject, NameForErrors, bSilent);
 }
 
-
 bool FPaperJsonSpriteSheetImporter::ImportTextures(const FString& LongPackagePath, const FString& SourcePath)
 {
 	bool bLoadedSuccessfully = true;
diff --git a/Engine/Plugins/2D/Paper2D/Source/PaperSpriteSheetImporter/Private/PaperJsonSpriteSheetImporter.h b/Engine/Plugins/2D/Paper2D/Source/PaperSpriteSheetImporter/Private/PaperJsonSpriteSheetImporter.h
--- a/Engine/Plugins/2D/Paper2D/Source/PaperSpriteSheetImporter/Private/PaperJsonSpriteSheetImporter.h
+++ b/Engine/Plugins/2D/Paper2D/Source/PaperSpriteSheetImporter/Private/PaperJsonSpriteSheetImporter.h
@@ -36,8 +36,8 @@ class FPaperJsonSpriteSheetImporter
 {
 public:
 	FPaperJsonSpriteSheetImporter();
-	bool ImportFromString(const FString& FileContents, const FString& NameForErrors);
-	bool ImportFromArchive(FArchive* Archive, const FString& NameForErrors);
+	bool ImportFromString(const FString& FileContents, const FString& NameForErrors, bool bSilent);
+	bool ImportFromArchive(FArchive* Archive, const FString& NameForErrors, bool bSilent);
 
 	bool ImportTextures(const FString& LongPackagePath, const FString& SourcePath);
 
@@ -49,7 +49,7 @@ public:
 	static UTexture2D* ImportTexture(const FString& TextureSourcePath, const FString& DestinationAssetFolder);
 
 protected:
-	bool Import(TSharedPtr<FJsonObject> SpriteDescriptorObject, const FString& NameForErrors);
+	bool Import(TSharedPtr<FJsonObject> SpriteDescriptorObject, const FString& NameForErrors, bool bSilent);
 	UPaperSprite* FindExistingSprite(const FString& Name);
 
 protected:
diff --git a/Engine/Plugins/2D/Paper2D/Source/PaperSpriteSheetImporter/Private/PaperSpriteSheetImportFactory.cpp b/Engine/Plugins/2D/Paper2D/Source/PaperSpriteSheetImporter/Private/PaperSpriteSheetImportFactory.cpp
--- a/Engine/Plugins/2D/Paper2D/Source/PaperSpriteSheetImporter/Private/PaperSpriteSheetImportFactory.cpp
+++ b/Engine/Plugins/2D/Paper2D/Source/PaperSpriteSheetImporter/Private/PaperSpriteSheetImportFactory.cpp
@@ -54,7 +54,7 @@ UObject* UPaperSpriteSheetImportFactory::FactoryCreateText(UClass* InClass, UObj
 	const FString NameForErrors(InName.ToString());
 	const FString FileContent(BufferEnd - Buffer, Buffer);
 	
-	if (Importer.ImportFromString(FileContent, NameForErrors) &&
+	if (Importer.ImportFromString(FileContent, NameForErrors, /*bSilent=*/ false) &&
 		Importer.ImportTextures(LongPackagePath, CurrentSourcePath))
 	{
 		UPaperSpriteSheet* SpriteSheet = NewObject<UPaperSpriteSheet>(InParent, InName, Flags);
