Paper2D: Changed asset set methods on Paper2D components to use AreDynamicDataChangesAllowed, allowing asset updates in a construction script

PR #1274: Paper2D static sprite update in ConstructionScript (Contributed by [REDACTED_PERSON])

diff --git a/Engine/Plugins/2D/Paper2D/Source/Paper2D/Private/PaperFlipbookComponent.cpp b/Engine/Plugins/2D/Paper2D/Source/Paper2D/Private/PaperFlipbookComponent.cpp
--- a/Engine/Plugins/2D/Paper2D/Source/Paper2D/Private/PaperFlipbookComponent.cpp
+++ b/Engine/Plugins/2D/Paper2D/Source/Paper2D/Private/PaperFlipbookComponent.cpp
@@ -362,7 +362,7 @@ bool UPaperFlipbookComponent::SetFlipbook(class UPaperFlipbook* NewFlipbook)
 	{
 		// Don't allow changing the sprite if we are "static".
 		AActor* Owner = GetOwner();
-		if (!IsRegistered() || (Owner == nullptr) || (Mobility != EComponentMobility::Static))
+		if ((Owner == nullptr) || AreDynamicDataChangesAllowed())
 		{
 			SourceFlipbook = NewFlipbook;
 
diff --git a/Engine/Plugins/2D/Paper2D/Source/Paper2D/Private/PaperSpriteComponent.cpp b/Engine/Plugins/2D/Paper2D/Source/Paper2D/Private/PaperSpriteComponent.cpp
--- a/Engine/Plugins/2D/Paper2D/Source/Paper2D/Private/PaperSpriteComponent.cpp
+++ b/Engine/Plugins/2D/Paper2D/Source/Paper2D/Private/PaperSpriteComponent.cpp
@@ -198,7 +198,7 @@ bool UPaperSpriteComponent::SetSprite(class UPaperSprite* NewSprite)
 	{
 		// Don't allow changing the sprite if we are "static".
 		AActor* Owner = GetOwner();
-		if (!IsRegistered() || (Owner == nullptr) || (Mobility != EComponentMobility::Static))
+		if ((Owner == nullptr) || AreDynamicDataChangesAllowed())
 		{
 			SourceSprite = NewSprite;
 
diff --git a/Engine/Plugins/2D/Paper2D/Source/Paper2D/Private/PaperTileMapComponent.cpp b/Engine/Plugins/2D/Paper2D/Source/Paper2D/Private/PaperTileMapComponent.cpp
--- a/Engine/Plugins/2D/Paper2D/Source/Paper2D/Private/PaperTileMapComponent.cpp
+++ b/Engine/Plugins/2D/Paper2D/Source/Paper2D/Private/PaperTileMapComponent.cpp
@@ -487,7 +487,7 @@ bool UPaperTileMapComponent::SetTileMap(class UPaperTileMap* NewTileMap)
 	{
 		// Don't allow changing the tile map if we are "static".
 		AActor* Owner = GetOwner();
-		if (!IsRegistered() || (Owner == nullptr) || (Mobility != EComponentMobility::Static))
+		if ((Owner == nullptr) || AreDynamicDataChangesAllowed())
 		{
 			TileMap = NewTileMap;
 
