Paper2D: Fixed a crash when building lighting for a map that contains a placed tile map component or grouped sprite component

diff --git a/Engine/Plugins/2D/Paper2D/Source/Paper2D/Private/PaperRenderSceneProxy.cpp b/Engine/Plugins/2D/Paper2D/Source/Paper2D/Private/PaperRenderSceneProxy.cpp
--- a/Engine/Plugins/2D/Paper2D/Source/Paper2D/Private/PaperRenderSceneProxy.cpp
+++ b/Engine/Plugins/2D/Paper2D/Source/Paper2D/Private/PaperRenderSceneProxy.cpp
@@ -79,22 +79,27 @@ FPaperSpriteVertexFactory::FPaperSpriteVertexFactory()
 
 void FPaperSpriteVertexFactory::Init(const FPaperSpriteVertexBuffer* InVertexBuffer)
 {
-	check(!IsInRenderingThread());
-
-	ENQUEUE_UNIQUE_RENDER_COMMAND_TWOPARAMETER(
-		InitPaperSpriteVertexFactory,
-		FPaperSpriteVertexFactory*,VertexFactory,this,
-		const FPaperSpriteVertexBuffer*,VB,InVertexBuffer,
+	if (IsInRenderingThread())
 	{
 		// Initialize the vertex factory's stream components.
 		DataType NewData;
-		NewData.PositionComponent = STRUCTMEMBER_VERTEXSTREAMCOMPONENT(VB,FPaperSpriteVertex,Position,VET_Float3);
-		NewData.TangentBasisComponents[0] = STRUCTMEMBER_VERTEXSTREAMCOMPONENT(VB,FPaperSpriteVertex,TangentX,VET_PackedNormal);
-		NewData.TangentBasisComponents[1] = STRUCTMEMBER_VERTEXSTREAMCOMPONENT(VB,FPaperSpriteVertex,TangentZ,VET_PackedNormal);
-		NewData.ColorComponent = STRUCTMEMBER_VERTEXSTREAMCOMPONENT(VB,FPaperSpriteVertex,Color,VET_Color);
-		NewData.TextureCoordinates.Add(FVertexStreamComponent(VB, STRUCT_OFFSET(FPaperSpriteVertex,TexCoords), sizeof(FPaperSpriteVertex), VET_Float2));
-		VertexFactory->SetData(NewData);
-	});
+		NewData.PositionComponent = STRUCTMEMBER_VERTEXSTREAMCOMPONENT(InVertexBuffer, FPaperSpriteVertex, Position, VET_Float3);
+		NewData.TangentBasisComponents[0] = STRUCTMEMBER_VERTEXSTREAMCOMPONENT(InVertexBuffer, FPaperSpriteVertex, TangentX, VET_PackedNormal);
+		NewData.TangentBasisComponents[1] = STRUCTMEMBER_VERTEXSTREAMCOMPONENT(InVertexBuffer, FPaperSpriteVertex, TangentZ, VET_PackedNormal);
+		NewData.ColorComponent = STRUCTMEMBER_VERTEXSTREAMCOMPONENT(InVertexBuffer, FPaperSpriteVertex, Color, VET_Color);
+		NewData.TextureCoordinates.Add(FVertexStreamComponent(InVertexBuffer, STRUCT_OFFSET(FPaperSpriteVertex, TexCoords), sizeof(FPaperSpriteVertex), VET_Float2));
+		SetData(NewData);
+	}
+	else
+	{
+		ENQUEUE_UNIQUE_RENDER_COMMAND_TWOPARAMETER(
+			InitPaperSpriteVertexFactory,
+			FPaperSpriteVertexFactory*, VertexFactory, this,
+			const FPaperSpriteVertexBuffer*, VB, InVertexBuffer,
+			{
+				VertexFactory->Init(VB);
+			});
+	}
 }
 
 //////////////////////////////////////////////////////////////////////////
