Even more shadow variable issue fixes
Fixed an issue with multiple declarations of a niagara stat, was only sometimes an issue in unity builds.

diff --git a/Engine/Plugins/2D/Paper2D/Source/Paper2D/Private/PaperFlipbookComponent.cpp b/Engine/Plugins/2D/Paper2D/Source/Paper2D/Private/PaperFlipbookComponent.cpp
--- a/Engine/Plugins/2D/Paper2D/Source/Paper2D/Private/PaperFlipbookComponent.cpp
+++ b/Engine/Plugins/2D/Paper2D/Source/Paper2D/Private/PaperFlipbookComponent.cpp
@@ -361,8 +361,8 @@ bool UPaperFlipbookComponent::SetFlipbook(class UPaperFlipbook* NewFlipbook)
 	if (NewFlipbook != SourceFlipbook)
 	{
 		// Don't allow changing the sprite if we are "static".
-		AActor* Owner = GetOwner();
-		if ((Owner == nullptr) || AreDynamicDataChangesAllowed())
+		AActor* ComponentOwner = GetOwner();
+		if ((ComponentOwner == nullptr) || AreDynamicDataChangesAllowed())
 		{
 			SourceFlipbook = NewFlipbook;
 
diff --git a/Engine/Plugins/2D/Paper2D/Source/Paper2D/Private/PaperRenderSceneProxy.cpp b/Engine/Plugins/2D/Paper2D/Source/Paper2D/Private/PaperRenderSceneProxy.cpp
--- a/Engine/Plugins/2D/Paper2D/Source/Paper2D/Private/PaperRenderSceneProxy.cpp
+++ b/Engine/Plugins/2D/Paper2D/Source/Paper2D/Private/PaperRenderSceneProxy.cpp
@@ -122,41 +122,41 @@ class FPaperSpriteVertexFactoryDummy : public FLocalVertexFactory
 public:
 	void AllocateStuff()
 	{
-		FLocalVertexFactory::DataType Data;
+		FLocalVertexFactory::DataType VertData;
 
-		Data.PositionComponent = FVertexStreamComponent(
+		VertData.PositionComponent = FVertexStreamComponent(
 			&GDummyMaterialSpriteVertexBuffer,
 			STRUCT_OFFSET(FPaperSpriteVertex,Position),
 			sizeof(FPaperSpriteVertex),
 			VET_Float3
 			);
-		Data.TangentBasisComponents[0] = FVertexStreamComponent(
+		VertData.TangentBasisComponents[0] = FVertexStreamComponent(
 			&GDummyMaterialSpriteVertexBuffer,
 			STRUCT_OFFSET(FPaperSpriteVertex,TangentX),
 			sizeof(FPaperSpriteVertex),
 			VET_PackedNormal
 			);
-		Data.TangentBasisComponents[1] = FVertexStreamComponent(
+		VertData.TangentBasisComponents[1] = FVertexStreamComponent(
 			&GDummyMaterialSpriteVertexBuffer,
 			STRUCT_OFFSET(FPaperSpriteVertex,TangentZ),
 			sizeof(FPaperSpriteVertex),
 			VET_PackedNormal
 			);
-		Data.ColorComponent = FVertexStreamComponent(
+		VertData.ColorComponent = FVertexStreamComponent(
 			&GDummyMaterialSpriteVertexBuffer,
 			STRUCT_OFFSET(FPaperSpriteVertex,Color),
 			sizeof(FPaperSpriteVertex),
 			VET_Color
 			);
-		Data.TextureCoordinates.Empty();
-		Data.TextureCoordinates.Add(FVertexStreamComponent(
+		VertData.TextureCoordinates.Empty();
+		VertData.TextureCoordinates.Add(FVertexStreamComponent(
 			&GDummyMaterialSpriteVertexBuffer,
 			STRUCT_OFFSET(FPaperSpriteVertex,TexCoords),
 			sizeof(FPaperSpriteVertex),
 			VET_Float2
 			));
 
-		SetData(Data);
+		SetData(VertData);
 	}
 
 	FPaperSpriteVertexFactoryDummy()
diff --git a/Engine/Plugins/2D/Paper2D/Source/Paper2D/Private/PaperSpriteComponent.cpp b/Engine/Plugins/2D/Paper2D/Source/Paper2D/Private/PaperSpriteComponent.cpp
--- a/Engine/Plugins/2D/Paper2D/Source/Paper2D/Private/PaperSpriteComponent.cpp
+++ b/Engine/Plugins/2D/Paper2D/Source/Paper2D/Private/PaperSpriteComponent.cpp
@@ -197,8 +197,8 @@ bool UPaperSpriteComponent::SetSprite(class UPaperSprite* NewSprite)
 	if (NewSprite != SourceSprite)
 	{
 		// Don't allow changing the sprite if we are "static".
-		AActor* Owner = GetOwner();
-		if ((Owner == nullptr) || AreDynamicDataChangesAllowed())
+		AActor* ComponentOwner = GetOwner();
+		if ((ComponentOwner == nullptr) || AreDynamicDataChangesAllowed())
 		{
 			SourceSprite = NewSprite;
 
diff --git a/Engine/Plugins/2D/Paper2D/Source/Paper2D/Private/PaperTileMapComponent.cpp b/Engine/Plugins/2D/Paper2D/Source/Paper2D/Private/PaperTileMapComponent.cpp
--- a/Engine/Plugins/2D/Paper2D/Source/Paper2D/Private/PaperTileMapComponent.cpp
+++ b/Engine/Plugins/2D/Paper2D/Source/Paper2D/Private/PaperTileMapComponent.cpp
@@ -486,8 +486,8 @@ bool UPaperTileMapComponent::SetTileMap(class UPaperTileMap* NewTileMap)
 	if (NewTileMap != TileMap)
 	{
 		// Don't allow changing the tile map if we are "static".
-		AActor* Owner = GetOwner();
-		if ((Owner == nullptr) || AreDynamicDataChangesAllowed())
+		AActor* ComponentOwner = GetOwner();
+		if ((ComponentOwner == nullptr) || AreDynamicDataChangesAllowed())
 		{
 			TileMap = NewTileMap;
 
diff --git a/Engine/Plugins/2D/Paper2D/Source/Paper2D/Private/PaperTileMapRenderSceneProxy.cpp b/Engine/Plugins/2D/Paper2D/Source/Paper2D/Private/PaperTileMapRenderSceneProxy.cpp
--- a/Engine/Plugins/2D/Paper2D/Source/Paper2D/Private/PaperTileMapRenderSceneProxy.cpp
+++ b/Engine/Plugins/2D/Paper2D/Source/Paper2D/Private/PaperTileMapRenderSceneProxy.cpp
@@ -71,11 +71,11 @@ void FPaperTileMapRenderSceneProxy::FinishConstruction_GameThread()
 
 void FPaperTileMapRenderSceneProxy::DrawBoundsForLayer(FPrimitiveDrawInterface* PDI, const FLinearColor& Color, int32 LayerIndex) const
 {
-	const FMatrix& LocalToWorld = GetLocalToWorld();
-	const FVector TL(LocalToWorld.TransformPosition(TileMap->GetTilePositionInLocalSpace(0, 0, LayerIndex)));
-	const FVector TR(LocalToWorld.TransformPosition(TileMap->GetTilePositionInLocalSpace(TileMap->MapWidth, 0, LayerIndex)));
-	const FVector BL(LocalToWorld.TransformPosition(TileMap->GetTilePositionInLocalSpace(0, TileMap->MapHeight, LayerIndex)));
-	const FVector BR(LocalToWorld.TransformPosition(TileMap->GetTilePositionInLocalSpace(TileMap->MapWidth, TileMap->MapHeight, LayerIndex)));
+	const FMatrix& LocalToWorldMat = GetLocalToWorld();
+	const FVector TL(LocalToWorldMat.TransformPosition(TileMap->GetTilePositionInLocalSpace(0, 0, LayerIndex)));
+	const FVector TR(LocalToWorldMat.TransformPosition(TileMap->GetTilePositionInLocalSpace(TileMap->MapWidth, 0, LayerIndex)));
+	const FVector BL(LocalToWorldMat.TransformPosition(TileMap->GetTilePositionInLocalSpace(0, TileMap->MapHeight, LayerIndex)));
+	const FVector BR(LocalToWorldMat.TransformPosition(TileMap->GetTilePositionInLocalSpace(TileMap->MapWidth, TileMap->MapHeight, LayerIndex)));
 
 	PDI->DrawLine(TL, TR, Color, SDPG_Foreground, 0.0f, WireDepthBias);
 	PDI->DrawLine(TR, BR, Color, SDPG_Foreground, 0.0f, WireDepthBias);
@@ -85,7 +85,7 @@ void FPaperTileMapRenderSceneProxy::DrawBoundsForLayer(FPrimitiveDrawInterface*
 
 void FPaperTileMapRenderSceneProxy::DrawNormalGridLines(FPrimitiveDrawInterface* PDI, const FLinearColor& Color, int32 LayerIndex) const
 {
-	const FMatrix& LocalToWorld = GetLocalToWorld();
+	const FMatrix& LocalToWorldMat = GetLocalToWorld();
 	const uint8 DPG = SDPG_Foreground;//GetDepthPriorityGroup(View);
 
 	// Draw horizontal lines on the selection
@@ -97,7 +97,7 @@ void FPaperTileMapRenderSceneProxy::DrawNormalGridLines(FPrimitiveDrawInterface*
 		X = TileMap->MapWidth;
 		const FVector End(TileMap->GetTilePositionInLocalSpace(X, Y, LayerIndex));
 
-		PDI->DrawLine(LocalToWorld.TransformPosition(Start), LocalToWorld.TransformPosition(End), Color, DPG, 0.0f, WireDepthBias);
+		PDI->DrawLine(LocalToWorldMat.TransformPosition(Start), LocalToWorldMat.TransformPosition(End), Color, DPG, 0.0f, WireDepthBias);
 	}
 
 	// Draw vertical lines
@@ -109,7 +109,7 @@ void FPaperTileMapRenderSceneProxy::DrawNormalGridLines(FPrimitiveDrawInterface*
 		Y = TileMap->MapHeight;
 		const FVector End(TileMap->GetTilePositionInLocalSpace(X, Y, LayerIndex));
 
-		PDI->DrawLine(LocalToWorld.TransformPosition(Start), LocalToWorld.TransformPosition(End), Color, DPG, 0.0f, WireDepthBias);
+		PDI->DrawLine(LocalToWorldMat.TransformPosition(Start), LocalToWorldMat.TransformPosition(End), Color, DPG, 0.0f, WireDepthBias);
 	}
 }
 
@@ -118,7 +118,7 @@ void FPaperTileMapRenderSceneProxy::DrawStaggeredGridLines(FPrimitiveDrawInterfa
 	TArray<FVector> Poly;
 	Poly.Empty(4);
 
-	const FMatrix& LocalToWorld = GetLocalToWorld();
+	const FMatrix& LocalToWorldMat = GetLocalToWorld();
 	const uint8 DPG = SDPG_Foreground;//GetDepthPriorityGroup(View);
 
 	FVector CornerPosition;
@@ -167,7 +167,7 @@ void FPaperTileMapRenderSceneProxy::DrawStaggeredGridLines(FPrimitiveDrawInterfa
 		TileMap->GetTilePolygon(XBottom, YBottom, LayerIndex, Poly);
 		const FVector LSB = Poly[2];
 
-		PDI->DrawLine(LocalToWorld.TransformPosition(LSA), LocalToWorld.TransformPosition(LSB), Color, DPG, 0.0f, WireDepthBias);
+		PDI->DrawLine(LocalToWorldMat.TransformPosition(LSA), LocalToWorldMat.TransformPosition(LSB), Color, DPG, 0.0f, WireDepthBias);
 	}
 
 	for (int32 X = 0; X < TileMap->MapWidth + ((TileMap->MapHeight + 1) / 2) + 1; ++X)
@@ -197,14 +197,14 @@ void FPaperTileMapRenderSceneProxy::DrawStaggeredGridLines(FPrimitiveDrawInterfa
 		TileMap->GetTilePolygon(XBottom, YBottom, LayerIndex, Poly);
 		const FVector LSB = Poly[3];
 
-		PDI->DrawLine(LocalToWorld.TransformPosition(LSA), LocalToWorld.TransformPosition(LSB), Color, DPG, 0.0f, WireDepthBias);
+		PDI->DrawLine(LocalToWorldMat.TransformPosition(LSA), LocalToWorldMat.TransformPosition(LSB), Color, DPG, 0.0f, WireDepthBias);
 	}
 }
 
 void FPaperTileMapRenderSceneProxy::DrawHexagonalGridLines(FPrimitiveDrawInterface* PDI, const FLinearColor& Color, int32 LayerIndex) const
 {
 	//@TODO: This isn't very efficient
-	const FMatrix& LocalToWorld = GetLocalToWorld();
+	const FMatrix& LocalToWorldMat = GetLocalToWorld();
 	const uint8 DPG = SDPG_Foreground;//GetDepthPriorityGroup(View);
 
 	TArray<FVector> Poly;
@@ -216,10 +216,10 @@ void FPaperTileMapRenderSceneProxy::DrawHexagonalGridLines(FPrimitiveDrawInterfa
 			Poly.Reset();
 			TileMap->GetTilePolygon(X, Y, LayerIndex, Poly);
 
-			FVector LastVertexWS = LocalToWorld.TransformPosition(Poly[5]);
+			FVector LastVertexWS = LocalToWorldMat.TransformPosition(Poly[5]);
 			for (int32 VI = 0; VI < Poly.Num(); ++VI)
 			{
-				FVector ThisVertexWS = LocalToWorld.TransformPosition(Poly[VI]);
+				FVector ThisVertexWS = LocalToWorldMat.TransformPosition(Poly[VI]);
 				PDI->DrawLine(LastVertexWS, ThisVertexWS, Color, DPG, 0.0f, WireDepthBias);
 				LastVertexWS = ThisVertexWS;
 			}
@@ -304,8 +304,6 @@ void FPaperTileMapRenderSceneProxy::GetDynamicMeshElements(const TArray<const FS
 				// Draw separation wires if selected
 				const FLinearColor OverrideColor = GetSelectionColor(FLinearColor::White, bShowAsSelected, IsHovered(), /*bUseOverlayIntensity=*/ false);
 
-				FTransform LocalToWorld(GetLocalToWorld());
-
 				// Draw the debug outline
 				if (bEffectivelySelected)
 				{
diff --git a/Engine/Plugins/2D/Paper2D/Source/Paper2D/Private/Rendering/PaperBatchComponent.cpp b/Engine/Plugins/2D/Paper2D/Source/Paper2D/Private/Rendering/PaperBatchComponent.cpp
--- a/Engine/Plugins/2D/Paper2D/Source/Paper2D/Private/Rendering/PaperBatchComponent.cpp
+++ b/Engine/Plugins/2D/Paper2D/Source/Paper2D/Private/Rendering/PaperBatchComponent.cpp
@@ -26,8 +26,8 @@ FPrimitiveSceneProxy* UPaperBatchComponent::CreateSceneProxy()
 FBoxSphereBounds UPaperBatchComponent::CalcBounds(const FTransform& LocalToWorld) const
 {
 	// Always visible
-	FBoxSphereBounds Bounds(FVector::ZeroVector, FVector(HALF_WORLD_MAX, HALF_WORLD_MAX, HALF_WORLD_MAX), HALF_WORLD_MAX);
-	return Bounds;
+	FBoxSphereBounds BoxSphereBounds(FVector::ZeroVector, FVector(HALF_WORLD_MAX, HALF_WORLD_MAX, HALF_WORLD_MAX), HALF_WORLD_MAX);
+	return BoxSphereBounds;
 }
 
 void UPaperBatchComponent::TickComponent(float DeltaTime, enum ELevelTick TickType, FActorComponentTickFunction* ThisTickFunction)
diff --git a/Engine/Plugins/Runtime/OculusRift/Source/OculusRift/Private/OculusRiftHMD.cpp b/Engine/Plugins/Runtime/OculusRift/Source/OculusRift/Private/OculusRiftHMD.cpp
--- a/Engine/Plugins/Runtime/OculusRift/Source/OculusRift/Private/OculusRiftHMD.cpp
+++ b/Engine/Plugins/Runtime/OculusRift/Source/OculusRift/Private/OculusRiftHMD.cpp
@@ -761,11 +761,11 @@ bool FOculusRiftHMD::DoEnableStereo(bool bStereo, bool bApplyToHmd)
 							// returns false.
 							if (GIsEditor)
 							{
-								FSlateRect PreFullScreenRect;
-								PopPreFullScreenRect(PreFullScreenRect);
-								if (PreFullScreenRect.GetSize().X > 0 && PreFullScreenRect.GetSize().Y > 0 && IsFullscreenAllowed())
+								FSlateRect LocalPreFullScreenRect;
+								PopPreFullScreenRect(LocalPreFullScreenRect);
+								if (LocalPreFullScreenRect.GetSize().X > 0 && LocalPreFullScreenRect.GetSize().Y > 0 && IsFullscreenAllowed())
 								{
-									Window->MoveWindowTo(FVector2D(PreFullScreenRect.Left, PreFullScreenRect.Top));
+									Window->MoveWindowTo(FVector2D(LocalPreFullScreenRect.Left, LocalPreFullScreenRect.Top));
 								}
 							}
 							else
diff --git a/Engine/Plugins/Runtime/OculusRift/Source/OculusRift/Private/OculusRiftHMD.h b/Engine/Plugins/Runtime/OculusRift/Source/OculusRift/Private/OculusRiftHMD.h
--- a/Engine/Plugins/Runtime/OculusRift/Source/OculusRift/Private/OculusRiftHMD.h
+++ b/Engine/Plugins/Runtime/OculusRift/Source/OculusRift/Private/OculusRiftHMD.h
@@ -223,7 +223,7 @@ public:
 	// Returns true if Engine should perform its own Present.
 	virtual bool Present(int32& SyncInterval) override;
 
-	virtual void UpdateViewport(const FViewport& Viewport, FRHIViewport* ViewportRHI, FGameFrame* InRenderFrame);
+	virtual void UpdateViewport(const FViewport& Viewport, FRHIViewport* InViewportRHI, FGameFrame* InRenderFrame);
 	FGameFrame* GetRenderFrame() const { check(IsInRenderingThread()); return static_cast<FGameFrame*>(RenderContext->RenderFrame.Get()); }
 	FViewExtension* GetRenderContext() const { return static_cast<FViewExtension*>(RenderContext.Get()); }
 	FSettings* GetFrameSetting() const { check(IsInRenderingThread()); return static_cast<FSettings*>(RenderContext->RenderFrame->GetSettings()); }
@@ -309,7 +309,7 @@ public:
     virtual void SetupView(FSceneViewFamily& InViewFamily, FSceneView& InView) override;
 	virtual FRHICustomPresent* GetCustomPresent() override { return pCustomPresent; }
 	virtual uint32 GetNumberOfBufferedFrames() const override { return 1; }
-	virtual bool AllocateRenderTargetTexture(uint32 Index, uint32 SizeX, uint32 SizeY, uint8 Format, uint32 NumMips, uint32 Flags, uint32 TargetableTextureFlags, FTexture2DRHIRef& OutTargetableTexture, FTexture2DRHIRef& OutShaderResourceTexture, uint32 NumSamples = 1) override;
+	virtual bool AllocateRenderTargetTexture(uint32 Index, uint32 SizeX, uint32 SizeY, uint8 Format, uint32 NumMips, uint32 InFlags, uint32 TargetableTextureFlags, FTexture2DRHIRef& OutTargetableTexture, FTexture2DRHIRef& OutShaderResourceTexture, uint32 NumSamples = 1) override;
 
 	/** Positional tracking control methods */
 	virtual bool IsHeadTrackingAllowed() const override;
@@ -393,7 +393,7 @@ public:
 
 		virtual void SetHmd(ovrHmd InHmd) override;
 
-		virtual bool AllocateRenderTargetTexture(uint32 SizeX, uint32 SizeY, uint8 Format, uint32 NumMips, uint32 Flags, uint32 TargetableTextureFlags, FTexture2DRHIRef& OutTargetableTexture, FTexture2DRHIRef& OutShaderResourceTexture, uint32 NumSamples);
+		virtual bool AllocateRenderTargetTexture(uint32 SizeX, uint32 SizeY, uint8 Format, uint32 NumMips, uint32 InFlags, uint32 TargetableTextureFlags, FTexture2DRHIRef& OutTargetableTexture, FTexture2DRHIRef& OutShaderResourceTexture, uint32 NumSamples);
 	protected:
 		void Init(ovrHmd InHmd);
 		void Reset_RenderThread();
@@ -421,7 +421,7 @@ public:
 
 		virtual void SetHmd(ovrHmd InHmd) override;
 
-		virtual bool AllocateRenderTargetTexture(uint32 SizeX, uint32 SizeY, uint8 Format, uint32 NumMips, uint32 Flags, uint32 TargetableTextureFlags, FTexture2DRHIRef& OutTargetableTexture, FTexture2DRHIRef& OutShaderResourceTexture, uint32 NumSamples);
+		virtual bool AllocateRenderTargetTexture(uint32 SizeX, uint32 SizeY, uint8 Format, uint32 NumMips, uint32 InFlags, uint32 TargetableTextureFlags, FTexture2DRHIRef& OutTargetableTexture, FTexture2DRHIRef& OutShaderResourceTexture, uint32 NumSamples);
 
 	protected:
 		void Init(ovrHmd InHmd);
diff --git a/Engine/Plugins/Runtime/OculusRift/Source/OculusRift/Private/OculusRiftRender.cpp b/Engine/Plugins/Runtime/OculusRift/Source/OculusRift/Private/OculusRiftRender.cpp
--- a/Engine/Plugins/Runtime/OculusRift/Source/OculusRift/Private/OculusRiftRender.cpp
+++ b/Engine/Plugins/Runtime/OculusRift/Source/OculusRift/Private/OculusRiftRender.cpp
@@ -14,10 +14,10 @@
 
 #include "SlateBasics.h"
 
-bool FOculusRiftHMD::AllocateRenderTargetTexture(uint32 Index, uint32 SizeX, uint32 SizeY, uint8 Format, uint32 NumMips, uint32 Flags, uint32 TargetableTextureFlags, FTexture2DRHIRef& OutTargetableTexture, FTexture2DRHIRef& OutShaderResourceTexture, uint32 NumSamples)
+bool FOculusRiftHMD::AllocateRenderTargetTexture(uint32 Index, uint32 SizeX, uint32 SizeY, uint8 Format, uint32 NumMips, uint32 InFlags, uint32 TargetableTextureFlags, FTexture2DRHIRef& OutTargetableTexture, FTexture2DRHIRef& OutShaderResourceTexture, uint32 NumSamples)
 {
 	check(Index == 0);
-	pCustomPresent->AllocateRenderTargetTexture(SizeX, SizeY, Format, NumMips, Flags, TargetableTextureFlags, OutTargetableTexture, OutShaderResourceTexture, NumSamples);
+	pCustomPresent->AllocateRenderTargetTexture(SizeX, SizeY, Format, NumMips, InFlags, TargetableTextureFlags, OutTargetableTexture, OutShaderResourceTexture, NumSamples);
 	return true;
 }
 
@@ -616,12 +616,12 @@ void FCustomPresent::SetRenderContext(FHMDViewExtension* InRenderContext)
 	}
 }
 
-void FCustomPresent::UpdateViewport(const FViewport& Viewport, FRHIViewport* ViewportRHI, FGameFrame* InRenderFrame)
+void FCustomPresent::UpdateViewport(const FViewport& Viewport, FRHIViewport* InViewportRHI, FGameFrame* InRenderFrame)
 {
 	check(IsInGameThread());
-	check(ViewportRHI);
+	check(InViewportRHI);
 
-	this->ViewportRHI = ViewportRHI;
+	this->ViewportRHI = InViewportRHI;
 	ViewportRHI->SetCustomPresent(this);
 }
 
diff --git a/Engine/Plugins/Runtime/OculusRift/Source/OculusRift/Private/OculusRiftRenderD3D.cpp b/Engine/Plugins/Runtime/OculusRift/Source/OculusRift/Private/OculusRiftRenderD3D.cpp
--- a/Engine/Plugins/Runtime/OculusRift/Source/OculusRift/Private/OculusRiftRenderD3D.cpp
+++ b/Engine/Plugins/Runtime/OculusRift/Source/OculusRift/Private/OculusRiftRenderD3D.cpp
@@ -345,7 +345,7 @@ void FOculusRiftHMD::D3D11Bridge::Init(ovrHmd InHmd)
 	bInitialized = true;
 }
 
-bool FOculusRiftHMD::D3D11Bridge::AllocateRenderTargetTexture(uint32 SizeX, uint32 SizeY, uint8 Format, uint32 NumMips, uint32 Flags, uint32 TargetableTextureFlags, FTexture2DRHIRef& OutTargetableTexture, FTexture2DRHIRef& OutShaderResourceTexture, uint32 NumSamples)
+bool FOculusRiftHMD::D3D11Bridge::AllocateRenderTargetTexture(uint32 SizeX, uint32 SizeY, uint8 Format, uint32 NumMips, uint32 InFlags, uint32 TargetableTextureFlags, FTexture2DRHIRef& OutTargetableTexture, FTexture2DRHIRef& OutShaderResourceTexture, uint32 NumSamples)
 {
 	check(SizeX != 0 && SizeY != 0);
 
diff --git a/Engine/Plugins/Runtime/OculusRift/Source/OculusRift/Private/OculusRiftRenderGL.cpp b/Engine/Plugins/Runtime/OculusRift/Source/OculusRift/Private/OculusRiftRenderGL.cpp
--- a/Engine/Plugins/Runtime/OculusRift/Source/OculusRift/Private/OculusRiftRenderGL.cpp
+++ b/Engine/Plugins/Runtime/OculusRift/Source/OculusRift/Private/OculusRiftRenderGL.cpp
@@ -245,7 +245,7 @@ void FOculusRiftHMD::OGLBridge::Init(ovrHmd InHmd)
 	bInitialized = true;
 }
 
-bool FOculusRiftHMD::OGLBridge::AllocateRenderTargetTexture(uint32 SizeX, uint32 SizeY, uint8 Format, uint32 NumMips, uint32 Flags, uint32 TargetableTextureFlags, FTexture2DRHIRef& OutTargetableTexture, FTexture2DRHIRef& OutShaderResourceTexture, uint32 NumSamples)
+bool FOculusRiftHMD::OGLBridge::AllocateRenderTargetTexture(uint32 SizeX, uint32 SizeY, uint8 Format, uint32 NumMips, uint32 InFlags, uint32 TargetableTextureFlags, FTexture2DRHIRef& OutTargetableTexture, FTexture2DRHIRef& OutShaderResourceTexture, uint32 NumSamples)
 {
 	check(SizeX != 0 && SizeY != 0);
 
@@ -254,14 +254,14 @@ bool FOculusRiftHMD::OGLBridge::AllocateRenderTargetTexture(uint32 SizeX, uint32
 		bNeedReAllocateTextureSet = true;
 	}
 
-	Flags |= TargetableTextureFlags;
+	InFlags |= TargetableTextureFlags;
 
 	if (Hmd && bNeedReAllocateTextureSet)
 	{
 		auto GLRHI = static_cast<FOpenGLDynamicRHI*>(GDynamicRHI);
 		FResourceBulkDataInterface* BulkData = nullptr;
 
-// 		FRHITexture* p = GLRHI->CreateOpenGLTexture(SizeX, SizeY, false, false, Format, NumMips, NumSamples, 1, Flags, BulkData);
+// 		FRHITexture* p = GLRHI->CreateOpenGLTexture(SizeX, SizeY, false, false, Format, NumMips, NumSamples, 1, InFlags, BulkData);
 // 		OutTargetableTexture = p->GetTexture2D();
 // 		OutShaderResourceTexture = p->GetTexture2D();
 // 		return true;
@@ -272,7 +272,7 @@ bool FOculusRiftHMD::OGLBridge::AllocateRenderTargetTexture(uint32 SizeX, uint32
 			ColorTextureSet = nullptr;
 		}
 
-		const bool bSRGB = (Flags & TexCreate_SRGB) != 0;
+		const bool bSRGB = (InFlags & TexCreate_SRGB) != 0;
 		const FOpenGLTextureFormat& GLFormat = GOpenGLTextureFormats[Format];
 		if (GLFormat.InternalFormat[bSRGB] == GL_NONE)
 		{
diff --git a/Engine/Source/Runtime/Engine/Public/PhysicsPublic.h b/Engine/Source/Runtime/Engine/Public/PhysicsPublic.h
--- a/Engine/Source/Runtime/Engine/Public/PhysicsPublic.h
+++ b/Engine/Source/Runtime/Engine/Public/PhysicsPublic.h
@@ -424,7 +424,7 @@ public:
 	ENGINE_API void StartFrame();
 
 	/** Ends a frame */
-	ENGINE_API void EndFrame(ULineBatchComponent* LineBatcher);
+	ENGINE_API void EndFrame(ULineBatchComponent* InLineBatcher);
 
 	/** Starts cloth Simulation*/
 	ENGINE_API void StartCloth();
diff --git a/Engine/Source/Runtime/Niagara/Private/NiagaraComponent.cpp b/Engine/Source/Runtime/Niagara/Private/NiagaraComponent.cpp
--- a/Engine/Source/Runtime/Niagara/Private/NiagaraComponent.cpp
+++ b/Engine/Source/Runtime/Niagara/Private/NiagaraComponent.cpp
@@ -14,7 +14,6 @@
 #include "ComponentReregisterContext.h"
 
 DECLARE_CYCLE_STAT(TEXT("Gen Verts"),STAT_NiagaraGenerateVertices,STATGROUP_Niagara);
-DECLARE_DWORD_COUNTER_STAT(TEXT("NumParticles"),STAT_NiagaraNumParticles,STATGROUP_Niagara);
 
 DEFINE_LOG_CATEGORY(LogNiagara);
 
diff --git a/Engine/Source/Runtime/Niagara/Private/NiagaraPrivate.h b/Engine/Source/Runtime/Niagara/Private/NiagaraPrivate.h
--- a/Engine/Source/Runtime/Niagara/Private/NiagaraPrivate.h
+++ b/Engine/Source/Runtime/Niagara/Private/NiagaraPrivate.h
@@ -3,4 +3,6 @@
 
 #pragma once
 #include "Public/Engine.h"
-#include "VectorVM.h"
\ No newline at end of file
+#include "VectorVM.h"
+
+DECLARE_DWORD_COUNTER_STAT(TEXT("NumParticles"), STAT_NiagaraNumParticles, STATGROUP_Niagara);
\ No newline at end of file
diff --git a/Engine/Source/Runtime/Niagara/Private/NiagaraSimulation.cpp b/Engine/Source/Runtime/Niagara/Private/NiagaraSimulation.cpp
--- a/Engine/Source/Runtime/Niagara/Private/NiagaraSimulation.cpp
+++ b/Engine/Source/Runtime/Niagara/Private/NiagaraSimulation.cpp
@@ -190,7 +190,6 @@ void FNiagaraSimulation::Tick(float DeltaSeconds)
 
 	CPUTimeMS = TickTime.GetElapsedMilliseconds();
 
-	DECLARE_DWORD_COUNTER_STAT(TEXT("NumParticles"), STAT_NiagaraNumParticles, STATGROUP_Niagara);
 	INC_DWORD_STAT_BY(STAT_NiagaraNumParticles, Data.GetNumParticles());
 
 }
diff --git a/Engine/Source/Runtime/OpenGLDrv/Public/OpenGLResources.h b/Engine/Source/Runtime/OpenGLDrv/Public/OpenGLResources.h
--- a/Engine/Source/Runtime/OpenGLDrv/Public/OpenGLResources.h
+++ b/Engine/Source/Runtime/OpenGLDrv/Public/OpenGLResources.h
@@ -1003,12 +1003,12 @@ public:
 	/**
 	 * Clone texture from a source using CopyImageSubData
 	 */
-	void CloneViaCopyImage( TOpenGLTexture* Src, uint32 NumMips, int32 SrcOffset, int32 DstOffset);
+	void CloneViaCopyImage( TOpenGLTexture* Src, uint32 InNumMips, int32 SrcOffset, int32 DstOffset);
 	
 	/**
 	 * Clone texture from a source going via PBOs
 	 */
-	void CloneViaPBO( TOpenGLTexture* Src, uint32 NumMips, int32 SrcOffset, int32 DstOffset);
+	void CloneViaPBO( TOpenGLTexture* Src, uint32 InNumMips, int32 SrcOffset, int32 DstOffset);
 
 	/**
 	 * Resolved the specified face for a read Lock, for non-renderable, CPU readable surfaces this eliminates the readback inside Lock itself.
diff --git a/Engine/Source/Runtime/Renderer/Private/PostProcess/RenderingCompositionGraph.cpp b/Engine/Source/Runtime/Renderer/Private/PostProcess/RenderingCompositionGraph.cpp
--- a/Engine/Source/Runtime/Renderer/Private/PostProcess/RenderingCompositionGraph.cpp
+++ b/Engine/Source/Runtime/Renderer/Private/PostProcess/RenderingCompositionGraph.cpp
@@ -240,10 +240,10 @@ void FRenderingCompositionGraph::RecursivelyGatherDependencies(FRenderingComposi
 			InputOutput->AddDependency();
 		}
 		
-		if(FRenderingCompositePass* Pass = OutputRefIt->GetPass())
+		if(FRenderingCompositePass* OutputRefItPass = OutputRefIt->GetPass())
 		{
 			// recursively process all inputs of this Pass
-			RecursivelyGatherDependencies(Pass);
+			RecursivelyGatherDependencies(OutputRefItPass);
 		}
 	}
 
