Paper2D: Fixed inconsistency between vertex colors on Paper2D assets and static meshes
- All sprite colors are now treated as being defined in linear space, and existing assets will be converted on load

Upgrade Node: Blueprints or C++ code that programmatically modifies the color of a Paper2D component will need to be updated, otherwise the color being applied may be darker than expected

diff --git a/Engine/Plugins/2D/Paper2D/Source/Paper2D/Classes/PaperTerrainComponent.h b/Engine/Plugins/2D/Paper2D/Source/Paper2D/Classes/PaperTerrainComponent.h
--- a/Engine/Plugins/2D/Paper2D/Source/Paper2D/Classes/PaperTerrainComponent.h
+++ b/Engine/Plugins/2D/Paper2D/Source/Paper2D/Classes/PaperTerrainComponent.h
@@ -88,6 +88,8 @@ protected:
 public:
 	// UObject interface
 	virtual const UObject* AdditionalStatObject() const override;
+	virtual void Serialize(FArchive& Ar) override;
+	virtual void PostLoad() override;
 #if WITH_EDITOR
 	virtual void PostEditChangeProperty(FPropertyChangedEvent& PropertyChangedEvent) override;
 #endif
diff --git a/Engine/Plugins/2D/Paper2D/Source/Paper2D/Classes/PaperTileMap.h b/Engine/Plugins/2D/Paper2D/Source/Paper2D/Classes/PaperTileMap.h
--- a/Engine/Plugins/2D/Paper2D/Source/Paper2D/Classes/PaperTileMap.h
+++ b/Engine/Plugins/2D/Paper2D/Source/Paper2D/Classes/PaperTileMap.h
@@ -126,16 +126,16 @@ public:
 
 public:
 	// UObject interface
+	virtual void Serialize(FArchive& Ar) override;
+	virtual void PostLoad() override;
 #if WITH_EDITOR
 	virtual void PreEditChange(UProperty* PropertyAboutToChange) override;
 	virtual void PostEditChangeProperty(FPropertyChangedEvent& PropertyChangedEvent) override;
 	virtual bool CanEditChange(const UProperty* InProperty) const override;
-	virtual void PostLoad() override;
 	void ValidateSelectedLayerIndex();
 #endif
 #if WITH_EDITORONLY_DATA
 	virtual void GetAssetRegistryTags(TArray<FAssetRegistryTag>& OutTags) const override;
-	virtual void Serialize(FArchive& Ar) override;
 #endif
 	// End of UObject interface
 
diff --git a/Engine/Plugins/2D/Paper2D/Source/Paper2D/Private/PaperCustomVersion.h b/Engine/Plugins/2D/Paper2D/Source/Paper2D/Private/PaperCustomVersion.h
--- a/Engine/Plugins/2D/Paper2D/Source/Paper2D/Private/PaperCustomVersion.h
+++ b/Engine/Plugins/2D/Paper2D/Source/Paper2D/Private/PaperCustomVersion.h
@@ -58,6 +58,10 @@ struct FPaperCustomVersion
 		// Fixed the tangent generation for sprites, tile maps, etc... being incorrect for the front face
 		FixTangentGenerationForFrontFace = 16,
 
+		// Converted Paper2D vertex colors to be treated as linear space instead of sRGB space, and the color picker now matches the result on screen
+		// rather than sprites appearing far brighter than they should (note: they are still sent down as a FColor so they have limited precision)
+		FixVertexColorSpace = 17,
+
 		// -----<new versions can be added above this line>-------------------------------------------------
 		VersionPlusOne,
 		LatestVersion = VersionPlusOne - 1
diff --git a/Engine/Plugins/2D/Paper2D/Source/Paper2D/Private/PaperFlipbookComponent.cpp b/Engine/Plugins/2D/Paper2D/Source/Paper2D/Private/PaperFlipbookComponent.cpp
--- a/Engine/Plugins/2D/Paper2D/Source/Paper2D/Private/PaperFlipbookComponent.cpp
+++ b/Engine/Plugins/2D/Paper2D/Source/Paper2D/Private/PaperFlipbookComponent.cpp
@@ -76,6 +76,12 @@ void UPaperFlipbookComponent::PostLoad()
 			SetMaterial(0, Material_DEPRECATED);
 		}
 	}
+
+	if (PaperVer < FPaperCustomVersion::FixVertexColorSpace)
+	{
+		const FColor SRGBColor = SpriteColor.ToFColor(/*bSRGB=*/ true);
+		SpriteColor = SRGBColor.ReinterpretAsLinear();
+	}
 }
 
 FPrimitiveSceneProxy* UPaperFlipbookComponent::CreateSceneProxy()
@@ -87,7 +93,7 @@ FPrimitiveSceneProxy* UPaperFlipbookComponent::CreateSceneProxy()
 
 	FSpriteDrawCallRecord DrawCall;
 	DrawCall.BuildFromSprite(SpriteToSend);
-	DrawCall.Color = SpriteColor.ToFColor(true);
+	DrawCall.Color = SpriteColor.ToFColor(/*bSRGB=*/ false);
 	NewProxy->SetDrawCall_RenderThread(DrawCall);
 	return NewProxy;
 }
@@ -342,7 +348,7 @@ void UPaperFlipbookComponent::SendRenderDynamicData_Concurrent()
 
 		FSpriteDrawCallRecord DrawCall;
 		DrawCall.BuildFromSprite(SpriteToSend);
-		DrawCall.Color = SpriteColor.ToFColor(true);
+		DrawCall.Color = SpriteColor.ToFColor(/*bSRGB=*/ false);
 
 		ENQUEUE_UNIQUE_RENDER_COMMAND_THREEPARAMETER(
 				FSendPaperRenderComponentDynamicData,
diff --git a/Engine/Plugins/2D/Paper2D/Source/Paper2D/Private/PaperGroupedSpriteComponent.cpp b/Engine/Plugins/2D/Paper2D/Source/Paper2D/Private/PaperGroupedSpriteComponent.cpp
--- a/Engine/Plugins/2D/Paper2D/Source/Paper2D/Private/PaperGroupedSpriteComponent.cpp
+++ b/Engine/Plugins/2D/Paper2D/Source/Paper2D/Private/PaperGroupedSpriteComponent.cpp
@@ -42,7 +42,7 @@ int32 UPaperGroupedSpriteComponent::AddInstanceWithMaterial(const FTransform& Tr
 	const FTransform LocalTransform(bWorldSpace ? Transform.GetRelativeTransform(ComponentToWorld) : Transform);
 
 	FSpriteInstanceData& NewInstanceData = *new (PerInstanceSpriteData)FSpriteInstanceData();
-	SetupNewInstanceData(NewInstanceData, NewInstanceIndex, LocalTransform, Sprite, MaterialOverride, Color.ToFColor(true));
+	SetupNewInstanceData(NewInstanceData, NewInstanceIndex, LocalTransform, Sprite, MaterialOverride, Color.ToFColor(/*bSRGB=*/ false));
 
 	MarkRenderStateDirty();
 
@@ -113,7 +113,7 @@ bool UPaperGroupedSpriteComponent::UpdateInstanceColor(int32 InstanceIndex, FLin
 	if (PerInstanceSpriteData.IsValidIndex(InstanceIndex))
 	{
 		FSpriteInstanceData& InstanceData = PerInstanceSpriteData[InstanceIndex];
-		InstanceData.VertexColor = NewInstanceColor.ToFColor(true);
+		InstanceData.VertexColor = NewInstanceColor.ToFColor(/*bSRGB=*/ false);
 
 		if (bMarkRenderStateDirty)
 		{
diff --git a/Engine/Plugins/2D/Paper2D/Source/Paper2D/Private/PaperSprite.cpp b/Engine/Plugins/2D/Paper2D/Source/Paper2D/Private/PaperSprite.cpp
--- a/Engine/Plugins/2D/Paper2D/Source/Paper2D/Private/PaperSprite.cpp
+++ b/Engine/Plugins/2D/Paper2D/Source/Paper2D/Private/PaperSprite.cpp
@@ -356,7 +356,7 @@ void FSpriteDrawCallRecord::BuildFromSprite(const UPaperSprite* Sprite)
 		BaseTexture = Sprite->GetBakedTexture();
 		Sprite->GetBakedAdditionalSourceTextures(/*out*/ AdditionalTextures);
 
-		Color = FLinearColor::White.ToFColor(true);
+		Color = FColor::White;
 
 		RenderVerts = Sprite->BakedRenderData;
 	}
diff --git a/Engine/Plugins/2D/Paper2D/Source/Paper2D/Private/PaperSpriteComponent.cpp b/Engine/Plugins/2D/Paper2D/Source/Paper2D/Private/PaperSpriteComponent.cpp
--- a/Engine/Plugins/2D/Paper2D/Source/Paper2D/Private/PaperSpriteComponent.cpp
+++ b/Engine/Plugins/2D/Paper2D/Source/Paper2D/Private/PaperSpriteComponent.cpp
@@ -58,6 +58,12 @@ void UPaperSpriteComponent::PostLoad()
 			SetMaterial(0, MaterialOverride_DEPRECATED);
 		}
 	}
+
+	if (PaperVer < FPaperCustomVersion::FixVertexColorSpace)
+	{
+		const FColor SRGBColor = SpriteColor.ToFColor(/*bSRGB=*/ true);
+		SpriteColor = SRGBColor.ReinterpretAsLinear();
+	}
 }
 
 FPrimitiveSceneProxy* UPaperSpriteComponent::CreateSceneProxy()
@@ -67,7 +73,7 @@ FPrimitiveSceneProxy* UPaperSpriteComponent::CreateSceneProxy()
 	{
 		FSpriteDrawCallRecord DrawCall;
 		DrawCall.BuildFromSprite(SourceSprite);
-		DrawCall.Color = SpriteColor.ToFColor(true);
+		DrawCall.Color = SpriteColor.ToFColor(/*bSRGB=*/ false);
 		NewProxy->SetSprite_RenderThread(DrawCall, SourceSprite->AlternateMaterialSplitIndex);
 	}
 	return NewProxy;
@@ -108,7 +114,7 @@ void UPaperSpriteComponent::SendRenderDynamicData_Concurrent()
 	{
 		FSpriteDrawCallRecord DrawCall;
 		DrawCall.BuildFromSprite(SourceSprite);
-		DrawCall.Color = SpriteColor.ToFColor(true);
+		DrawCall.Color = SpriteColor.ToFColor(/*bSRGB=*/ false);
 		const int32 SplitIndex = (SourceSprite != nullptr) ? SourceSprite->AlternateMaterialSplitIndex : INDEX_NONE;
 
 		ENQUEUE_UNIQUE_RENDER_COMMAND_THREEPARAMETER(
diff --git a/Engine/Plugins/2D/Paper2D/Source/Paper2D/Private/PaperTileMap.cpp b/Engine/Plugins/2D/Paper2D/Source/Paper2D/Private/PaperTileMap.cpp
--- a/Engine/Plugins/2D/Paper2D/Source/Paper2D/Private/PaperTileMap.cpp
+++ b/Engine/Plugins/2D/Paper2D/Source/Paper2D/Private/PaperTileMap.cpp
@@ -51,6 +51,46 @@ void UPaperTileMap::PreEditChange(UProperty* PropertyAboutToChange)
 }
 #endif
 
+void UPaperTileMap::Serialize(FArchive& Ar)
+{
+	Super::Serialize(Ar);
+
+	Ar.UsingCustomVersion(FPaperCustomVersion::GUID);
+
+#if WITH_EDITORONLY_DATA
+	if (Ar.IsLoading() && (Ar.UE4Ver() < VER_UE4_ASSET_IMPORT_DATA_AS_JSON) && (AssetImportData == nullptr))
+	{
+		// AssetImportData should always be valid
+		AssetImportData = NewObject<UAssetImportData>(this, TEXT("AssetImportData"));
+	}
+#endif
+}
+
+void UPaperTileMap::PostLoad()
+{
+	Super::PostLoad();
+
+	const int32 PaperVer = GetLinkerCustomVersion(FPaperCustomVersion::GUID);
+
+	// Make sure that the layers are all of the right size (there was a bug at one point when undoing resizes that could cause the layers to get stuck at a bad size)
+	for (UPaperTileLayer* TileLayer : TileLayers)
+	{
+		TileLayer->ConditionalPostLoad();
+
+		TileLayer->ResizeMap(MapWidth, MapHeight);
+
+		if (PaperVer < FPaperCustomVersion::FixVertexColorSpace)
+		{
+			const FColor SRGBColor = TileLayer->GetLayerColor().ToFColor(/*bSRGB=*/ true);
+			TileLayer->SetLayerColor(SRGBColor.ReinterpretAsLinear());
+		}
+	}
+
+#if WITH_EDITOR
+	ValidateSelectedLayerIndex();
+#endif
+}
+
 #if WITH_EDITOR
 
 #include "PaperTileMapComponent.h"
@@ -161,34 +201,6 @@ bool UPaperTileMap::CanEditChange(const UProperty* InProperty) const
 	return bIsEditable;
 }
 
-void UPaperTileMap::PostLoad()
-{
-	Super::PostLoad();
-
-	// Make sure that the layers are all of the right size (there was a bug at one point when undoing resizes that could cause the layers to get stuck at a bad size)
-	for (UPaperTileLayer* TileLayer : TileLayers)
-	{
-		TileLayer->ConditionalPostLoad();
-
-		TileLayer->ResizeMap(MapWidth, MapHeight);
-	}
-
-	ValidateSelectedLayerIndex();
-}
-
-#if WITH_EDITORONLY_DATA
-void UPaperTileMap::Serialize(FArchive& Ar)
-{
-	Super::Serialize(Ar);
-	
-	if ( Ar.IsLoading() && Ar.UE4Ver() < VER_UE4_ASSET_IMPORT_DATA_AS_JSON && !AssetImportData)
-	{
-		// AssetImportData should always be valid
-		AssetImportData = NewObject<UAssetImportData>(this, TEXT("AssetImportData"));
-	}
-}
-#endif
-
 void UPaperTileMap::ValidateSelectedLayerIndex()
 {
 	if (!TileLayers.IsValidIndex(SelectedLayerIndex))
diff --git a/Engine/Plugins/2D/Paper2D/Source/Paper2D/Private/PaperTileMapComponent.cpp b/Engine/Plugins/2D/Paper2D/Source/Paper2D/Private/PaperTileMapComponent.cpp
--- a/Engine/Plugins/2D/Paper2D/Source/Paper2D/Private/PaperTileMapComponent.cpp
+++ b/Engine/Plugins/2D/Paper2D/Source/Paper2D/Private/PaperTileMapComponent.cpp
@@ -109,7 +109,9 @@ void UPaperTileMapComponent::PostLoad()
 {
 	Super::PostLoad();
 
-	if (GetLinkerCustomVersion(FPaperCustomVersion::GUID) < FPaperCustomVersion::MovedTileMapDataToSeparateClass)
+	const int32 PaperVer = GetLinkerCustomVersion(FPaperCustomVersion::GUID);
+	
+	if (PaperVer < FPaperCustomVersion::MovedTileMapDataToSeparateClass)
 	{
 		// Create a tile map object and move our old properties over to it
 		TileMap = NewObject<UPaperTileMap>(this);
@@ -135,6 +137,12 @@ void UPaperTileMapComponent::PostLoad()
 		Material_DEPRECATED = nullptr;
 		TileLayers_DEPRECATED.Empty();
 	}
+
+	if (PaperVer < FPaperCustomVersion::FixVertexColorSpace)
+	{
+		const FColor SRGBColor = TileMapColor.ToFColor(/*bSRGB=*/ true);
+		TileMapColor = SRGBColor.ReinterpretAsLinear();
+	}
 }
 
 #if WITH_EDITOR
@@ -297,7 +305,8 @@ void UPaperTileMapComponent::RebuildRenderData(TArray<FSpriteRenderSection>& Sec
 			}
 		}
 
-		const FColor DrawColor(( TileMapColor * Layer->GetLayerColor() ).ToFColor(true));
+		const FLinearColor DrawColorLinear = TileMapColor * Layer->GetLayerColor();
+		const FColor DrawColor(DrawColorLinear.ToFColor(/*bSRGB=*/ false));
 
 #if WITH_EDITORONLY_DATA
 		if (!Layer->ShouldRenderInEditor())
diff --git a/Engine/Plugins/2D/Paper2D/Source/Paper2D/Private/Terrain/PaperTerrainComponent.cpp b/Engine/Plugins/2D/Paper2D/Source/Paper2D/Private/Terrain/PaperTerrainComponent.cpp
--- a/Engine/Plugins/2D/Paper2D/Source/Paper2D/Private/Terrain/PaperTerrainComponent.cpp
+++ b/Engine/Plugins/2D/Paper2D/Source/Paper2D/Private/Terrain/PaperTerrainComponent.cpp
@@ -2,6 +2,7 @@
 
 #include "Paper2DPrivatePCH.h"
 #include "Components/SplineComponent.h"
+#include "PaperCustomVersion.h"
 
 #include "PaperRenderSceneProxy.h"
 #include "PaperGeomTools.h"
@@ -117,6 +118,26 @@ const UObject* UPaperTerrainComponent::AdditionalStatObject() const
 	return TerrainMaterial;
 }
 
+void UPaperTerrainComponent::Serialize(FArchive& Ar)
+{
+	Super::Serialize(Ar);
+
+	Ar.UsingCustomVersion(FPaperCustomVersion::GUID);
+}
+
+void UPaperTerrainComponent::PostLoad()
+{
+	Super::PostLoad();
+
+	const int32 PaperVer = GetLinkerCustomVersion(FPaperCustomVersion::GUID);
+
+	if (PaperVer < FPaperCustomVersion::FixVertexColorSpace)
+	{
+		const FColor SRGBColor = TerrainColor.ToFColor(/*bSRGB=*/ true);
+		TerrainColor = SRGBColor.ReinterpretAsLinear();
+	}
+}
+
 #if WITH_EDITOR
 void UPaperTerrainComponent::PostEditChangeProperty(FPropertyChangedEvent& PropertyChangedEvent)
 {
@@ -637,7 +658,7 @@ void UPaperTerrainComponent::OnSplineEdited()
 				FSpriteDrawCallRecord& FillDrawCall = *new (MaterialBatch.Records) FSpriteDrawCallRecord();
 				FillDrawCall.BuildFromSprite(FillSprite);
 				FillDrawCall.RenderVerts.Empty();
-				FillDrawCall.Color = TerrainColor.ToFColor(true);
+				FillDrawCall.Color = TerrainColor.ToFColor(/*bSRGB=*/ false);
 				FillDrawCall.Destination = PaperAxisZ * 0.1f;
 
 				const FVector2D TextureSize = GetSpriteRenderDataBounds2D(FillSprite->BakedRenderData).GetSize();
@@ -818,7 +839,7 @@ void UPaperTerrainComponent::SpawnSegments(const TArray<FTerrainSegment>& Terrai
 
 				FSpriteDrawCallRecord& Record = *new (MaterialBatch.Records) FSpriteDrawCallRecord();
 				Record.BuildFromSprite(NewSprite);
-				Record.Color = TerrainColor.ToFColor(true);
+				Record.Color = TerrainColor.ToFColor(/*bSRGB=*/ false);
 
 				// Work out if the sprite is likely to be bent > X deg (folded over itself)
 				const FVector ForwardVector(1, 0, 0);
@@ -1015,14 +1036,14 @@ void UPaperTerrainComponent::SetTerrainColor(FLinearColor NewColor)
 	{
 		TerrainColor = NewColor;
 
-		FColor TerrainGammaSpaceColor = TerrainColor.ToFColor(true);
+		const FColor TerrainColorQuantized = TerrainColor.ToFColor(/*bSRGB=*/ false);
 
 		// Update the color in the game-thread copy of the render geometry
 		for (FPaperTerrainSpriteGeometry& Batch : GeneratedSpriteGeometry)
 		{
 			for (FSpriteDrawCallRecord& DrawCall : Batch.Records)
 			{
-				DrawCall.Color = TerrainGammaSpaceColor;
+				DrawCall.Color = TerrainColorQuantized;
 			}
 		}
 
diff --git a/Engine/Plugins/2D/Paper2D/Source/Paper2DEditor/Private/GroupedSprites/PaperGroupedSpriteUtilities.cpp b/Engine/Plugins/2D/Paper2D/Source/Paper2DEditor/Private/GroupedSprites/PaperGroupedSpriteUtilities.cpp
--- a/Engine/Plugins/2D/Paper2D/Source/Paper2DEditor/Private/GroupedSprites/PaperGroupedSpriteUtilities.cpp
+++ b/Engine/Plugins/2D/Paper2D/Source/Paper2DEditor/Private/GroupedSprites/PaperGroupedSpriteUtilities.cpp
@@ -108,7 +108,7 @@ void FPaperGroupedSpriteUtilities::SplitSprites(const TArray<UObject*>& InObject
 								FComponentReregisterContext ReregisterContext(SpawnedSpriteComponent);
 
 								SpawnedSpriteComponent->Modify();
-								SpawnedSpriteComponent->SetSpriteColor(InstanceData.VertexColor);
+								SpawnedSpriteComponent->SetSpriteColor(InstanceData.VertexColor.ReinterpretAsLinear());
 								SpawnedSpriteComponent->SetSprite(InstanceData.SourceSprite);
 
 								// Apply the material override if there is one
