Fixed non-unity builds

diff --git a/Engine/Plugins/2D/Paper2D/Source/Paper2D/Private/PaperTileMap.cpp b/Engine/Plugins/2D/Paper2D/Source/Paper2D/Private/PaperTileMap.cpp
--- a/Engine/Plugins/2D/Paper2D/Source/Paper2D/Private/PaperTileMap.cpp
+++ b/Engine/Plugins/2D/Paper2D/Source/Paper2D/Private/PaperTileMap.cpp
@@ -3,6 +3,7 @@
 #include "Paper2DPrivatePCH.h"
 #include "PhysicsEngine/BodySetup.h"
 #include "PhysicsEngine/BodySetup2D.h"
+#include "PaperCustomVersion.h"
 #include "PaperTileMap.h"
 #include "PaperTileLayer.h"
 #include "PaperRuntimeSettings.h"
