Paper2D: Exposed a way to set collision thickness / offset for tile maps and tile layers, and added a RebuildCollision method for use after SetTile calls

diff --git a/Engine/Plugins/2D/Paper2D/Source/Paper2D/Classes/PaperTileLayer.h b/Engine/Plugins/2D/Paper2D/Source/Paper2D/Classes/PaperTileLayer.h
--- a/Engine/Plugins/2D/Paper2D/Source/Paper2D/Classes/PaperTileLayer.h
+++ b/Engine/Plugins/2D/Paper2D/Source/Paper2D/Classes/PaperTileLayer.h
@@ -244,6 +244,20 @@ public:
 		return AllocatedCells.GetData();
 	}
 
+	void SetLayerCollides(bool bShouldCollide) { bLayerCollides = bShouldCollide; }
+
+	void SetLayerCollisionThickness(bool bShouldOverride, float OverrideValue)
+	{
+		bOverrideCollisionThickness = bShouldOverride;
+		CollisionThicknessOverride = OverrideValue;
+	}
+
+	void SetLayerCollisionOffset(bool bShouldOverride, float OverrideValue)
+	{
+		bOverrideCollisionOffset = bShouldOverride;
+		CollisionOffsetOverride = OverrideValue;
+	}
+
 protected:
 	void ReallocateAndCopyMap();
 
diff --git a/Engine/Plugins/2D/Paper2D/Source/Paper2D/Classes/PaperTileMap.h b/Engine/Plugins/2D/Paper2D/Source/Paper2D/Classes/PaperTileMap.h
--- a/Engine/Plugins/2D/Paper2D/Source/Paper2D/Classes/PaperTileMap.h
+++ b/Engine/Plugins/2D/Paper2D/Source/Paper2D/Classes/PaperTileMap.h
@@ -166,6 +166,12 @@ public:
 		return SpriteCollisionDomain;
 	}
 
+	// Sets the collision thickness
+	void SetCollisionThickness(float Thickness = 50.0f);
+
+	// Sets the collision domain
+	void SetCollisionDomain(ESpriteCollisionMode::Type Domain);
+
 	FBoxSphereBounds GetRenderBounds() const;
 
 	// Creates and adds a new layer and returns it
diff --git a/Engine/Plugins/2D/Paper2D/Source/Paper2D/Classes/PaperTileMapComponent.h b/Engine/Plugins/2D/Paper2D/Source/Paper2D/Classes/PaperTileMapComponent.h
--- a/Engine/Plugins/2D/Paper2D/Source/Paper2D/Classes/PaperTileMapComponent.h
+++ b/Engine/Plugins/2D/Paper2D/Source/Paper2D/Classes/PaperTileMapComponent.h
@@ -153,7 +153,8 @@ public:
 	UFUNCTION(BlueprintPure, Category="Sprite", meta=(Layer="0"))
 	FPaperTileInfo GetTile(int32 X, int32 Y, int32 Layer) const;
 
-	// Modifies the contents of a specified tile cell (Note: This will only work on components that own their own tile map (OwnsTileMap returns true), you cannot modify standalone tile map assets) 
+	// Modifies the contents of a specified tile cell (Note: This will only work on components that own their own tile map (OwnsTileMap returns true), you cannot modify standalone tile map assets)
+	// Note: Does not update collision by default, call RebuildCollision after all edits have been done in a frame if necessary
 	UFUNCTION(BlueprintCallable, Category="Sprite", meta=(Layer="0"))
 	void SetTile(int32 X, int32 Y, int32 Layer, FPaperTileInfo NewValue);
 
@@ -203,6 +204,20 @@ public:
 	UFUNCTION(BlueprintPure, Category="Sprite")
 	void GetTilePolygon(int32 TileX, int32 TileY, TArray<FVector>& Points, int32 LayerIndex = 0, bool bWorldSpace = false) const;
 
+	// Sets the default thickness for any layers that don't override the collision thickness
+	// Note: This will only work on components that own their own tile map (OwnsTileMap returns true), you cannot modify standalone tile map assets
+	UFUNCTION(BlueprintCallable, Category="Sprite")
+	void SetDefaultCollisionThickness(float Thickness, bool bRebuildCollision = true);
+
+	// Sets the collision thickness for a specific layer
+	// Note: This will only work on components that own their own tile map (OwnsTileMap returns true), you cannot modify standalone tile map assets
+	UFUNCTION(BlueprintCallable, Category="Sprite")
+	void SetLayerCollision(int32 Layer = 0, bool bHasCollision = true, bool bOverrideThickness = true, float CustomThickness = 50.0f, bool bOverrideOffset = false, float CustomOffset = 0.0f, bool bRebuildCollision = true);
+
+	// Rebuilds collision for the tile map
+	UFUNCTION(BlueprintCallable, Category = "Sprite")
+	void RebuildCollision();
+
 #if WITH_EDITOR
 	// Returns the rendering stats for this component
 	void GetRenderingStats(int32& OutNumTriangles, int32& OutNumBatches) const;
diff --git a/Engine/Plugins/2D/Paper2D/Source/Paper2D/Private/PaperTileMap.cpp b/Engine/Plugins/2D/Paper2D/Source/Paper2D/Private/PaperTileMap.cpp
--- a/Engine/Plugins/2D/Paper2D/Source/Paper2D/Private/PaperTileMap.cpp
+++ b/Engine/Plugins/2D/Paper2D/Source/Paper2D/Private/PaperTileMap.cpp
@@ -515,6 +515,16 @@ FVector UPaperTileMap::GetTileCenterInLocalSpace(float TileX, float TileY, int32
 	}
 }
 
+void UPaperTileMap::SetCollisionThickness(float Thickness)
+{
+	CollisionThickness = Thickness;
+}
+
+void UPaperTileMap::SetCollisionDomain(ESpriteCollisionMode::Type Domain)
+{
+	SpriteCollisionDomain = Domain;
+}
+
 FBoxSphereBounds UPaperTileMap::GetRenderBounds() const
 {
 	const float Depth = SeparationPerLayer * (TileLayers.Num() - 1);
diff --git a/Engine/Plugins/2D/Paper2D/Source/Paper2D/Private/PaperTileMapComponent.cpp b/Engine/Plugins/2D/Paper2D/Source/Paper2D/Private/PaperTileMapComponent.cpp
--- a/Engine/Plugins/2D/Paper2D/Source/Paper2D/Private/PaperTileMapComponent.cpp
+++ b/Engine/Plugins/2D/Paper2D/Source/Paper2D/Private/PaperTileMapComponent.cpp
@@ -697,6 +697,51 @@ void UPaperTileMapComponent::GetTilePolygon(int32 TileX, int32 TileY, TArray<FVe
 	}
 }
 
+void UPaperTileMapComponent::SetDefaultCollisionThickness(float Thickness, bool bRebuildCollision)
+{
+	if (OwnsTileMap())
+	{
+		TileMap->SetCollisionThickness(Thickness);
+
+		if (bRebuildCollision)
+		{
+			RecreatePhysicsState();
+			UpdateBounds();
+		}
+	}
+}
+
+void UPaperTileMapComponent::SetLayerCollision(int32 Layer, bool bHasCollision, bool bOverrideThickness, float CustomThickness, bool bOverrideOffset, float CustomOffset, bool bRebuildCollision)
+{
+	if (OwnsTileMap())
+	{
+		if (TileMap->TileLayers.IsValidIndex(Layer))
+		{
+			UPaperTileLayer* TileLayer = TileMap->TileLayers[Layer];
+
+			TileLayer->SetLayerCollides(bHasCollision);
+			TileLayer->SetLayerCollisionThickness(bOverrideThickness, CustomThickness);
+			TileLayer->SetLayerCollisionOffset(bOverrideOffset, CustomThickness);
+
+			if (bRebuildCollision)
+			{
+				RecreatePhysicsState();
+				UpdateBounds();
+			}
+		}
+		else
+		{
+			UE_LOG(LogPaper2D, Warning, TEXT("Invalid layer index %d for %s"), Layer, *TileMap->GetPathName());
+		}
+	}
+}
+
+void UPaperTileMapComponent::RebuildCollision()
+{
+	RecreatePhysicsState();
+	UpdateBounds();
+}
+
 //////////////////////////////////////////////////////////////////////////
 
 #undef LOCTEXT_NAMESPACE
\ No newline at end of file
