Paper2D: Removed old non-functional batching prototype (avoids a per-world tick)

diff --git a/Engine/Plugins/2D/Paper2D/Source/Paper2D/Private/Paper2DModule.cpp b/Engine/Plugins/2D/Paper2D/Source/Paper2D/Private/Paper2DModule.cpp
--- a/Engine/Plugins/2D/Paper2D/Source/Paper2D/Private/Paper2DModule.cpp
+++ b/Engine/Plugins/2D/Paper2D/Source/Paper2D/Private/Paper2DModule.cpp
@@ -3,8 +3,6 @@
 #include "Paper2DPrivatePCH.h"
 #include "Paper2DModule.h"
 
-#include "Rendering/PaperBatchManager.h"
-
 DEFINE_STAT(STAT_PaperRender_SetSpriteRT);
 
 //////////////////////////////////////////////////////////////////////////
@@ -27,13 +25,10 @@ public:
 		}
 
 		PaperAxisZ = FVector::CrossProduct(PaperAxisX, PaperAxisY);
-
-		FPaperBatchManager::Initialize();
 	}
 
 	virtual void ShutdownModule() override
 	{
-		FPaperBatchManager::Shutdown();
 	}
 };
 
diff --git a/Engine/Plugins/2D/Paper2D/Source/Paper2D/Private/PaperRenderSceneProxy.cpp b/Engine/Plugins/2D/Paper2D/Source/Paper2D/Private/PaperRenderSceneProxy.cpp
--- a/Engine/Plugins/2D/Paper2D/Source/Paper2D/Private/PaperRenderSceneProxy.cpp
+++ b/Engine/Plugins/2D/Paper2D/Source/Paper2D/Private/PaperRenderSceneProxy.cpp
@@ -3,8 +3,6 @@
 #include "Paper2DPrivatePCH.h"
 #include "PaperRenderSceneProxy.h"
 #include "PhysicsEngine/BodySetup.h"
-#include "Rendering/PaperBatchManager.h"
-#include "Rendering/PaperBatchSceneProxy.h"
 #include "PhysicsEngine/BodySetup2D.h"
 #include "LocalVertexFactory.h"
 #include "MeshBatch.h"
@@ -282,27 +280,10 @@ FPaperRenderSceneProxy::FPaperRenderSceneProxy(const UPrimitiveComponent* InComp
 
 FPaperRenderSceneProxy::~FPaperRenderSceneProxy()
 {
-#if TEST_BATCHING
-	if (FPaperBatchSceneProxy* Batcher = FPaperBatchManager::GetBatcher(GetScene()))
-	{
-		Batcher->UnregisterManagedProxy(this);
-	}
-#endif
-
 	VertexBuffer.ReleaseResource();
 	MyVertexFactory.ReleaseResource();
 }
 
-void FPaperRenderSceneProxy::CreateRenderThreadResources()
-{
-#if TEST_BATCHING
-	if (FPaperBatchSceneProxy* Batcher = FPaperBatchManager::GetBatcher(GetScene()))
-	{
-		Batcher->RegisterManagedProxy(this);
-	}
-#endif
-}
-
 void FPaperRenderSceneProxy::DebugDrawBodySetup(const FSceneView* View, int32 ViewIndex, FMeshElementCollector& Collector, UBodySetup* BodySetup, const FMatrix& GeomTransformMatrix, const FLinearColor& CollisionColor, bool bDrawSolid) const
 {
 	if (FMath::Abs(GeomTransformMatrix.Determinant()) < SMALL_NUMBER)
@@ -371,10 +352,7 @@ void FPaperRenderSceneProxy::GetDynamicMeshElements(const TArray<const FSceneVie
 			{
 				const FSceneView* View = Views[ViewIndex];
 
-#if TEST_BATCHING
-#else
 				GetDynamicMeshElementsForView(View, ViewIndex, Collector);
-#endif
 			}
 		}
 	}
diff --git a/Engine/Plugins/2D/Paper2D/Source/Paper2D/Private/PaperRenderSceneProxy.h b/Engine/Plugins/2D/Paper2D/Source/Paper2D/Private/PaperRenderSceneProxy.h
--- a/Engine/Plugins/2D/Paper2D/Source/Paper2D/Private/PaperRenderSceneProxy.h
+++ b/Engine/Plugins/2D/Paper2D/Source/Paper2D/Private/PaperRenderSceneProxy.h
@@ -5,8 +5,6 @@
 #include "PaperSprite.h"
 #include "SpriteDrawCall.h"
 
-#define TEST_BATCHING 0
-
 #if WITH_EDITOR
 typedef TMap<const UTexture*, const UTexture*> FPaperRenderSceneProxyTextureOverrideMap;
 
@@ -150,7 +148,6 @@ public:
 	virtual FPrimitiveViewRelevance GetViewRelevance(const FSceneView* View) override;
 	virtual uint32 GetMemoryFootprint() const override;
 	virtual bool CanBeOccluded() const override;
-	virtual void CreateRenderThreadResources() override;
 	// End of FPrimitiveSceneProxy interface.
 
 	void SetDrawCall_RenderThread(const FSpriteDrawCallRecord& NewDynamicData);
diff --git a/Engine/Plugins/2D/Paper2D/Source/Paper2D/Private/Rendering/PaperBatchComponent.cpp b/Engine/Plugins/2D/Paper2D/Source/Paper2D/Private/Rendering/PaperBatchComponent.cpp
deleted file mode 100644
diff --git a/Engine/Plugins/2D/Paper2D/Source/Paper2D/Private/Rendering/PaperBatchComponent.h b/Engine/Plugins/2D/Paper2D/Source/Paper2D/Private/Rendering/PaperBatchComponent.h
deleted file mode 100644
diff --git a/Engine/Plugins/2D/Paper2D/Source/Paper2D/Private/Rendering/PaperBatchManager.cpp b/Engine/Plugins/2D/Paper2D/Source/Paper2D/Private/Rendering/PaperBatchManager.cpp
deleted file mode 100644
diff --git a/Engine/Plugins/2D/Paper2D/Source/Paper2D/Private/Rendering/PaperBatchManager.h b/Engine/Plugins/2D/Paper2D/Source/Paper2D/Private/Rendering/PaperBatchManager.h
deleted file mode 100644
diff --git a/Engine/Plugins/2D/Paper2D/Source/Paper2D/Private/Rendering/PaperBatchSceneProxy.cpp b/Engine/Plugins/2D/Paper2D/Source/Paper2D/Private/Rendering/PaperBatchSceneProxy.cpp
deleted file mode 100644
diff --git a/Engine/Plugins/2D/Paper2D/Source/Paper2D/Private/Rendering/PaperBatchSceneProxy.h b/Engine/Plugins/2D/Paper2D/Source/Paper2D/Private/Rendering/PaperBatchSceneProxy.h
deleted file mode 100644
