Fix to preserve the AssetImportData when loading an asset.

diff --git a/Engine/Plugins/2D/Paper2D/Source/Paper2D/Classes/PaperTileMap.h b/Engine/Plugins/2D/Paper2D/Source/Paper2D/Classes/PaperTileMap.h
--- a/Engine/Plugins/2D/Paper2D/Source/Paper2D/Classes/PaperTileMap.h
+++ b/Engine/Plugins/2D/Paper2D/Source/Paper2D/Classes/PaperTileMap.h
@@ -127,6 +127,7 @@ public:
 public:
 	// UObject interface
 	virtual void Serialize(FArchive& Ar) override;
+	virtual void PostInitProperties() override;
 	virtual void PostLoad() override;
 #if WITH_EDITOR
 	virtual void PreEditChange(UProperty* PropertyAboutToChange) override;
diff --git a/Engine/Plugins/2D/Paper2D/Source/Paper2D/Private/PaperTileMap.cpp b/Engine/Plugins/2D/Paper2D/Source/Paper2D/Private/PaperTileMap.cpp
--- a/Engine/Plugins/2D/Paper2D/Source/Paper2D/Private/PaperTileMap.cpp
+++ b/Engine/Plugins/2D/Paper2D/Source/Paper2D/Private/PaperTileMap.cpp
@@ -30,7 +30,6 @@ UPaperTileMap::UPaperTileMap(const FObjectInitializer& ObjectInitializer)
 #if WITH_EDITORONLY_DATA
 	SelectedLayerIndex = INDEX_NONE;
 	BackgroundColor = FColor(55, 55, 55);
-	AssetImportData = CreateEditorOnlyDefaultSubobject<UAssetImportData>(TEXT("AssetImportData"));
 #endif
 
 	LayerNameIndex = 0;
@@ -39,6 +38,17 @@ UPaperTileMap::UPaperTileMap(const FObjectInitializer& ObjectInitializer)
 	Material = DefaultMaterial.Object;
 }
 
+void UPaperTileMap::PostInitProperties()
+{
+#if WITH_EDITORONLY_DATA
+	if (!HasAnyFlags(RF_ClassDefaultObject))
+	{
+		AssetImportData = NewObject<UAssetImportData>(this, TEXT("AssetImportData"));
+	}
+#endif
+	Super::PostInitProperties();
+}
+
 #if WITH_EDITOR
 void UPaperTileMap::PreEditChange(UProperty* PropertyAboutToChange)
 {
diff --git a/Engine/Plugins/2D/Paper2D/Source/PaperSpriteSheetImporter/Private/PaperSpriteSheet.cpp b/Engine/Plugins/2D/Paper2D/Source/PaperSpriteSheetImporter/Private/PaperSpriteSheet.cpp
--- a/Engine/Plugins/2D/Paper2D/Source/PaperSpriteSheetImporter/Private/PaperSpriteSheet.cpp
+++ b/Engine/Plugins/2D/Paper2D/Source/PaperSpriteSheetImporter/Private/PaperSpriteSheet.cpp
@@ -6,9 +6,17 @@
 UPaperSpriteSheet::UPaperSpriteSheet(const FObjectInitializer& ObjectInitializer)
 	: Super(ObjectInitializer)
 {
+}
+
+void UPaperSpriteSheet::PostInitProperties()
+{
 #if WITH_EDITORONLY_DATA
-	AssetImportData = CreateEditorOnlyDefaultSubobject<UAssetImportData>(TEXT("AssetImportData"));
+	if (!HasAnyFlags(RF_ClassDefaultObject))
+	{
+		AssetImportData = NewObject<UAssetImportData>(this, TEXT("AssetImportData"));
+	}
 #endif
+	Super::PostInitProperties();
 }
 
 #if WITH_EDITORONLY_DATA
diff --git a/Engine/Plugins/2D/Paper2D/Source/PaperSpriteSheetImporter/Private/PaperSpriteSheet.h b/Engine/Plugins/2D/Paper2D/Source/PaperSpriteSheetImporter/Private/PaperSpriteSheet.h
--- a/Engine/Plugins/2D/Paper2D/Source/PaperSpriteSheetImporter/Private/PaperSpriteSheet.h
+++ b/Engine/Plugins/2D/Paper2D/Source/PaperSpriteSheetImporter/Private/PaperSpriteSheet.h
@@ -39,6 +39,7 @@ public:
 	class UAssetImportData* AssetImportData;
 
 	// UObject interface
+	virtual void PostInitProperties() override;
 	virtual void GetAssetRegistryTags(TArray<FAssetRegistryTag>& OutTags) const override;
 	virtual void Serialize(FArchive& Ar) override;
 	// End of UObject interface
diff --git a/Engine/Source/Runtime/Engine/Classes/Animation/AnimSequence.h b/Engine/Source/Runtime/Engine/Classes/Animation/AnimSequence.h
--- a/Engine/Source/Runtime/Engine/Classes/Animation/AnimSequence.h
+++ b/Engine/Source/Runtime/Engine/Classes/Animation/AnimSequence.h
@@ -528,6 +528,7 @@ class ENGINE_API UAnimSequence : public UAnimSequenceBase
 public:
 	// Begin UObject interface
 	virtual void Serialize(FArchive& Ar) override;
+	virtual void PostInitProperties() override;
 	virtual void PostLoad() override;
 	virtual void PreSave() override;
 #if WITH_EDITOR
diff --git a/Engine/Source/Runtime/Engine/Classes/Engine/SkeletalMesh.h b/Engine/Source/Runtime/Engine/Classes/Engine/SkeletalMesh.h
--- a/Engine/Source/Runtime/Engine/Classes/Engine/SkeletalMesh.h
+++ b/Engine/Source/Runtime/Engine/Classes/Engine/SkeletalMesh.h
@@ -740,7 +740,8 @@ public:
 	virtual bool IsReadyForFinishDestroy() override;
 	virtual void PreSave() override;
 	virtual void Serialize(FArchive& Ar) override;
-	virtual void PostLoad() override;	
+	virtual void PostInitProperties() override;
+	virtual void PostLoad() override;
 	virtual void GetAssetRegistryTags(TArray<FAssetRegistryTag>& OutTags) const override;
 	virtual FString GetDesc() override;
 	virtual FString GetDetailedInfoInternal() const override;
diff --git a/Engine/Source/Runtime/Engine/Classes/Engine/StaticMesh.h b/Engine/Source/Runtime/Engine/Classes/Engine/StaticMesh.h
--- a/Engine/Source/Runtime/Engine/Classes/Engine/StaticMesh.h
+++ b/Engine/Source/Runtime/Engine/Classes/Engine/StaticMesh.h
@@ -458,6 +458,7 @@ public:
 	ENGINE_API virtual void GetAssetRegistryTagMetadata(TMap<FName, FAssetRegistryTagMetadata>& OutMetadata) const override;
 #endif // WITH_EDITOR
 	ENGINE_API virtual void Serialize(FArchive& Ar) override;
+	ENGINE_API virtual void PostInitProperties() override;
 	ENGINE_API virtual void PostLoad() override;
 	ENGINE_API virtual void BeginDestroy() override;
 	ENGINE_API virtual bool IsReadyForFinishDestroy() override;
diff --git a/Engine/Source/Runtime/Engine/Private/Animation/AnimSequence.cpp b/Engine/Source/Runtime/Engine/Private/Animation/AnimSequence.cpp
--- a/Engine/Source/Runtime/Engine/Private/Animation/AnimSequence.cpp
+++ b/Engine/Source/Runtime/Engine/Private/Animation/AnimSequence.cpp
@@ -113,11 +113,18 @@ UAnimSequence::UAnimSequence(const FObjectInitializer& ObjectInitializer)
 	, RootMotionRootLock(ERootMotionRootLock::RefPose)
 	, bRootMotionSettingsCopiedFromMontage(false)
 {
-
 	RateScale = 1.0;
+}
+
+void UAnimSequence::PostInitProperties()
+{
 #if WITH_EDITORONLY_DATA
-	AssetImportData = CreateEditorOnlyDefaultSubobject<UAssetImportData>(TEXT("AssetImportData"));
+	if (!HasAnyFlags(RF_ClassDefaultObject))
+	{
+		AssetImportData = NewObject<UAssetImportData>(this, TEXT("AssetImportData"));
+	}
 #endif
+	Super::PostInitProperties();
 }
 
 SIZE_T UAnimSequence::GetResourceSize(EResourceSizeMode::Type Mode)
diff --git a/Engine/Source/Runtime/Engine/Private/SkeletalMesh.cpp b/Engine/Source/Runtime/Engine/Private/SkeletalMesh.cpp
--- a/Engine/Source/Runtime/Engine/Private/SkeletalMesh.cpp
+++ b/Engine/Source/Runtime/Engine/Private/SkeletalMesh.cpp
@@ -1999,12 +1999,23 @@ USkeletalMesh::USkeletalMesh(const FObjectInitializer& ObjectInitializer)
 	StreamingDistanceMultiplier = 1.0f;
 #if WITH_EDITORONLY_DATA
 	SelectedEditorSection = INDEX_NONE;
-	AssetImportData = CreateEditorOnlyDefaultSubobject<UAssetImportData>(TEXT("AssetImportData"));
 #endif
 	ImportedResource = MakeShareable(new FSkeletalMeshResource());
 }
 
 
+void USkeletalMesh::PostInitProperties()
+{
+#if WITH_EDITORONLY_DATA
+	if (!HasAnyFlags(RF_ClassDefaultObject))
+	{
+		AssetImportData = NewObject<UAssetImportData>(this, TEXT("AssetImportData"));
+	}
+#endif
+	Super::PostInitProperties();
+}
+
+
 void USkeletalMesh::InitResources()
 {
 	ImportedResource->InitResources(bHasVertexColors);
diff --git a/Engine/Source/Runtime/Engine/Private/StaticMesh.cpp b/Engine/Source/Runtime/Engine/Private/StaticMesh.cpp
--- a/Engine/Source/Runtime/Engine/Private/StaticMesh.cpp
+++ b/Engine/Source/Runtime/Engine/Private/StaticMesh.cpp
@@ -1242,13 +1242,23 @@ UStaticMesh::UStaticMesh(const FObjectInitializer& ObjectInitializer)
 #if WITH_EDITORONLY_DATA
 	AutoLODPixelError = 1.0f;
 	bAutoComputeLODScreenSize=true;
-	AssetImportData = CreateEditorOnlyDefaultSubobject<UAssetImportData>(TEXT("AssetImportData"));
 #endif // #if WITH_EDITORONLY_DATA
 	LightMapResolution = 4;
 	LpvBiasMultiplier = 1.0f;
 	MinLOD = 0;
 }
 
+void UStaticMesh::PostInitProperties()
+{
+#if WITH_EDITORONLY_DATA
+	if (!HasAnyFlags(RF_ClassDefaultObject))
+	{
+		AssetImportData = NewObject<UAssetImportData>(this, TEXT("AssetImportData"));
+	}
+#endif
+	Super::PostInitProperties();
+}
+
 /**
  * Initializes the static mesh's render resources.
  */
diff --git a/Engine/Source/Runtime/GeometryCache/Classes/GeometryCache.h b/Engine/Source/Runtime/GeometryCache/Classes/GeometryCache.h
--- a/Engine/Source/Runtime/GeometryCache/Classes/GeometryCache.h
+++ b/Engine/Source/Runtime/GeometryCache/Classes/GeometryCache.h
@@ -22,6 +22,7 @@ public:
 
 	// Begin UObject interface.
 	virtual void Serialize(FArchive& Ar) override;
+	virtual void PostInitProperties() override;
 	virtual FString GetDesc() override;
 	virtual SIZE_T GetResourceSize(EResourceSizeMode::Type Mode) override;
 	virtual void GetAssetRegistryTags(TArray<FAssetRegistryTag>& OutTags) const override;
diff --git a/Engine/Source/Runtime/GeometryCache/Private/GeometryCache.cpp b/Engine/Source/Runtime/GeometryCache/Private/GeometryCache.cpp
--- a/Engine/Source/Runtime/GeometryCache/Private/GeometryCache.cpp
+++ b/Engine/Source/Runtime/GeometryCache/Private/GeometryCache.cpp
@@ -10,6 +10,17 @@ UGeometryCache::UGeometryCache(const FObjectInitializer& ObjectInitializer /*= F
 	NumVertexAnimationTracks = 0;
 }
 
+void UGeometryCache::PostInitProperties()
+{
+#if WITH_EDITORONLY_DATA
+	if (!HasAnyFlags(RF_ClassDefaultObject))
+	{
+		AssetImportData = NewObject<UAssetImportData>(this, TEXT("AssetImportData"));
+	}
+#endif
+	Super::PostInitProperties();
+}
+
 void UGeometryCache::Serialize(FArchive& Ar)
 {
 	// Custom serialization
