Paper2D sprite importer improvements. Fixed issue with rotation on reimport. Also decreased the amount of rebuilding data on reimport.

diff --git a/Engine/Plugins/2D/Paper2D/Source/Paper2D/Classes/PaperSprite.h b/Engine/Plugins/2D/Paper2D/Source/Paper2D/Classes/PaperSprite.h
--- a/Engine/Plugins/2D/Paper2D/Source/Paper2D/Classes/PaperSprite.h
+++ b/Engine/Plugins/2D/Paper2D/Source/Paper2D/Classes/PaperSprite.h
@@ -205,6 +205,8 @@ public:
 	void RescaleSpriteData(UTexture2D* Texture);
 	bool NeedRescaleSpriteData();
 
+	/** This is a generic "rebuild all" function that calls RebuildCollisionData() and then RebuildRenderData(). */
+	void RebuildData();
 	void RebuildCollisionData();
 	void RebuildRenderData();
 
@@ -221,11 +223,11 @@ public:
 	void CreatePolygonFromBoundingBox(FSpriteGeometryCollection& GeomOwner, bool bUseTightBounds);
 
 	// Reinitializes this sprite (NOTE: Does not register existing components in the world)
-	void InitializeSprite(const FSpriteAssetInitParameters& InitParams);
+	void InitializeSprite(const FSpriteAssetInitParameters& InitParams, bool bRebuildData = true);
 
-	void SetTrim(bool bTrimmed, const FVector2D& OriginInSourceImage, const FVector2D& SourceImageDimension);
-	void SetRotated(bool bRotated);
-	void SetPivotMode(ESpritePivotMode::Type PivotMode, FVector2D CustomTextureSpacePivot);
+	void SetTrim(bool bTrimmed, const FVector2D& OriginInSourceImage, const FVector2D& SourceImageDimension, bool bRebuildData = true);
+	void SetRotated(bool bRotated, bool bRebuildData = true);
+	void SetPivotMode(ESpritePivotMode::Type PivotMode, FVector2D CustomTextureSpacePivot, bool bRebuildData = true);
 	
 	// Returns the Origin within SourceImage, prior to atlasing
 	FVector2D GetOriginInSourceImageBeforeTrimming() const { return OriginInSourceImageBeforeTrimming; }
diff --git a/Engine/Plugins/2D/Paper2D/Source/Paper2D/Private/PaperSprite.cpp b/Engine/Plugins/2D/Paper2D/Source/Paper2D/Private/PaperSprite.cpp
--- a/Engine/Plugins/2D/Paper2D/Source/Paper2D/Private/PaperSprite.cpp
+++ b/Engine/Plugins/2D/Paper2D/Source/Paper2D/Private/PaperSprite.cpp
@@ -1391,7 +1391,13 @@ void UPaperSprite::ExtractRectsFromTexture(UTexture2D* Texture, TArray<FIntRect>
 	SpriteTextureBitmap.ExtractRects(/*out*/ OutRects);
 }
 
-void UPaperSprite::InitializeSprite(const FSpriteAssetInitParameters& InitParams)
+void UPaperSprite::RebuildData()
+{
+	RebuildCollisionData();
+	RebuildRenderData();
+}
+
+void UPaperSprite::InitializeSprite(const FSpriteAssetInitParameters& InitParams, bool bRebuildData /*= true*/)
 {
 	if (InitParams.bOverridePixelsPerUnrealUnit)
 	{
@@ -1422,32 +1428,40 @@ void UPaperSprite::InitializeSprite(const FSpriteAssetInitParameters& InitParams
 	SourceUV = InitParams.Offset;
 	SourceDimension = InitParams.Dimension;
 
-	RebuildCollisionData();
-	RebuildRenderData();
+	if (bRebuildData)
+	{
+		RebuildData();
+	}
 }
 
-void UPaperSprite::SetTrim(bool bTrimmed, const FVector2D& OriginInSourceImage, const FVector2D& SourceImageDimension)
+void UPaperSprite::SetTrim(bool bTrimmed, const FVector2D& OriginInSourceImage, const FVector2D& SourceImageDimension, bool bRebuildData /*= true*/)
 {
 	this->bTrimmedInSourceImage = bTrimmed;
 	this->OriginInSourceImageBeforeTrimming = OriginInSourceImage;
 	this->SourceImageDimensionBeforeTrimming = SourceImageDimension;
-	RebuildRenderData();
-	RebuildCollisionData();
+	if (bRebuildData)
+	{
+		RebuildData();
+	}
 }
 
-void UPaperSprite::SetRotated(bool bRotated)
+void UPaperSprite::SetRotated(bool bRotated, bool bRebuildData /*= true*/)
 {
 	this->bRotatedInSourceImage = bRotated;
-	RebuildRenderData();
-	RebuildCollisionData();
+	if (bRebuildData)
+	{
+		RebuildData();
+	}
 }
 
-void UPaperSprite::SetPivotMode(ESpritePivotMode::Type InPivotMode, FVector2D InCustomTextureSpacePivot)
+void UPaperSprite::SetPivotMode(ESpritePivotMode::Type InPivotMode, FVector2D InCustomTextureSpacePivot, bool bRebuildData /*= true*/)
 {
 	PivotMode = InPivotMode;
 	CustomPivotPoint = InCustomTextureSpacePivot;
-	RebuildRenderData();
-	RebuildCollisionData();
+	if (bRebuildData)
+	{
+		RebuildData();
+	}
 }
 
 FVector2D UPaperSprite::ConvertTextureSpaceToPivotSpace(FVector2D Input) const
diff --git a/Engine/Plugins/2D/Paper2D/Source/PaperSpriteSheetImporter/Private/PaperJsonSpriteSheetImporter.cpp b/Engine/Plugins/2D/Paper2D/Source/PaperSpriteSheetImporter/Private/PaperJsonSpriteSheetImporter.cpp
--- a/Engine/Plugins/2D/Paper2D/Source/PaperSpriteSheetImporter/Private/PaperJsonSpriteSheetImporter.cpp
+++ b/Engine/Plugins/2D/Paper2D/Source/PaperSpriteSheetImporter/Private/PaperJsonSpriteSheetImporter.cpp
@@ -612,17 +612,10 @@ bool FPaperJsonSpriteSheetImporter::PerformImport(const FString& LongPackagePath
 
 		GetDefault<UPaperImporterSettings>()->ApplySettingsForSpriteInit(/*inout*/ SpriteInitParams, (NormalMapTexture != nullptr) ? ESpriteInitMaterialLightingMode::ForceLit : ESpriteInitMaterialLightingMode::Automatic);
 
-		TargetSprite->InitializeSprite(SpriteInitParams);
+		TargetSprite->InitializeSprite(SpriteInitParams, false);
 
-		if (Frame.bRotated)
-		{
-			TargetSprite->SetRotated(true);
-		}
-
-		if (Frame.bTrimmed)
-		{
-			TargetSprite->SetTrim(Frame.bTrimmed, Frame.SpriteSourcePos, Frame.ImageSourceSize);
-		}
+		TargetSprite->SetRotated(Frame.bRotated, false);
+		TargetSprite->SetTrim(Frame.bTrimmed, Frame.SpriteSourcePos, Frame.ImageSourceSize, false);
 
 		// Set up pivot on object based on Texture Packer json
 		ESpritePivotMode::Type PivotType = GetBestPivotType(Frame.Pivot);
@@ -632,7 +625,9 @@ bool FPaperJsonSpriteSheetImporter::PerformImport(const FString& LongPackagePath
 			TextureSpacePivotPoint.X = Frame.SpritePosInSheet.X - Frame.SpriteSourcePos.X + Frame.ImageSourceSize.X * Frame.Pivot.X;
 			TextureSpacePivotPoint.Y = Frame.SpritePosInSheet.Y - Frame.SpriteSourcePos.Y + Frame.ImageSourceSize.Y * Frame.Pivot.Y;
 		}
-		TargetSprite->SetPivotMode(PivotType, TextureSpacePivotPoint);
+		TargetSprite->SetPivotMode(PivotType, TextureSpacePivotPoint, false);
+
+		TargetSprite->RebuildData();
 
 		// Create the entry in the animation
 		SpriteSheet->SpriteNames.Add(Frame.FrameName.ToString());
