Fixed non-unity build.

diff --git a/Engine/Plugins/2D/Paper2D/Source/Paper2DEditor/Private/GroupedSprites/PaperGroupedSpriteUtilities.cpp b/Engine/Plugins/2D/Paper2D/Source/Paper2DEditor/Private/GroupedSprites/PaperGroupedSpriteUtilities.cpp
--- a/Engine/Plugins/2D/Paper2D/Source/Paper2DEditor/Private/GroupedSprites/PaperGroupedSpriteUtilities.cpp
+++ b/Engine/Plugins/2D/Paper2D/Source/Paper2DEditor/Private/GroupedSprites/PaperGroupedSpriteUtilities.cpp
@@ -4,6 +4,8 @@
 #include "AssetData.h"
 #include "PaperGroupedSpriteActor.h"
 #include "PaperGroupedSpriteUtilities.h"
+#include "PaperSpriteActor.h"
+#include "PaperSpriteComponent.h"
 #include "ScopedTransaction.h"
 #include "ILayers.h"
 #include "ComponentReregisterContext.h"
diff --git a/Engine/Source/Runtime/Online/OnlineSubsystemUtils/Private/JoinSessionCallbackProxy.cpp b/Engine/Source/Runtime/Online/OnlineSubsystemUtils/Private/JoinSessionCallbackProxy.cpp
--- a/Engine/Source/Runtime/Online/OnlineSubsystemUtils/Private/JoinSessionCallbackProxy.cpp
+++ b/Engine/Source/Runtime/Online/OnlineSubsystemUtils/Private/JoinSessionCallbackProxy.cpp
@@ -1,6 +1,7 @@
 // Copyright 1998-2015 Epic Games, Inc. All Rights Reserved.
 
 #include "OnlineSubsystemUtilsPrivatePCH.h"
+#include "FindSessionsCallbackProxy.h"
 #include "JoinSessionCallbackProxy.h"
 
 //////////////////////////////////////////////////////////////////////////
diff --git a/Engine/Source/Runtime/Online/OnlineSubsystemUtils/Private/PartyBeaconState.cpp b/Engine/Source/Runtime/Online/OnlineSubsystemUtils/Private/PartyBeaconState.cpp
--- a/Engine/Source/Runtime/Online/OnlineSubsystemUtils/Private/PartyBeaconState.cpp
+++ b/Engine/Source/Runtime/Online/OnlineSubsystemUtils/Private/PartyBeaconState.cpp
@@ -1,6 +1,7 @@
 // Copyright 1998-2015 Epic Games, Inc. All Rights Reserved.
 
 #include "OnlineSubsystemUtilsPrivatePCH.h"
+#include "OnlineBeacon.h"
 #include "PartyBeaconState.h"
 
 bool FPartyReservation::IsValid() const
