PR #4864: Use Default scale for Sprite Asset Thumbnails (Contributed by [REDACTED_PERSON])

diff --git a/Engine/Plugins/2D/Paper2D/Source/Paper2DEditor/Private/SpriteEditor/SSpriteList.cpp b/Engine/Plugins/2D/Paper2D/Source/Paper2DEditor/Private/SpriteEditor/SSpriteList.cpp
--- a/Engine/Plugins/2D/Paper2D/Source/Paper2DEditor/Private/SpriteEditor/SSpriteList.cpp
+++ b/Engine/Plugins/2D/Paper2D/Source/Paper2DEditor/Private/SpriteEditor/SSpriteList.cpp
@@ -28,7 +28,6 @@ void SSpriteList::RebuildWidget(UTexture2D* NewTextureFilter)
 	//Config.Filter.bRecursiveClasses = true;
 	Config.Filter.ClassNames.Add(UPaperSprite::StaticClass()->GetFName());
 
-	Config.ThumbnailScale = 0.0f;
 	Config.InitialAssetViewType = EAssetViewType::Tile;
 	Config.SyncToAssetsDelegates.Add(&SyncToAssetsDelegate);
 
