Paper2D: Fixed an issue where you had to press the Extract button twice when extracting sprites if the naming template was edited
PR #4863: UE-19612: Prevent OnRelease being called too early due to PreviewExtr. (Contributed by [REDACTED_PERSON])

diff --git a/Engine/Plugins/2D/Paper2D/Source/Paper2DEditor/Private/ExtractSprites/SPaperExtractSpritesDialog.cpp b/Engine/Plugins/2D/Paper2D/Source/Paper2DEditor/Private/ExtractSprites/SPaperExtractSpritesDialog.cpp
--- a/Engine/Plugins/2D/Paper2D/Source/Paper2DEditor/Private/ExtractSprites/SPaperExtractSpritesDialog.cpp
+++ b/Engine/Plugins/2D/Paper2D/Source/Paper2DEditor/Private/ExtractSprites/SPaperExtractSpritesDialog.cpp
@@ -154,6 +154,7 @@ FText SPaperExtractSpritesViewport::GetTitleText() const
 void SPaperExtractSpritesDialog::Construct(const FArguments& InArgs, UTexture2D* InSourceTexture)
 {
 	SourceTexture = InSourceTexture;
+	bDirtyExtractedSprites = false;
 
 	ExtractSpriteSettings = NewObject<UPaperExtractSpritesSettings>();
 	ExtractSpriteSettings->AddToRoot();
@@ -219,6 +220,7 @@ void SPaperExtractSpritesDialog::Construct(const FArguments& InArgs, UTexture2D*
 						SNew(SButton)
 						.ButtonStyle(FEditorStyle::Get(), "FlatButton.Success")
 						.ForegroundColor(FLinearColor::White)
+						.ContentPadding(FMargin(6, 2))
 						.Text(LOCTEXT("PaperExtractSpritesExtractButton", "Extract..."))
 						.OnClicked(this, &SPaperExtractSpritesDialog::ExtractClicked)
 					]
@@ -227,6 +229,7 @@ void SPaperExtractSpritesDialog::Construct(const FArguments& InArgs, UTexture2D*
 						SNew(SButton)
 						.ButtonStyle(FEditorStyle::Get(), "FlatButton")
 						.ForegroundColor(FLinearColor::White)
+						.ContentPadding(FMargin(6, 2))
 						.Text(LOCTEXT("PaperExtractSpritesCancelButton", "Cancel"))
 						.OnClicked(this, &SPaperExtractSpritesDialog::CancelClicked)
 					]
@@ -393,8 +396,13 @@ void SPaperExtractSpritesDialog::OnFinishedChangingProperties(const FPropertyCha
 	if (PropertyChangedEvent.Property != nullptr)
 	{
 		const FName PropertyName(PropertyChangedEvent.Property->GetFName());
-		
-		if (PropertyName != GET_MEMBER_NAME_CHECKED(UPaperExtractSpritesSettings, OutlineColor) &&
+
+		if (PropertyName == GET_MEMBER_NAME_CHECKED(UPaperExtractSpritesSettings, NamingTemplate) ||
+			PropertyName == GET_MEMBER_NAME_CHECKED(UPaperExtractSpritesSettings, NamingStartIndex))
+		{
+			bDirtyExtractedSprites = true;
+		}
+		else if (PropertyName != GET_MEMBER_NAME_CHECKED(UPaperExtractSpritesSettings, OutlineColor) &&
 			PropertyName != GET_MEMBER_NAME_CHECKED(UPaperExtractSpritesSettings, BackgroundColor) &&
 			PropertyName != GET_MEMBER_NAME_CHECKED(UPaperExtractSpritesSettings, ViewportTextureTint))
 		{
@@ -410,6 +418,12 @@ void SPaperExtractSpritesDialog::OnFinishedChangingProperties(const FPropertyCha
 
 FReply SPaperExtractSpritesDialog::ExtractClicked()
 {
+	if (bDirtyExtractedSprites)
+	{
+		bDirtyExtractedSprites = false;
+		PreviewExtractedSprites();
+	}
+
 	CreateExtractedSprites();
 
 	CloseContainingWindow();
diff --git a/Engine/Plugins/2D/Paper2D/Source/Paper2DEditor/Private/ExtractSprites/SPaperExtractSpritesDialog.h b/Engine/Plugins/2D/Paper2D/Source/Paper2DEditor/Private/ExtractSprites/SPaperExtractSpritesDialog.h
--- a/Engine/Plugins/2D/Paper2D/Source/Paper2DEditor/Private/ExtractSprites/SPaperExtractSpritesDialog.h
+++ b/Engine/Plugins/2D/Paper2D/Source/Paper2DEditor/Private/ExtractSprites/SPaperExtractSpritesDialog.h
@@ -57,7 +57,7 @@ private:
 private:
 	// Source texture to extract from
 	UTexture2D* SourceTexture;
-
+	bool bDirtyExtractedSprites;
 
 	class UPaperExtractSpritesSettings* ExtractSpriteSettings;
 	class UPaperExtractSpriteGridSettings* ExtractSpriteGridSettings;
