PR #4948: [Paper2D][bug] Collisions misplaced when mixing different tile sizes in TileSet and TileMap (Contributed by [REDACTED_PERSON])

diff --git a/Engine/Plugins/2D/Paper2D/Source/Paper2D/Private/PaperTileLayer.cpp b/Engine/Plugins/2D/Paper2D/Source/Paper2D/Private/PaperTileLayer.cpp
--- a/Engine/Plugins/2D/Paper2D/Source/Paper2D/Private/PaperTileLayer.cpp
+++ b/Engine/Plugins/2D/Paper2D/Source/Paper2D/Private/PaperTileLayer.cpp
@@ -263,6 +263,11 @@ void UPaperTileLayer::AugmentBodySetup(UBodySetup* ShapeBodySetup, float RenderS
 					break;
 					}
 
+					// Calculate offset for tile sets that don't match the tile map's size
+					const FIntPoint TileSetSize = CellInfo.TileSet->GetTileSize();
+					ProjectionOffset.X += (TileSetSize.X - TileWidth) * 0.5f;
+					ProjectionOffset.Y += (TileSetSize.Y - TileHeight) * 0.5f;
+
 					const FVector2D CellOffset(ProjectionOffset.X, ProjectionOffset.Y);
 					
 					CollisionBuilder.SetCellOffset(CellOffset, LocalTransform);
