PR #4949: [Paper2D][Bugfix] Fixes tiles being misplaced in TileMap editor when using a multi-tile brush and the TileMap and TileSet tile sizes differ (Contributed by [REDACTED_PERSON])

diff --git a/Engine/Plugins/2D/Paper2D/Source/Paper2DEditor/Private/TileMapEditing/EdModeTileMap.cpp b/Engine/Plugins/2D/Paper2D/Source/Paper2DEditor/Private/TileMapEditing/EdModeTileMap.cpp
--- a/Engine/Plugins/2D/Paper2D/Source/Paper2DEditor/Private/TileMapEditing/EdModeTileMap.cpp
+++ b/Engine/Plugins/2D/Paper2D/Source/Paper2DEditor/Private/TileMapEditing/EdModeTileMap.cpp
@@ -1110,13 +1110,29 @@ void FEdModeTileMap::SetActivePaint(UPaperTileSet* TileSet, FIntPoint TopLeft, F
 		bHasValidInkSource = true;
 	}
 
-	DestructiveResizePreviewComponent(Dimensions.X, Dimensions.Y);
+	FIntPoint Scale;
+	Scale.X = FMath::Max(1, TileSet->GetTileSize().X / CursorPreviewComponent->TileMap->TileWidth);
+	Scale.Y = FMath::Max(1, TileSet->GetTileSize().Y / CursorPreviewComponent->TileMap->TileHeight);
+
+	DestructiveResizePreviewComponent(Dimensions.X * Scale.X, Dimensions.Y * Scale.Y);
 
 	UPaperTileMap* PreviewMap = CursorPreviewComponent->TileMap;
 	UPaperTileLayer* PreviewLayer = GetSourceInkLayer();
-	for (int32 Y = 0; Y < PreviewMap->MapHeight; ++Y)
+
+	// Empty the preview layer (when selecting tiles from a tile set with a larger tile size than the map, there will be gaps in what is setup below)
+	FPaperTileInfo EmptyCell;
+	for (int32 Y = 0; Y < PreviewLayer->GetLayerHeight(); ++Y)
 	{
-		for (int32 X = 0; X < PreviewMap->MapWidth; ++X)
+		for (int32 X = 0; X < PreviewLayer->GetLayerWidth(); ++X)
+		{
+			PreviewLayer->SetCell(X, Y, EmptyCell);
+		}
+	}
+
+	// Copy the selection into the preview ink layer
+	for (int32 Y = 0; Y < FMath::Max(1, Dimensions.Y); ++Y)
+	{
+		for (int32 X = 0; X < FMath::Max(1, Dimensions.X); ++X)
 		{
 			FPaperTileInfo TileInfo;
 
@@ -1129,7 +1145,7 @@ void FEdModeTileMap::SetActivePaint(UPaperTileSet* TileSet, FIntPoint TopLeft, F
 				TileInfo.TileSet = TileSet;
 			}
 
-			PreviewLayer->SetCell(X, Y, TileInfo);
+			PreviewLayer->SetCell(X * Scale.X, Y * Scale.Y + Scale.Y - 1, TileInfo);
 		}
 	}
 
