Paper2D: Fixing static analysis error (accessing TileSet ptr when it could be nullptr)
#rnx

diff --git a/Engine/Plugins/2D/Paper2D/Source/Paper2DEditor/Private/TileMapEditing/EdModeTileMap.cpp b/Engine/Plugins/2D/Paper2D/Source/Paper2DEditor/Private/TileMapEditing/EdModeTileMap.cpp
--- a/Engine/Plugins/2D/Paper2D/Source/Paper2DEditor/Private/TileMapEditing/EdModeTileMap.cpp
+++ b/Engine/Plugins/2D/Paper2D/Source/Paper2DEditor/Private/TileMapEditing/EdModeTileMap.cpp
@@ -1110,9 +1110,12 @@ void FEdModeTileMap::SetActivePaint(UPaperTileSet* TileSet, FIntPoint TopLeft, F
 		bHasValidInkSource = true;
 	}
 
-	FIntPoint Scale;
-	Scale.X = FMath::Max(1, TileSet->GetTileSize().X / CursorPreviewComponent->TileMap->TileWidth);
-	Scale.Y = FMath::Max(1, TileSet->GetTileSize().Y / CursorPreviewComponent->TileMap->TileHeight);
+	FIntPoint Scale(1, 1);
+	if (TileSet != nullptr)
+	{
+		Scale.X = FMath::Max(1, TileSet->GetTileSize().X / CursorPreviewComponent->TileMap->TileWidth);
+		Scale.Y = FMath::Max(1, TileSet->GetTileSize().Y / CursorPreviewComponent->TileMap->TileHeight);
+	}
 
 	DestructiveResizePreviewComponent(Dimensions.X * Scale.X, Dimensions.Y * Scale.Y);
 
