Fixed artifact in sprite batching with temporary workaround.

diff --git a/Engine/Plugins/2D/Paper2D/Source/Paper2D/Private/PaperRenderSceneProxy.cpp b/Engine/Plugins/2D/Paper2D/Source/Paper2D/Private/PaperRenderSceneProxy.cpp
--- a/Engine/Plugins/2D/Paper2D/Source/Paper2D/Private/PaperRenderSceneProxy.cpp
+++ b/Engine/Plugins/2D/Paper2D/Source/Paper2D/Private/PaperRenderSceneProxy.cpp
@@ -272,11 +272,15 @@ void FPaperRenderSceneProxy::GetNewBatchMeshes(const FSceneView* View, int32 Vie
 
 	if (Vertices.Num())
 	{
-		FDynamicMeshBuilder DynamicMeshBuilder(View->GetFeatureLevel());
-		DynamicMeshBuilder.AddVertices(Vertices);
-
 		for (const FSpriteRenderSection& Batch : BatchedSections)
 		{
+			FDynamicMeshBuilder DynamicMeshBuilder(View->GetFeatureLevel());
+
+			for (int32 i = Batch.VertexOffset; i < Batch.VertexOffset + Batch.NumVertices; ++i)
+			{
+				DynamicMeshBuilder.AddVertex(Vertices[i]);
+			}
+
 			if (Batch.IsValid())
 			{
 				FMaterialRenderProxy* ParentMaterialProxy = Batch.Material->GetRenderProxy();
@@ -311,13 +315,8 @@ void FPaperRenderSceneProxy::GetNewBatchMeshes(const FSceneView* View, int32 Vie
 				Settings.CastShadow = bCastShadow;
 				Settings.bDisableBackfaceCulling = bDrawTwoSided;
 
-				FDynamicMeshDrawOffset DrawOffset;
-				DrawOffset.FirstIndex = Batch.VertexOffset;
-				DrawOffset.MinVertexIndex = Batch.VertexOffset;
-				DrawOffset.MaxVertexIndex = Batch.VertexOffset + Batch.NumVertices;
-				DrawOffset.NumPrimitives = Batch.NumVertices / 3;
 
-				DynamicMeshBuilder.GetMesh(GetLocalToWorld(), TextureOverrideMaterialProxy, DPG, Settings, &DrawOffset, ViewIndex, Collector);
+				DynamicMeshBuilder.GetMesh(GetLocalToWorld(), TextureOverrideMaterialProxy, DPG, Settings, nullptr, ViewIndex, Collector);
 			}
 		}
 	}
