Fix paper2d not playing nice with parallel initviews

diff --git a/Engine/Plugins/2D/Paper2D/Source/Paper2D/Private/PaperRenderSceneProxy.cpp b/Engine/Plugins/2D/Paper2D/Source/Paper2D/Private/PaperRenderSceneProxy.cpp
--- a/Engine/Plugins/2D/Paper2D/Source/Paper2D/Private/PaperRenderSceneProxy.cpp
+++ b/Engine/Plugins/2D/Paper2D/Source/Paper2D/Private/PaperRenderSceneProxy.cpp
@@ -401,7 +401,7 @@ FPrimitiveViewRelevance FPaperRenderSceneProxy::GetViewRelevance(const FSceneVie
 {
 	const FEngineShowFlags& EngineShowFlags = View->Family->EngineShowFlags;
 
-	checkSlow(IsInRenderingThread());
+	checkSlow(IsInParallelRenderingThread());
 
 	FPrimitiveViewRelevance Result;
 	Result.bDrawRelevance = IsShown(View) && EngineShowFlags.Paper2DSprites;
