Remove replicated flags from PaperFlipbookComponent properties that did not register themselves (and thus were never replicated)

diff --git a/Engine/Plugins/2D/Paper2D/Source/Paper2D/Classes/PaperFlipbookComponent.h b/Engine/Plugins/2D/Paper2D/Source/Paper2D/Classes/PaperFlipbookComponent.h
--- a/Engine/Plugins/2D/Paper2D/Source/Paper2D/Classes/PaperFlipbookComponent.h
+++ b/Engine/Plugins/2D/Paper2D/Source/Paper2D/Classes/PaperFlipbookComponent.h
@@ -33,23 +33,23 @@ protected:
 	UMaterialInterface* Material_DEPRECATED;
 
 	/** Current play rate of the flipbook */
-	UPROPERTY(Category=Sprite, EditAnywhere, Replicated)
+	UPROPERTY(Category=Sprite, EditAnywhere)
 	float PlayRate;
 
 	/** Whether the flipbook should loop when it reaches the end, or stop */
-	UPROPERTY(Replicated)
+	UPROPERTY()
 	uint32 bLooping:1;
 
 	/** If playback should move the current position backwards instead of forwards */
-	UPROPERTY(Replicated)
+	UPROPERTY()
 	uint32 bReversePlayback:1;
 
 	/** Are we currently playing (moving Position) */
-	UPROPERTY(Replicated)
+	UPROPERTY()
 	uint32 bPlaying:1;
 
 	/** Current position in the timeline */
-	UPROPERTY(Replicated)
+	UPROPERTY()
 	float AccumulatedTime;
 
 	/** Last frame index calculated */
