Paper2D - Making the SourceTexture and associated properties Editor Only.  Also making the source texture a SoftObjectPtr so that we can ignore source textures from cooking safely and only cook the atlases.

Michael.Noland

diff --git a/Engine/Plugins/2D/Paper2D/Source/Paper2D/Classes/PaperSprite.h b/Engine/Plugins/2D/Paper2D/Source/Paper2D/Classes/PaperSprite.h
--- a/Engine/Plugins/2D/Paper2D/Source/Paper2D/Classes/PaperSprite.h
+++ b/Engine/Plugins/2D/Paper2D/Source/Paper2D/Classes/PaperSprite.h
@@ -70,6 +70,7 @@ protected:
 	FVector2D SourceTextureDimension;
 #endif
 
+#if WITH_EDITORONLY_DATA
 	// Position within SourceTexture (in pixels)
 	UPROPERTY(Category=Sprite, EditAnywhere, AssetRegistrySearchable)
 	FVector2D SourceUV;
@@ -80,7 +81,8 @@ protected:
 
 	// The source texture that the sprite comes from
 	UPROPERTY(Category=Sprite, EditAnywhere, AssetRegistrySearchable)
-	UTexture2D* SourceTexture;
+	TSoftObjectPtr<UTexture2D> SourceTexture;
+#endif
 
 	// Additional source textures for other slots
 	UPROPERTY(Category=Sprite, EditAnywhere, AssetRegistrySearchable, meta=(DisplayName="Additional Textures"))
@@ -260,10 +262,9 @@ public:
 
 	ESpritePivotMode::Type GetPivotMode(FVector2D& OutCustomTextureSpacePivot) const { OutCustomTextureSpacePivot = CustomPivotPoint; return PivotMode; }
 
-
 	FVector2D GetSourceUV() const { return SourceUV; }
 	FVector2D GetSourceSize() const { return SourceDimension; }
-	UTexture2D* GetSourceTexture() const { return SourceTexture; }
+	UTexture2D* GetSourceTexture() const;
 
 	const UPaperSpriteAtlas* GetAtlasGroup() const { return AtlasGroup; }
 #endif
diff --git a/Engine/Plugins/2D/Paper2D/Source/Paper2D/Private/PaperSprite.cpp b/Engine/Plugins/2D/Paper2D/Source/Paper2D/Private/PaperSprite.cpp
--- a/Engine/Plugins/2D/Paper2D/Source/Paper2D/Private/PaperSprite.cpp
+++ b/Engine/Plugins/2D/Paper2D/Source/Paper2D/Private/PaperSprite.cpp
@@ -334,7 +334,7 @@ void UPaperSprite::ExtractSourceRegionFromTexturePoint(const FVector2D& SourcePo
 	FIntPoint SourceIntPoint(FMath::RoundToInt(SourcePoint.X), FMath::RoundToInt(SourcePoint.Y));
 	FIntPoint ClosestValidPoint;
 
-	FBitmap Bitmap(SourceTexture, 0, 0);
+	FBitmap Bitmap(GetSourceTexture(), 0, 0);
 	if (Bitmap.IsValid() && Bitmap.FoundClosestValidPoint(SourceIntPoint.X, SourceIntPoint.Y, 10, /*out*/ ClosestValidPoint))
 	{
 		FIntPoint Origin;
@@ -984,7 +984,7 @@ void UPaperSprite::FindTextureBoundingBox(float AlphaThreshold, /*out*/ FVector2
 	int32 BottomBound = (int32)(SourceUV.Y + SourceDimension.Y - 1);
 
 	const int32 AlphaThresholdInt = FMath::Clamp<int32>(AlphaThreshold * 255, 0, 255);
-	FBitmap SourceBitmap(SourceTexture, AlphaThresholdInt);
+	FBitmap SourceBitmap(GetSourceTexture(), AlphaThresholdInt);
 	if (SourceBitmap.IsValid())
 	{
 		// Make sure the initial bounds starts in the texture
@@ -1048,7 +1048,7 @@ void UPaperSprite::BuildGeometryFromContours(FSpriteGeometryCollection& GeomOwne
 
 	// DK: FindContours only returns positive contours, i.e. outsides
 	// Contour generation is simplified in FindContours by downscaling the detail prior to generating contour data
-	FindContours(InitialPos, InitialSize, GeomOwner.AlphaThreshold, GeomOwner.DetailAmount, SourceTexture, /*out*/ Contours);
+	FindContours(InitialPos, InitialSize, GeomOwner.AlphaThreshold, GeomOwner.DetailAmount, GetSourceTexture(), /*out*/ Contours);
 
 	// Convert the contours into geometry
 	GeomOwner.Shapes.Empty();
@@ -1460,6 +1460,15 @@ void UPaperSprite::SetPivotMode(ESpritePivotMode::Type InPivotMode, FVector2D In
 	}
 }
 
+#if WITH_EDITOR
+
+UTexture2D* UPaperSprite::GetSourceTexture() const
+{
+	return SourceTexture.LoadSynchronous();
+}
+
+#endif
+
 FVector2D UPaperSprite::ConvertTextureSpaceToPivotSpace(FVector2D Input) const
 {
 	const FVector2D Pivot = GetPivotPosition();
@@ -1647,27 +1656,29 @@ void UPaperSprite::GetAssetRegistryTags(TArray<FAssetRegistryTag>& OutTags) cons
 
 FSlateAtlasData UPaperSprite::GetSlateAtlasData() const
 {
-	if ( SourceTexture == nullptr && BakedSourceTexture == nullptr )
+	if (BakedSourceTexture)
 	{
-		return FSlateAtlasData(nullptr, FVector2D::ZeroVector, FVector2D::ZeroVector);
+		const FVector2D ImportedSize = FVector2D(BakedSourceTexture->GetImportedSize());
+
+		const FVector2D StartUV = BakedSourceUV / ImportedSize;
+		const FVector2D SizeUV = BakedSourceDimension / ImportedSize;
+
+		return FSlateAtlasData(BakedSourceTexture, StartUV, SizeUV);
 	}
-	else if ( BakedSourceTexture == nullptr )
+#if WITH_EDITOR
+	else if (UTexture2D* SourceTexturePtr = GetSourceTexture())
 	{
-		const FVector2D ImportedSize = FVector2D(SourceTexture->GetImportedSize());
+		const FVector2D ImportedSize = FVector2D(SourceTexturePtr->GetImportedSize());
 
 		const FVector2D StartUV = SourceUV / ImportedSize;
 		const FVector2D SizeUV = SourceDimension / ImportedSize;
 
-		return FSlateAtlasData(SourceTexture, StartUV, SizeUV);
+		return FSlateAtlasData(SourceTexturePtr, StartUV, SizeUV);
 	}
+#endif
 	else
 	{
-		const FVector2D ImportedSize = FVector2D(BakedSourceTexture->GetImportedSize());
-
-		const FVector2D StartUV = BakedSourceUV / ImportedSize;
-		const FVector2D SizeUV = BakedSourceDimension / ImportedSize;
-
-		return FSlateAtlasData(BakedSourceTexture, StartUV, SizeUV);
+		return FSlateAtlasData(nullptr, FVector2D::ZeroVector, FVector2D::ZeroVector);
 	}
 }
 
@@ -1869,7 +1880,11 @@ void UPaperSprite::PostLoad()
 
 UTexture2D* UPaperSprite::GetBakedTexture() const
 {
-	return (BakedSourceTexture != nullptr) ? BakedSourceTexture : SourceTexture;
+#if WITH_EDITOR
+	return (BakedSourceTexture != nullptr) ? BakedSourceTexture : GetSourceTexture();
+#else
+	return BakedSourceTexture;
+#endif
 }
 
 void UPaperSprite::GetBakedAdditionalSourceTextures(FAdditionalSpriteTextureArray& OutTextureList) const
diff --git a/Engine/Plugins/2D/Paper2D/Source/Paper2DEditor/Private/SpriteEditor/SpriteEditorViewportClient.cpp b/Engine/Plugins/2D/Paper2D/Source/Paper2DEditor/Private/SpriteEditor/SpriteEditorViewportClient.cpp
--- a/Engine/Plugins/2D/Paper2D/Source/Paper2DEditor/Private/SpriteEditor/SpriteEditorViewportClient.cpp
+++ b/Engine/Plugins/2D/Paper2D/Source/Paper2DEditor/Private/SpriteEditor/SpriteEditorViewportClient.cpp
@@ -150,7 +150,7 @@ void FSpriteEditorViewportClient::UpdateSourceTextureSpriteFromSprite(UPaperSpri
 
 			FSpriteAssetInitParameters SpriteReinitParams;
 
-			SpriteReinitParams.SetTextureAndFill(SourceSprite->SourceTexture);
+			SpriteReinitParams.SetTextureAndFill(SourceSprite->GetSourceTexture());
 			SpriteReinitParams.DefaultMaterialOverride = SourceSprite->DefaultMaterial;
 			SpriteReinitParams.AlternateMaterialOverride = SourceSprite->AlternateMaterial;
 			SpriteReinitParams.SetPixelsPerUnrealUnit(SourceSprite->PixelsPerUnrealUnit);
@@ -809,7 +809,7 @@ UPaperSprite* FSpriteEditorViewportClient::CreateNewSprite(const FIntPoint& TopL
 
 	// Create the factory used to generate the sprite
 	UPaperSpriteFactory* SpriteFactory = NewObject<UPaperSpriteFactory>();
-	SpriteFactory->InitialTexture = CurrentSprite->SourceTexture;
+	SpriteFactory->InitialTexture = CurrentSprite->GetSourceTexture();
 	SpriteFactory->bUseSourceRegion = true;
 	SpriteFactory->InitialSourceUV = TopLeft;
 	SpriteFactory->InitialSourceDimension = Dimensions;
