Paper2D - If we don't have a sprite atlas group we belong to, then we need to make sure our baked texture is set to our source texture at editor time, otherwise it wont get cooked, if it's a stand alone texture and not in an atlas.

Michael.Noland

diff --git a/Engine/Plugins/2D/Paper2D/Source/Paper2D/Private/PaperSprite.cpp b/Engine/Plugins/2D/Paper2D/Source/Paper2D/Private/PaperSprite.cpp
--- a/Engine/Plugins/2D/Paper2D/Source/Paper2D/Private/PaperSprite.cpp
+++ b/Engine/Plugins/2D/Paper2D/Source/Paper2D/Private/PaperSprite.cpp
@@ -1900,6 +1900,12 @@ void UPaperSprite::PostLoad()
 	{
 		RebuildRenderData();
 	}
+
+	/** If we don't have an atlas group make sure we assign the source texture to our baked texture. */
+	if (AtlasGroup == nullptr)
+	{
+		BakedSourceTexture = GetSourceTexture();
+	}
 #endif
 }
 
