Moved vector constant initialization in Paper2D tile layer code to be inside body setup code to avoid global constructor ordering issues on some platforms (now that Neon vector opts are enabled)
#rnx

diff --git a/Engine/Plugins/2D/Paper2D/Source/Paper2D/Private/PaperTileLayer.cpp b/Engine/Plugins/2D/Paper2D/Source/Paper2D/Private/PaperTileLayer.cpp
--- a/Engine/Plugins/2D/Paper2D/Source/Paper2D/Private/PaperTileLayer.cpp
+++ b/Engine/Plugins/2D/Paper2D/Source/Paper2D/Private/PaperTileLayer.cpp
@@ -6,36 +6,6 @@
 #include "PaperTileMap.h"
 #include "PaperTileSet.h"
 
-// Handles the rotation and flipping of collision geometry from a tile
-// 0,5,6,3 are clockwise rotations of a regular tile
-// 4,7,2,1 are clockwise rotations of a horizontally flipped tile
-const static FTransform TilePermutationTransforms[8] =
-{
-	// 000 - normal
-	FTransform::Identity,
-
-	// 001 - diagonal
-	FTransform(FRotator(  90.0f, 0.0f, 0.0f), FVector::ZeroVector, -PaperAxisX.GetAbs() + PaperAxisY.GetAbs() + PaperAxisZ.GetAbs()),
-
-	// 010 - flip Y
-	FTransform(FRotator(-180.0f, 0.0f, 0.0f), FVector::ZeroVector, -PaperAxisX.GetAbs() + PaperAxisY.GetAbs() + PaperAxisZ.GetAbs()),
-
-	// 011 - diagonal then flip Y (rotate 270 clockwise)
-	FTransform(FRotator(  90.0f, 0.0f, 0.0f)),
-
-	// 100 - flip X
-	FTransform(FRotator::ZeroRotator, FVector::ZeroVector, -PaperAxisX.GetAbs() + PaperAxisY.GetAbs() + PaperAxisZ.GetAbs()),
-
-	// 101 - diagonal then flip X (clockwise 90)
-	FTransform(FRotator( -90.0f, 0.0f, 0.0f)),
-
-	// 110 - flip X and flip Y (rotate 180 either way)
-	FTransform(FRotator(-180.0f, 0.0f, 0.0f)),
-
-	// 111 - diagonal then flip X and Y
-	FTransform(FRotator(-90.0f, 0.0f, 0.0f), FVector::ZeroVector, -PaperAxisX.GetAbs() + PaperAxisY.GetAbs() + PaperAxisZ.GetAbs()),
-};
-
 //////////////////////////////////////////////////////////////////////////
 // FPaperTileLayerToBodySetupBuilder
 
@@ -238,7 +208,7 @@ void UPaperTileLayer::AugmentBodySetup(UBodySetup* ShapeBodySetup, float RenderS
 				{
 					const int32 Flags = CellInfo.GetFlagsAsIndex();
 
-					const FTransform& LocalTransform = TilePermutationTransforms[Flags];
+					const FTransform& LocalTransform = GetTileTransform(Flags);
 					
 					// When building the collison we must consider setting the correct offsets for the projection mode.
 					// This must be inline with the way you are making any changes too the 2D grid.
@@ -320,6 +290,37 @@ bool UPaperTileLayer::UsesTileSet(UPaperTileSet* TileSet) const
 
 FTransform UPaperTileLayer::GetTileTransform(int32 FlagIndex)
 {
+	// Handles the rotation and flipping of collision geometry from a tile
+	// 0,5,6,3 are clockwise rotations of a regular tile
+	// 4,7,2,1 are clockwise rotations of a horizontally flipped tile
+	const static FTransform TilePermutationTransforms[8] =
+	{
+		// 000 - normal
+		FTransform::Identity,
+
+		// 001 - diagonal
+		FTransform(FRotator(90.0f, 0.0f, 0.0f), FVector::ZeroVector, -PaperAxisX.GetAbs() + PaperAxisY.GetAbs() + PaperAxisZ.GetAbs()),
+
+		// 010 - flip Y
+		FTransform(FRotator(-180.0f, 0.0f, 0.0f), FVector::ZeroVector, -PaperAxisX.GetAbs() + PaperAxisY.GetAbs() + PaperAxisZ.GetAbs()),
+
+		// 011 - diagonal then flip Y (rotate 270 clockwise)
+		FTransform(FRotator(90.0f, 0.0f, 0.0f)),
+
+		// 100 - flip X
+		FTransform(FRotator::ZeroRotator, FVector::ZeroVector, -PaperAxisX.GetAbs() + PaperAxisY.GetAbs() + PaperAxisZ.GetAbs()),
+
+		// 101 - diagonal then flip X (clockwise 90)
+		FTransform(FRotator(-90.0f, 0.0f, 0.0f)),
+
+		// 110 - flip X and flip Y (rotate 180 either way)
+		FTransform(FRotator(-180.0f, 0.0f, 0.0f)),
+
+		// 111 - diagonal then flip X and Y
+		FTransform(FRotator(-90.0f, 0.0f, 0.0f), FVector::ZeroVector, -PaperAxisX.GetAbs() + PaperAxisY.GetAbs() + PaperAxisZ.GetAbs()),
+	};
+
+
 	checkSlow((FlagIndex >= 0) && (FlagIndex < 8));
 	return TilePermutationTransforms[FlagIndex];
 }
