Paper2D: Fix crash running monolithic editor due to global constructor using localization system before it's initialized.

diff --git a/Engine/Plugins/2D/Paper2D/Source/Paper2DEditor/Private/SPaperEditorViewport.cpp b/Engine/Plugins/2D/Paper2D/Source/Paper2DEditor/Private/SPaperEditorViewport.cpp
--- a/Engine/Plugins/2D/Paper2D/Source/Paper2DEditor/Private/SPaperEditorViewport.cpp
+++ b/Engine/Plugins/2D/Paper2D/Source/Paper2DEditor/Private/SPaperEditorViewport.cpp
@@ -27,26 +27,30 @@ struct FZoomLevelEntry
 	float ZoomAmount;
 };
 
-static const FZoomLevelEntry ZoomLevels[NumZoomLevels] =
+static const FZoomLevelEntry& GetZoomLevel(int Index)
 {
-	FZoomLevelEntry(0.03125f, NSLOCTEXT("PaperEditor", "ZoomLevel_1_32", "1:32")),
-	FZoomLevelEntry(0.0625f, NSLOCTEXT("PaperEditor", "ZoomLevel_1_16", "1:16")),
-	FZoomLevelEntry(0.125f, NSLOCTEXT("PaperEditor", "ZoomLevel_1_8", "1:8")),
-	FZoomLevelEntry(0.250f, NSLOCTEXT("PaperEditor", "ZoomLevel_1_4", "1:4")),
-	FZoomLevelEntry(0.500f, NSLOCTEXT("PaperEditor", "ZoomLevel_1_2", "1:2")),
-	FZoomLevelEntry(0.750f, NSLOCTEXT("PaperEditor", "ZoomLevel_3_4", "3:4")),
-	FZoomLevelEntry(0.875f, NSLOCTEXT("PaperEditor", "ZoomLevel_7_8", "7:8")),
-	FZoomLevelEntry(1.000f, NSLOCTEXT("PaperEditor", "ZoomLevel_1_1", "1:1")),
-	FZoomLevelEntry(2.000f, NSLOCTEXT("PaperEditor", "ZoomLevel_2x", "2x")),
-	FZoomLevelEntry(3.000f, NSLOCTEXT("PaperEditor", "ZoomLevel_3x", "3x")),
-	FZoomLevelEntry(4.500f, NSLOCTEXT("PaperEditor", "ZoomLevel_4x", "4x")),
-	FZoomLevelEntry(5.000f, NSLOCTEXT("PaperEditor", "ZoomLevel_5x", "5x")),
-	FZoomLevelEntry(6.000f, NSLOCTEXT("PaperEditor", "ZoomLevel_6x", "6x")),
-	FZoomLevelEntry(7.000f, NSLOCTEXT("PaperEditor", "ZoomLevel_7x", "7x")),
-	FZoomLevelEntry(8.000f, NSLOCTEXT("PaperEditor", "ZoomLevel_8x", "8x")),
-	FZoomLevelEntry(16.000f, NSLOCTEXT("PaperEditor", "ZoomLevel_16x", "16x")),
-	FZoomLevelEntry(32.000f, NSLOCTEXT("PaperEditor", "ZoomLevel_32x", "32x"))
-};
+	static const FZoomLevelEntry ZoomLevels[NumZoomLevels] =
+	{
+		FZoomLevelEntry(0.03125f, NSLOCTEXT("PaperEditor", "ZoomLevel_1_32", "1:32")),
+		FZoomLevelEntry(0.0625f, NSLOCTEXT("PaperEditor", "ZoomLevel_1_16", "1:16")),
+		FZoomLevelEntry(0.125f, NSLOCTEXT("PaperEditor", "ZoomLevel_1_8", "1:8")),
+		FZoomLevelEntry(0.250f, NSLOCTEXT("PaperEditor", "ZoomLevel_1_4", "1:4")),
+		FZoomLevelEntry(0.500f, NSLOCTEXT("PaperEditor", "ZoomLevel_1_2", "1:2")),
+		FZoomLevelEntry(0.750f, NSLOCTEXT("PaperEditor", "ZoomLevel_3_4", "3:4")),
+		FZoomLevelEntry(0.875f, NSLOCTEXT("PaperEditor", "ZoomLevel_7_8", "7:8")),
+		FZoomLevelEntry(1.000f, NSLOCTEXT("PaperEditor", "ZoomLevel_1_1", "1:1")),
+		FZoomLevelEntry(2.000f, NSLOCTEXT("PaperEditor", "ZoomLevel_2x", "2x")),
+		FZoomLevelEntry(3.000f, NSLOCTEXT("PaperEditor", "ZoomLevel_3x", "3x")),
+		FZoomLevelEntry(4.500f, NSLOCTEXT("PaperEditor", "ZoomLevel_4x", "4x")),
+		FZoomLevelEntry(5.000f, NSLOCTEXT("PaperEditor", "ZoomLevel_5x", "5x")),
+		FZoomLevelEntry(6.000f, NSLOCTEXT("PaperEditor", "ZoomLevel_6x", "6x")),
+		FZoomLevelEntry(7.000f, NSLOCTEXT("PaperEditor", "ZoomLevel_7x", "7x")),
+		FZoomLevelEntry(8.000f, NSLOCTEXT("PaperEditor", "ZoomLevel_8x", "8x")),
+		FZoomLevelEntry(16.000f, NSLOCTEXT("PaperEditor", "ZoomLevel_16x", "16x")),
+		FZoomLevelEntry(32.000f, NSLOCTEXT("PaperEditor", "ZoomLevel_32x", "32x"))
+	};
+	return ZoomLevels[Index];
+}
 
 //////////////////////////////////////////////////////////////////////////
 // SPaperEditorViewport
@@ -391,7 +395,7 @@ int32 SPaperEditorViewport::FindNearestZoomLevel(float InZoomAmount, bool bRound
 
 	for (int32 ZoomLevelIndex=0; ZoomLevelIndex < NumZoomLevels; ++ZoomLevelIndex)
 	{
-		if (InZoomAmount <= ZoomLevels[ZoomLevelIndex].ZoomAmount)
+		if (InZoomAmount <= GetZoomLevel(ZoomLevelIndex).ZoomAmount)
 		{
 			Result = ZoomLevelIndex;
 			bFoundItem = true;
@@ -418,12 +422,12 @@ int32 SPaperEditorViewport::FindNearestZoomLevel(float InZoomAmount, bool bRound
 
 float SPaperEditorViewport::GetZoomAmount() const
 {
-	return ZoomLevels[ZoomLevel].ZoomAmount;
+	return GetZoomLevel(ZoomLevel).ZoomAmount;
 }
 
 FText SPaperEditorViewport::GetZoomText() const
 {
-	return ZoomLevels[ZoomLevel].DisplayText;
+	return GetZoomLevel(ZoomLevel).DisplayText;
 }
 
 FSlateColor SPaperEditorViewport::GetZoomTextColorAndOpacity() const
