Deprecating ARRAY_COUNT and changing it to UE_ARRAY_COUNT.
Replicated from CL# 7924370.

diff --git a/Engine/Plugins/2D/Paper2D/Source/PaperTiledImporter/Private/PaperTiledImporterFactory.cpp b/Engine/Plugins/2D/Paper2D/Source/PaperTiledImporter/Private/PaperTiledImporterFactory.cpp
--- a/Engine/Plugins/2D/Paper2D/Source/PaperTiledImporter/Private/PaperTiledImporterFactory.cpp
+++ b/Engine/Plugins/2D/Paper2D/Source/PaperTiledImporter/Private/PaperTiledImporterFactory.cpp
@@ -421,13 +421,13 @@ void UPaperTiledImporterFactory::ParseGlobalInfoFromJSON(TSharedPtr<FJsonObject>
 		FRequiredIntField( OutParsedInfo.TileWidth, TEXT("tilewidth") ),
 		FRequiredIntField( OutParsedInfo.TileHeight, TEXT("tileheight") )
 	};
-	bSuccessfullyParsed = bSuccessfullyParsed && ParseIntegerFields(IntFields, ARRAY_COUNT(IntFields), Tree, NameForErrors, bSilent);
+	bSuccessfullyParsed = bSuccessfullyParsed && ParseIntegerFields(IntFields, UE_ARRAY_COUNT(IntFields), Tree, NameForErrors, bSilent);
 
 	// Parse hexsidelength if present
 	FRequiredIntField OptionalIntFields[] = {
 		FRequiredIntField(OutParsedInfo.HexSideLength, TEXT("hexsidelength"), 0)
 	};
-	ParseIntegerFields(OptionalIntFields, ARRAY_COUNT(OptionalIntFields), Tree, NameForErrors, /*bSilent=*/ true);
+	ParseIntegerFields(OptionalIntFields, UE_ARRAY_COUNT(OptionalIntFields), Tree, NameForErrors, /*bSilent=*/ true);
 
 	// Parse StaggerAxis if present
 	const FString StaggerAxisStr = FPaperJSONHelpers::ReadString(Tree, TEXT("staggeraxis"), FString());
@@ -619,7 +619,7 @@ void FTileSetFromTiled::ParseTileSetFromJSON(TSharedPtr<FJsonObject> Tree, const
 		FRequiredIntField( TileHeight, TEXT("tileheight"), 1 )
 	};
 
-	bSuccessfullyParsed = bSuccessfullyParsed && ParseIntegerFields(IntFields, ARRAY_COUNT(IntFields), Tree, NameForErrors, bSilent);
+	bSuccessfullyParsed = bSuccessfullyParsed && ParseIntegerFields(IntFields, UE_ARRAY_COUNT(IntFields), Tree, NameForErrors, bSilent);
 
 	// Parse the tile offset
 	if (bSuccessfullyParsed)
@@ -752,7 +752,7 @@ bool FTileLayerFromTiled::ParseFromJSON(TSharedPtr<FJsonObject> Tree, const FStr
 		FRequiredIntField( OffsetY, TEXT("y"), 0 )
 	};
 
-	bSuccessfullyParsed = bSuccessfullyParsed && ParseIntegerFields(IntFields, ARRAY_COUNT(IntFields), Tree, NameForErrors, bSilent);
+	bSuccessfullyParsed = bSuccessfullyParsed && ParseIntegerFields(IntFields, UE_ARRAY_COUNT(IntFields), Tree, NameForErrors, bSilent);
 
 	if (!Tree->TryGetBoolField(TEXT("visible"), /*out*/ bVisible))
 	{
@@ -1018,7 +1018,7 @@ bool FTiledObject::ParseFromJSON(TSharedPtr<FJsonObject> Tree, const FString& Na
 		FRequiredDoubleField(RotationDegrees, TEXT("rotation"), -MAX_FLT)
 	};
 
-	bSuccessfullyParsed = bSuccessfullyParsed && ParseScalarFields(FloatFields, ARRAY_COUNT(FloatFields), Tree, NameForErrors, bSilent);
+	bSuccessfullyParsed = bSuccessfullyParsed && ParseScalarFields(FloatFields, UE_ARRAY_COUNT(FloatFields), Tree, NameForErrors, bSilent);
 
 	if (!Tree->TryGetBoolField(TEXT("visible"), /*out*/ bVisible))
 	{
@@ -1096,7 +1096,7 @@ bool FTiledObject::ParsePointArray(TArray<FVector2D>& OutPoints, const TArray<TS
 			FRequiredDoubleField(Y, TEXT("y"), -MAX_FLT)
 		};
 
-		bSuccessfullyParsed = bSuccessfullyParsed && ParseScalarFields(FloatFields, ARRAY_COUNT(FloatFields), ArrayElement->AsObject(), NameForErrors, bSilent);
+		bSuccessfullyParsed = bSuccessfullyParsed && ParseScalarFields(FloatFields, UE_ARRAY_COUNT(FloatFields), ArrayElement->AsObject(), NameForErrors, bSilent);
 
 		OutPoints.Add(FVector2D(X, Y));
 	}
diff --git a/Engine/Plugins/Compositing/Composure/Source/Composure/Private/CompositingElements/CompositingTextureLookupTable.cpp b/Engine/Plugins/Compositing/Composure/Source/Composure/Private/CompositingElements/CompositingTextureLookupTable.cpp
--- a/Engine/Plugins/Compositing/Composure/Source/Composure/Private/CompositingElements/CompositingTextureLookupTable.cpp
+++ b/Engine/Plugins/Compositing/Composure/Source/Composure/Private/CompositingElements/CompositingTextureLookupTable.cpp
@@ -29,7 +29,7 @@ namespace CompositingTextureLookupTable_Impl
 
 static bool CompositingTextureLookupTable_Impl::IsLookupNameAPrePassAlias(FName LookupName)
 {
-	for (int32 AliasIndex = 0; AliasIndex < ARRAY_COUNT(PrePassAliases); ++AliasIndex)
+	for (int32 AliasIndex = 0; AliasIndex < UE_ARRAY_COUNT(PrePassAliases); ++AliasIndex)
 	{
 		if (LookupName == PrePassAliases[AliasIndex])
 		{
diff --git a/Engine/Plugins/Developer/Concert/ConcertMain/Source/Concert/Private/ConcertLogger.cpp b/Engine/Plugins/Developer/Concert/ConcertMain/Source/Concert/Private/ConcertLogger.cpp
--- a/Engine/Plugins/Developer/Concert/ConcertMain/Source/Concert/Private/ConcertLogger.cpp
+++ b/Engine/Plugins/Developer/Concert/ConcertMain/Source/Concert/Private/ConcertLogger.cpp
@@ -518,7 +518,7 @@ void FConcertLogger::InternalStartLogging()
 	if (CSVArchive)
 	{
 		UTF8CHAR UTF8BOM[] = { 0xEF, 0xBB, 0xBF };
-		CSVArchive->Serialize(&UTF8BOM, ARRAY_COUNT(UTF8BOM) * sizeof(UTF8CHAR));
+		CSVArchive->Serialize(&UTF8BOM, UE_ARRAY_COUNT(UTF8BOM) * sizeof(UTF8CHAR));
 	}
 	LogHeader();
 	bIsLogging = CSVArchive.IsValid();
diff --git a/Engine/Plugins/Developer/VisualStudioSourceCodeAccess/Source/VisualStudioSourceCodeAccess/Private/VisualStudioSourceCodeAccessor.cpp b/Engine/Plugins/Developer/VisualStudioSourceCodeAccess/Source/VisualStudioSourceCodeAccess/Private/VisualStudioSourceCodeAccessor.cpp
--- a/Engine/Plugins/Developer/VisualStudioSourceCodeAccess/Source/VisualStudioSourceCodeAccess/Private/VisualStudioSourceCodeAccessor.cpp
+++ b/Engine/Plugins/Developer/VisualStudioSourceCodeAccess/Source/VisualStudioSourceCodeAccess/Private/VisualStudioSourceCodeAccessor.cpp
@@ -119,9 +119,9 @@ int32 GetVisualStudioVersionForSolution(const FString& InSolutionFile)
 			const TCHAR* VersionChar = *SolutionFileContents + VersionStringStart + VisualStudioVersionString.Len();
 
 			const TCHAR VersionSuffix[] = TEXT("Version ");
-			if (FCString::Strnicmp(VersionChar, VersionSuffix, ARRAY_COUNT(VersionSuffix) - 1) == 0)
+			if (FCString::Strnicmp(VersionChar, VersionSuffix, UE_ARRAY_COUNT(VersionSuffix) - 1) == 0)
 			{
-				VersionChar += ARRAY_COUNT(VersionSuffix) - 1;
+				VersionChar += UE_ARRAY_COUNT(VersionSuffix) - 1;
 			}
 
 			FString VersionString;
diff --git a/Engine/Plugins/Experimental/PanoramicCapture/Source/PanoramicCapture/Private/SceneCapturer.cpp b/Engine/Plugins/Experimental/PanoramicCapture/Source/PanoramicCapture/Private/SceneCapturer.cpp
--- a/Engine/Plugins/Experimental/PanoramicCapture/Source/PanoramicCapture/Private/SceneCapturer.cpp
+++ b/Engine/Plugins/Experimental/PanoramicCapture/Source/PanoramicCapture/Private/SceneCapturer.cpp
@@ -121,7 +121,7 @@ USceneCapturer::USceneCapturer(FVTableHelper& Helper)
     , SphericalAtlasHeight(SphericalAtlasWidth / 2)
     , bForceAlpha(FStereoPanoramaManager::ForceAlpha->GetInt() != 0)
     , bEnableBilerp(FStereoPanoramaManager::EnableBilerp->GetInt() != 0)
-    , SSMethod(FMath::Clamp<int32>(FStereoPanoramaManager::SuperSamplingMethod->GetInt(), 0, ARRAY_COUNT(g_ssPatterns)))
+    , SSMethod(FMath::Clamp<int32>(FStereoPanoramaManager::SuperSamplingMethod->GetInt(), 0, UE_ARRAY_COUNT(g_ssPatterns)))
     , bOverrideInitialYaw(FStereoPanoramaManager::ShouldOverrideInitialYaw->GetInt() != 0)
     , ForcedInitialYaw(FRotator::ClampAxis(FStereoPanoramaManager::ForcedInitialYaw->GetFloat()))
     , OutputDir(FStereoPanoramaManager::OutputDir->GetString().IsEmpty() ? FPaths::ProjectSavedDir() / TEXT("StereoPanorama") : FStereoPanoramaManager::OutputDir->GetString())
@@ -152,7 +152,7 @@ USceneCapturer::USceneCapturer()
 	, SphericalAtlasHeight(SphericalAtlasWidth / 2)
 	, bForceAlpha(FStereoPanoramaManager::ForceAlpha->GetInt() != 0)
 	, bEnableBilerp(FStereoPanoramaManager::EnableBilerp->GetInt() != 0)
-	, SSMethod(FMath::Clamp<int32>(FStereoPanoramaManager::SuperSamplingMethod->GetInt(), 0, ARRAY_COUNT(g_ssPatterns)))
+	, SSMethod(FMath::Clamp<int32>(FStereoPanoramaManager::SuperSamplingMethod->GetInt(), 0, UE_ARRAY_COUNT(g_ssPatterns)))
 	, bOverrideInitialYaw(FStereoPanoramaManager::ShouldOverrideInitialYaw->GetInt() != 0)
 	, ForcedInitialYaw(FRotator::ClampAxis(FStereoPanoramaManager::ForcedInitialYaw->GetFloat()))
 	, OutputDir(FStereoPanoramaManager::OutputDir->GetString().IsEmpty() ? FPaths::ProjectSavedDir() / TEXT("StereoPanorama") : FStereoPanoramaManager::OutputDir->GetString())
diff --git a/Engine/Plugins/Experimental/RemoteSession/Source/RemoteSession/Private/Channels/RemoteSessionARCameraChannel.cpp b/Engine/Plugins/Experimental/RemoteSession/Source/RemoteSession/Private/Channels/RemoteSessionARCameraChannel.cpp
--- a/Engine/Plugins/Experimental/RemoteSession/Source/RemoteSession/Private/Channels/RemoteSessionARCameraChannel.cpp
+++ b/Engine/Plugins/Experimental/RemoteSession/Source/RemoteSession/Private/Channels/RemoteSessionARCameraChannel.cpp
@@ -216,7 +216,7 @@ void FARCameraSceneViewExtension::PreRenderView_RenderThread(FRHICommandListImme
 		const uint16 Indices[] = { 0, 1, 2, 2, 1, 3 };
 
 		TResourceArray<uint16, INDEXBUFFER_ALIGNMENT> IndexBuffer;
-		const uint32 NumIndices = ARRAY_COUNT(Indices);
+		const uint32 NumIndices = UE_ARRAY_COUNT(Indices);
 		IndexBuffer.AddUninitialized(NumIndices);
 		FMemory::Memcpy(IndexBuffer.GetData(), Indices, NumIndices * sizeof(uint16));
 
diff --git a/Engine/Plugins/Experimental/SkeletalReduction/Source/Private/SkeletalMeshReductionPlugin.cpp b/Engine/Plugins/Experimental/SkeletalReduction/Source/Private/SkeletalMeshReductionPlugin.cpp
--- a/Engine/Plugins/Experimental/SkeletalReduction/Source/Private/SkeletalMeshReductionPlugin.cpp
+++ b/Engine/Plugins/Experimental/SkeletalReduction/Source/Private/SkeletalMeshReductionPlugin.cpp
@@ -962,7 +962,7 @@ float FQuadricSkeletalMeshReduction::SimplifyMesh( const FSkeletalMeshOptimizati
 		2.8f,	// High
 		8.0f,	// Highest
 	};
-	static_assert(ARRAY_COUNT(ImportanceTable) == SMOI_MAX, "Bad importance table size.");
+	static_assert(UE_ARRAY_COUNT(ImportanceTable) == SMOI_MAX, "Bad importance table size.");
 
 	NormalWeight    *= ImportanceTable[Settings.ShadingImportance];
 	TangentWeight   *= ImportanceTable[Settings.ShadingImportance];
diff --git a/Engine/Plugins/Lumin/MagicLeap/Source/MagicLeap/Private/MagicLeapCustomPresent.cpp b/Engine/Plugins/Lumin/MagicLeap/Source/MagicLeap/Private/MagicLeapCustomPresent.cpp
--- a/Engine/Plugins/Lumin/MagicLeap/Source/MagicLeap/Private/MagicLeapCustomPresent.cpp
+++ b/Engine/Plugins/Lumin/MagicLeap/Source/MagicLeap/Private/MagicLeapCustomPresent.cpp
@@ -361,7 +361,7 @@ void FMagicLeapCustomPresentOpenGL::FinishRendering()
 			glEnable(GL_FRAMEBUFFER_SRGB);
 		}
 
-		static_assert(ARRAY_COUNT(vp_array.virtual_cameras) == 2, "The MLSDK has updated the size of the virtual_cameras array.");
+		static_assert(UE_ARRAY_COUNT(vp_array.virtual_cameras) == 2, "The MLSDK has updated the size of the virtual_cameras array.");
 #if 0 // Enable this in case the MLSDK increases the size of the virtual_cameras array past 2
 		for (uint32 i = 2; i < vp_array.num_virtual_cameras; ++i)
 		{
diff --git a/Engine/Plugins/Media/MediaIOFramework/Source/MediaIOCore/Private/MediaIOCoreCommonDisplayMode.cpp b/Engine/Plugins/Media/MediaIOFramework/Source/MediaIOCore/Private/MediaIOCoreCommonDisplayMode.cpp
--- a/Engine/Plugins/Media/MediaIOFramework/Source/MediaIOCore/Private/MediaIOCoreCommonDisplayMode.cpp
+++ b/Engine/Plugins/Media/MediaIOFramework/Source/MediaIOCore/Private/MediaIOCoreCommonDisplayMode.cpp
@@ -111,7 +111,7 @@ TArrayView<const FMediaIOCommonDisplayModeResolutionInfo> FMediaIOCommonDisplayM
 
 const FMediaIOCommonDisplayModeResolutionInfo* FMediaIOCommonDisplayModes::Find(int32 InWidth, int32 InHeight)
 {
-	const int32 NumModeResolutions = ARRAY_COUNT(MediaIOCommonDisplayModeInfo::AllCommonDisplayModeResolutions);
+	const int32 NumModeResolutions = UE_ARRAY_COUNT(MediaIOCommonDisplayModeInfo::AllCommonDisplayModeResolutions);
 	for(int32 Index = 0; Index < NumModeResolutions; ++Index)
 	{
 		const FMediaIOCommonDisplayModeResolutionInfo* Info = MediaIOCommonDisplayModeInfo::AllCommonDisplayModeResolutions + Index;
@@ -125,7 +125,7 @@ const FMediaIOCommonDisplayModeResolutionInfo* FMediaIOCommonDisplayModes::Find(
 
 const FMediaIOCommonDisplayModeInfo* FMediaIOCommonDisplayModes::Find(int32 InWidth, int32 InHeight, const FFrameRate& InFrameRate, EMediaIOStandardType InStandard)
 {
-	const int32 NumModes = ARRAY_COUNT(MediaIOCommonDisplayModeInfo::AllCommonDisplayModes);
+	const int32 NumModes = UE_ARRAY_COUNT(MediaIOCommonDisplayModeInfo::AllCommonDisplayModes);
 	for (int32 Index = 0; Index < NumModes; ++Index)
 	{
 		const FMediaIOCommonDisplayModeInfo* Info = MediaIOCommonDisplayModeInfo::AllCommonDisplayModes + Index;
diff --git a/Engine/Plugins/NetcodeUnitTest/NetcodeUnitTest/Source/NetcodeUnitTest/Private/UnitTests/VMReflection.cpp b/Engine/Plugins/NetcodeUnitTest/NetcodeUnitTest/Source/NetcodeUnitTest/Private/UnitTests/VMReflection.cpp
--- a/Engine/Plugins/NetcodeUnitTest/NetcodeUnitTest/Source/NetcodeUnitTest/Private/UnitTests/VMReflection.cpp
+++ b/Engine/Plugins/NetcodeUnitTest/NetcodeUnitTest/Source/NetcodeUnitTest/Private/UnitTests/VMReflection.cpp
@@ -869,7 +869,7 @@ bool UVMReflection::ExecuteUnitTest()
 		Results[2] = (uint8*)((FVMReflection(TestObj)->*"BytePropArray")["uint8"][2], &bError[2]);
 		Results[3] = (uint8*)((FVMReflection(TestObj)->*"BytePropArray")["uint8"][3], &bError[3]);
 
-		for (int32 i=0; i<ARRAY_COUNT(Results); i++)
+		for (int32 i=0; i<UE_ARRAY_COUNT(Results); i++)
 		{
 			if (!bError[i] && Results[i] != NULL)
 			{
@@ -983,7 +983,7 @@ bool UVMReflection::ExecuteUnitTest()
 		Results[2] = (uint8*)((FVMReflection(TestObj)->*"DynBytePropArray")["uint8"][2], &bError[2]);
 		Results[3] = (uint8*)((FVMReflection(TestObj)->*"DynBytePropArray")["uint8"][3], &bError[3]);
 
-		for (int32 i=0; i<ARRAY_COUNT(Results); i++)
+		for (int32 i=0; i<UE_ARRAY_COUNT(Results); i++)
 		{
 			if (!bError[i] && Results[i] != NULL)
 			{
diff --git a/Engine/Plugins/NetcodeUnitTest/NetcodeUnitTest/Source/NetcodeUnitTest/Public/NUTUtil.h b/Engine/Plugins/NetcodeUnitTest/NetcodeUnitTest/Source/NetcodeUnitTest/Public/NUTUtil.h
--- a/Engine/Plugins/NetcodeUnitTest/NetcodeUnitTest/Source/NetcodeUnitTest/Public/NUTUtil.h
+++ b/Engine/Plugins/NetcodeUnitTest/NetcodeUnitTest/Source/NetcodeUnitTest/Public/NUTUtil.h
@@ -553,7 +553,7 @@ struct NUTUtil
 		bool bReturnVal = false;
 		TCHAR Temp[4096] = TEXT("");
 
-		if (ParseValue(Stream, Match, Temp, ARRAY_COUNT(Temp), bShouldStopOnComma))
+		if (ParseValue(Stream, Match, Temp, UE_ARRAY_COUNT(Temp), bShouldStopOnComma))
 		{
 			Value = Temp;
 			bReturnVal = true;
diff --git a/Engine/Plugins/Online/OnlineSubsystemSteam/Source/Private/IPAddressSteam.cpp b/Engine/Plugins/Online/OnlineSubsystemSteam/Source/Private/IPAddressSteam.cpp
--- a/Engine/Plugins/Online/OnlineSubsystemSteam/Source/Private/IPAddressSteam.cpp
+++ b/Engine/Plugins/Online/OnlineSubsystemSteam/Source/Private/IPAddressSteam.cpp
@@ -57,7 +57,7 @@ void FInternetAddrSteam::SetIp(const TCHAR* InAddr, bool& bIsValid)
 	FString SteamIPAddrStr;
 	if (InAddrStr.StartsWith(STEAM_URL_PREFIX))
 	{
-		SteamIPAddrStr = InAddrStr.Mid(ARRAY_COUNT(STEAM_URL_PREFIX) - 1);
+		SteamIPAddrStr = InAddrStr.Mid(UE_ARRAY_COUNT(STEAM_URL_PREFIX) - 1);
 	}
 	else
 	{
diff --git a/Engine/Plugins/Online/OnlineSubsystemSteam/Source/Private/OnlineAuthInterfaceSteam.cpp b/Engine/Plugins/Online/OnlineSubsystemSteam/Source/Private/OnlineAuthInterfaceSteam.cpp
--- a/Engine/Plugins/Online/OnlineSubsystemSteam/Source/Private/OnlineAuthInterfaceSteam.cpp
+++ b/Engine/Plugins/Online/OnlineSubsystemSteam/Source/Private/OnlineAuthInterfaceSteam.cpp
@@ -97,7 +97,7 @@ FString FOnlineAuthSteam::GetAuthTicket(uint32& AuthTokenHandle)
 	{
 		uint8 AuthToken[STEAM_AUTH_MAX_TICKET_LENGTH_IN_BYTES];
 		uint32 AuthTokenSize = 0;
-		AuthTokenHandle = SteamUserPtr->GetAuthSessionTicket(AuthToken, ARRAY_COUNT(AuthToken), &AuthTokenSize);
+		AuthTokenHandle = SteamUserPtr->GetAuthSessionTicket(AuthToken, UE_ARRAY_COUNT(AuthToken), &AuthTokenSize);
 		if (AuthTokenHandle != k_HAuthTicketInvalid && AuthTokenSize > 0)
 		{
 			ResultToken = BytesToHex(AuthToken, AuthTokenSize);
diff --git a/Engine/Plugins/Online/OnlineSubsystemSteam/Source/Private/OnlineIdentityInterfaceSteam.cpp b/Engine/Plugins/Online/OnlineSubsystemSteam/Source/Private/OnlineIdentityInterfaceSteam.cpp
--- a/Engine/Plugins/Online/OnlineSubsystemSteam/Source/Private/OnlineIdentityInterfaceSteam.cpp
+++ b/Engine/Plugins/Online/OnlineSubsystemSteam/Source/Private/OnlineIdentityInterfaceSteam.cpp
@@ -185,7 +185,7 @@ FString FOnlineIdentitySteam::GetAuthToken(int32 LocalUserNum) const
 		{
 			uint8 AuthToken[1024];
 			uint32 AuthTokenSize = 0;
-			if (SteamUserPtr->GetAuthSessionTicket(AuthToken, ARRAY_COUNT(AuthToken), &AuthTokenSize) != k_HAuthTicketInvalid &&
+			if (SteamUserPtr->GetAuthSessionTicket(AuthToken, UE_ARRAY_COUNT(AuthToken), &AuthTokenSize) != k_HAuthTicketInvalid &&
 				AuthTokenSize > 0)
 			{
 				ResultToken = BytesToHex(AuthToken, AuthTokenSize);
diff --git a/Engine/Plugins/Online/OnlineSubsystemSteam/Source/Private/OnlineSessionAsyncServerSteam.cpp b/Engine/Plugins/Online/OnlineSubsystemSteam/Source/Private/OnlineSessionAsyncServerSteam.cpp
--- a/Engine/Plugins/Online/OnlineSubsystemSteam/Source/Private/OnlineSessionAsyncServerSteam.cpp
+++ b/Engine/Plugins/Online/OnlineSubsystemSteam/Source/Private/OnlineSessionAsyncServerSteam.cpp
@@ -959,7 +959,7 @@ void FOnlineAsyncTaskSteamFindServerBase::ParseSearchResult(class gameserveritem
 		GameTags.ParseIntoArray(TagArray, TEXT(","), true);
 		if (TagArray.Num() > 0 && TagArray[0].StartsWith(STEAMKEY_BUILDUNIQUEID))
 		{
-			ServerBuildId = FCString::Atoi(*TagArray[0].Mid(ARRAY_COUNT(STEAMKEY_BUILDUNIQUEID)));
+			ServerBuildId = FCString::Atoi(*TagArray[0].Mid(UE_ARRAY_COUNT(STEAMKEY_BUILDUNIQUEID)));
 		}
 
 		if (ServerBuildId == BuildUniqueId)
@@ -1258,7 +1258,7 @@ void FOnlineAsyncEventSteamInviteAccepted::Finalize()
 		SessionInt->CurrentSessionSearch->SearchState = EOnlineAsyncTaskState::InProgress;
 
 		TCHAR ParsedURL[1024];
-		if (!FParse::Value(*ConnectionURL, TEXT("SteamConnectIP="), ParsedURL, ARRAY_COUNT(ParsedURL)))
+		if (!FParse::Value(*ConnectionURL, TEXT("SteamConnectIP="), ParsedURL, UE_ARRAY_COUNT(ParsedURL)))
 		{
 			UE_LOG_ONLINE_SESSION(Warning, TEXT("FOnlineAsyncEventSteamInviteAccepted: Failed to parse connection URL"));
 			return;
diff --git a/Engine/Plugins/Online/OnlineSubsystemSteam/Source/Private/OnlineSubsystemSteamPrivate.h b/Engine/Plugins/Online/OnlineSubsystemSteam/Source/Private/OnlineSubsystemSteamPrivate.h
--- a/Engine/Plugins/Online/OnlineSubsystemSteam/Source/Private/OnlineSubsystemSteamPrivate.h
+++ b/Engine/Plugins/Online/OnlineSubsystemSteam/Source/Private/OnlineSubsystemSteamPrivate.h
@@ -20,8 +20,8 @@
 #undef ONLINE_LOG_PREFIX
 #define ONLINE_LOG_PREFIX TEXT("STEAM: ")
 
-#pragma push_macro("ARRAY_COUNT")
-#undef ARRAY_COUNT
+#pragma push_macro("UE_ARRAY_COUNT")
+#undef UE_ARRAY_COUNT
 
 // Steamworks SDK headers
 #if STEAMSDK_FOUND == 0
@@ -35,4 +35,4 @@ THIRD_PARTY_INCLUDES_START
 
 THIRD_PARTY_INCLUDES_END
 
-#pragma pop_macro("ARRAY_COUNT")
+#pragma pop_macro("UE_ARRAY_COUNT")
diff --git a/Engine/Plugins/Online/OnlineSubsystemSteam/Source/Private/SteamSessionKeys.h b/Engine/Plugins/Online/OnlineSubsystemSteam/Source/Private/SteamSessionKeys.h
--- a/Engine/Plugins/Online/OnlineSubsystemSteam/Source/Private/SteamSessionKeys.h
+++ b/Engine/Plugins/Online/OnlineSubsystemSteam/Source/Private/SteamSessionKeys.h
@@ -86,7 +86,7 @@ inline bool SteamKeyToSessionSetting(const TCHAR* SteamKey, const ANSICHAR* Stea
 
 	TCHAR SteamKeyCopy[1024];
 
-	FCString::Strncpy(SteamKeyCopy, SteamKey, ARRAY_COUNT(SteamKeyCopy));
+	FCString::Strncpy(SteamKeyCopy, SteamKey, UE_ARRAY_COUNT(SteamKeyCopy));
 
 	TCHAR* DataType = FCString::Strrchr(SteamKeyCopy, '_');
 	if (DataType)
diff --git a/Engine/Plugins/Online/OnlineSubsystemTencent/Source/Private/OnlineAsyncTasksTencent.cpp b/Engine/Plugins/Online/OnlineSubsystemTencent/Source/Private/OnlineAsyncTasksTencent.cpp
--- a/Engine/Plugins/Online/OnlineSubsystemTencent/Source/Private/OnlineAsyncTasksTencent.cpp
+++ b/Engine/Plugins/Online/OnlineSubsystemTencent/Source/Private/OnlineAsyncTasksTencent.cpp
@@ -83,7 +83,7 @@ namespace
 
 		TCHAR RailKeyCopy[kRailCommonMaxKeyLength];
 
-		FCString::Strncpy(RailKeyCopy, ANSI_TO_TCHAR(RailKey), ARRAY_COUNT(RailKeyCopy));
+		FCString::Strncpy(RailKeyCopy, ANSI_TO_TCHAR(RailKey), UE_ARRAY_COUNT(RailKeyCopy));
 
 		TCHAR* DataType = FCString::Strrchr(RailKeyCopy, '_');
 		if (DataType)
diff --git a/Engine/Plugins/Online/OnlineSubsystemUtils/Source/OnlineSubsystemUtils/Private/Tests/TestSessionInterface.cpp b/Engine/Plugins/Online/OnlineSubsystemUtils/Source/OnlineSubsystemUtils/Private/Tests/TestSessionInterface.cpp
--- a/Engine/Plugins/Online/OnlineSubsystemUtils/Source/OnlineSubsystemUtils/Private/Tests/TestSessionInterface.cpp
+++ b/Engine/Plugins/Online/OnlineSubsystemUtils/Source/OnlineSubsystemUtils/Private/Tests/TestSessionInterface.cpp
@@ -476,7 +476,7 @@ bool FTestSessionInterface::Exec(UWorld* InWorld, const TCHAR* Cmd, FOutputDevic
 		else if (FParse::Command(&Cmd, TEXT("JOIN")))
 		{
 			TCHAR SearchIdxStr[256];
-			if (FParse::Token(Cmd, SearchIdxStr, ARRAY_COUNT(SearchIdxStr), true))
+			if (FParse::Token(Cmd, SearchIdxStr, UE_ARRAY_COUNT(SearchIdxStr), true))
 			{
 				int32 SearchIdx = FCString::Atoi(SearchIdxStr);
 				if (SearchIdx >= 0 && SearchIdx < SearchSettings->SearchResults.Num())
@@ -491,7 +491,7 @@ bool FTestSessionInterface::Exec(UWorld* InWorld, const TCHAR* Cmd, FOutputDevic
 			if (FParse::Command(&Cmd, TEXT("LOBBY")))
 			{
 				TCHAR FriendNameStr[256];
-				if (FParse::Token(Cmd, FriendNameStr, ARRAY_COUNT(FriendNameStr), true))
+				if (FParse::Token(Cmd, FriendNameStr, UE_ARRAY_COUNT(FriendNameStr), true))
 				{
 					for (int32 FriendIdx=0; FriendIdx<FriendsCache.Num(); FriendIdx++)
 					{
@@ -508,7 +508,7 @@ bool FTestSessionInterface::Exec(UWorld* InWorld, const TCHAR* Cmd, FOutputDevic
 			else
 			{
 				TCHAR FriendIdStr[256];
-				if (FParse::Token(Cmd, FriendIdStr, ARRAY_COUNT(FriendIdStr), true))
+				if (FParse::Token(Cmd, FriendIdStr, UE_ARRAY_COUNT(FriendIdStr), true))
 				{
 					TSharedPtr<const FUniqueNetId> FriendId = Identity->CreateUniquePlayerId((uint8*)FriendIdStr, FCString::Strlen(FriendIdStr));
 					if (!FriendId.IsValid())
diff --git a/Engine/Plugins/Online/OnlineSubsystemUtils/Source/OnlineSubsystemUtils/Private/Tests/TestVoice.cpp b/Engine/Plugins/Online/OnlineSubsystemUtils/Source/OnlineSubsystemUtils/Private/Tests/TestVoice.cpp
--- a/Engine/Plugins/Online/OnlineSubsystemUtils/Source/OnlineSubsystemUtils/Private/Tests/TestVoice.cpp
+++ b/Engine/Plugins/Online/OnlineSubsystemUtils/Source/OnlineSubsystemUtils/Private/Tests/TestVoice.cpp
@@ -191,13 +191,13 @@ void FTestVoice::SetStaticVoiceData(TArray<uint8>& VoiceData, uint32& TotalVoice
 
 	if (bTimeToQueue)
 	{
-		TotalVoiceBytes = ARRAY_COUNT(RawVoiceTestData);
+		TotalVoiceBytes = UE_ARRAY_COUNT(RawVoiceTestData);
 
 		VoiceData.Empty(TotalVoiceBytes);
 		VoiceData.AddUninitialized(TotalVoiceBytes);
 
-		VOICE_BUFFER_CHECK(VoiceData, ARRAY_COUNT(RawVoiceTestData));
-		FMemory::Memcpy(VoiceData.GetData(), RawVoiceTestData, ARRAY_COUNT(RawVoiceTestData));
+		VOICE_BUFFER_CHECK(VoiceData, UE_ARRAY_COUNT(RawVoiceTestData));
+		FMemory::Memcpy(VoiceData.GetData(), RawVoiceTestData, UE_ARRAY_COUNT(RawVoiceTestData));
 
 		LastQueueTime = CurrentTime;
 		bTimeToQueue = false;
diff --git a/Engine/Plugins/Runtime/AR/Apple/AppleARKit/Source/AppleARKit/Private/AppleARKitVideoOverlay.cpp b/Engine/Plugins/Runtime/AR/Apple/AppleARKit/Source/AppleARKit/Private/AppleARKitVideoOverlay.cpp
--- a/Engine/Plugins/Runtime/AR/Apple/AppleARKit/Source/AppleARKit/Private/AppleARKitVideoOverlay.cpp
+++ b/Engine/Plugins/Runtime/AR/Apple/AppleARKit/Source/AppleARKit/Private/AppleARKitVideoOverlay.cpp
@@ -315,7 +315,7 @@ void FAppleARKitVideoOverlay::RenderVideoOverlayWithMaterial(FRHICommandListImme
 		const uint16 Indices[] = { 0, 1, 2, 2, 1, 3 };
 
 		TResourceArray<uint16, INDEXBUFFER_ALIGNMENT> IndexBuffer;
-		const uint32 NumIndices = ARRAY_COUNT(Indices);
+		const uint32 NumIndices = UE_ARRAY_COUNT(Indices);
 		IndexBuffer.AddUninitialized(NumIndices);
 		FMemory::Memcpy(IndexBuffer.GetData(), Indices, NumIndices * sizeof(uint16));
 
diff --git a/Engine/Plugins/Runtime/AR/Google/GoogleARCore/Source/GoogleARCoreRendering/Private/GoogleARCorePassthroughCameraRenderer.cpp b/Engine/Plugins/Runtime/AR/Google/GoogleARCore/Source/GoogleARCoreRendering/Private/GoogleARCorePassthroughCameraRenderer.cpp
--- a/Engine/Plugins/Runtime/AR/Google/GoogleARCore/Source/GoogleARCoreRendering/Private/GoogleARCorePassthroughCameraRenderer.cpp
+++ b/Engine/Plugins/Runtime/AR/Google/GoogleARCore/Source/GoogleARCoreRendering/Private/GoogleARCorePassthroughCameraRenderer.cpp
@@ -79,7 +79,7 @@ void FGoogleARCorePassthroughCameraRenderer::InitializeRenderer_RenderThread(FTe
 	const uint16 Indices[] = { 0, 1, 2, 2, 1, 3};
 
 	TResourceArray<uint16, INDEXBUFFER_ALIGNMENT> IndexBuffer;
-	uint32 NumIndices = ARRAY_COUNT(Indices);
+	uint32 NumIndices = UE_ARRAY_COUNT(Indices);
 	IndexBuffer.AddUninitialized(NumIndices);
 	FMemory::Memcpy(IndexBuffer.GetData(), Indices, NumIndices * sizeof(uint16));
 
diff --git a/Engine/Plugins/Runtime/Database/SQLiteCore/Source/SQLiteCore/Private/SQLiteEmbeddedPlatform.cpp b/Engine/Plugins/Runtime/Database/SQLiteCore/Source/SQLiteCore/Private/SQLiteEmbeddedPlatform.cpp
--- a/Engine/Plugins/Runtime/Database/SQLiteCore/Source/SQLiteCore/Private/SQLiteEmbeddedPlatform.cpp
+++ b/Engine/Plugins/Runtime/Database/SQLiteCore/Source/SQLiteCore/Private/SQLiteEmbeddedPlatform.cpp
@@ -696,7 +696,7 @@ private:
 			check(InOutputBuffer);
 
 			TCHAR ErrorBuffer[1024];
-			FPlatformMisc::GetSystemErrorMessage(ErrorBuffer, ARRAY_COUNT(ErrorBuffer), 0);
+			FPlatformMisc::GetSystemErrorMessage(ErrorBuffer, UE_ARRAY_COUNT(ErrorBuffer), 0);
 			FCStringAnsi::Strncpy(InOutputBuffer, TCHAR_TO_UTF8(ErrorBuffer), InOutputBufferSizeBytes);
 			InOutputBuffer[InOutputBufferSizeBytes - 1] = 0;
 		}
diff --git a/Engine/Plugins/Runtime/GameplayAbilities/Source/GameplayAbilities/Private/AbilitySystemGlobals.cpp b/Engine/Plugins/Runtime/GameplayAbilities/Source/GameplayAbilities/Private/AbilitySystemGlobals.cpp
--- a/Engine/Plugins/Runtime/GameplayAbilities/Source/GameplayAbilities/Private/AbilitySystemGlobals.cpp
+++ b/Engine/Plugins/Runtime/GameplayAbilities/Source/GameplayAbilities/Private/AbilitySystemGlobals.cpp
@@ -130,7 +130,7 @@ const FName& UAbilitySystemGlobals::GetGameplayModEvaluationChannelAlias(EGamepl
 
 const FName& UAbilitySystemGlobals::GetGameplayModEvaluationChannelAlias(int32 Index) const
 {
-	check(Index >= 0 && Index < ARRAY_COUNT(GameplayModEvaluationChannelAliases));
+	check(Index >= 0 && Index < UE_ARRAY_COUNT(GameplayModEvaluationChannelAliases));
 	return GameplayModEvaluationChannelAliases[Index];
 }
 
diff --git a/Engine/Plugins/Runtime/GameplayAbilities/Source/GameplayAbilities/Private/GameplayEffectAggregator.cpp b/Engine/Plugins/Runtime/GameplayAbilities/Source/GameplayAbilities/Private/GameplayEffectAggregator.cpp
--- a/Engine/Plugins/Runtime/GameplayAbilities/Source/GameplayAbilities/Private/GameplayEffectAggregator.cpp
+++ b/Engine/Plugins/Runtime/GameplayAbilities/Source/GameplayAbilities/Private/GameplayEffectAggregator.cpp
@@ -128,7 +128,7 @@ void FAggregatorModChannel::RemoveModsWithActiveHandle(const FActiveGameplayEffe
 {
 	check(Handle.IsValid());
 
-	for (int32 ModOpIdx = 0; ModOpIdx < ARRAY_COUNT(Mods); ++ModOpIdx)
+	for (int32 ModOpIdx = 0; ModOpIdx < UE_ARRAY_COUNT(Mods); ++ModOpIdx)
 	{
 		Mods[ModOpIdx].RemoveAllSwap([&Handle](const FAggregatorMod& Element)
 		{
@@ -140,7 +140,7 @@ void FAggregatorModChannel::RemoveModsWithActiveHandle(const FActiveGameplayEffe
 
 void FAggregatorModChannel::AddModsFrom(const FAggregatorModChannel& Other)
 {
-	for (int32 ModOpIdx = 0; ModOpIdx < ARRAY_COUNT(Mods); ++ModOpIdx)
+	for (int32 ModOpIdx = 0; ModOpIdx < UE_ARRAY_COUNT(Mods); ++ModOpIdx)
 	{
 		Mods[ModOpIdx].Append(Other.Mods[ModOpIdx]);
 	}
@@ -148,7 +148,7 @@ void FAggregatorModChannel::AddModsFrom(const FAggregatorModChannel& Other)
 
 void FAggregatorModChannel::UpdateQualifiesOnAllMods(const FAggregatorEvaluateParameters& Parameters) const
 {
-	for (int32 ModOpIdx = 0; ModOpIdx < ARRAY_COUNT(Mods); ++ModOpIdx)
+	for (int32 ModOpIdx = 0; ModOpIdx < UE_ARRAY_COUNT(Mods); ++ModOpIdx)
 	{
 		for (const FAggregatorMod& Mod : Mods[ModOpIdx])
 		{
@@ -159,7 +159,7 @@ void FAggregatorModChannel::UpdateQualifiesOnAllMods(const FAggregatorEvaluatePa
 
 void FAggregatorModChannel::ForEachMod(FAggregatorModInfo& Info, TFunction<void (const FAggregatorModInfo&) > Func) const
 {
-	for (int32 ModOpIdx = 0; ModOpIdx < ARRAY_COUNT(Mods); ++ModOpIdx)
+	for (int32 ModOpIdx = 0; ModOpIdx < UE_ARRAY_COUNT(Mods); ++ModOpIdx)
 	{
 		Info.Op = (EGameplayModOp::Type)ModOpIdx;
 		for (const FAggregatorMod& Mod : Mods[ModOpIdx])
@@ -177,7 +177,7 @@ void FAggregatorModChannel::GetAllAggregatorMods(EGameplayModEvaluationChannel C
 
 void FAggregatorModChannel::OnActiveEffectDependenciesSwapped(const TMap<FActiveGameplayEffectHandle, FActiveGameplayEffectHandle>& SwappedDependencies)
 {
-	for (int32 ModOpIdx = 0; ModOpIdx < ARRAY_COUNT(Mods); ++ModOpIdx)
+	for (int32 ModOpIdx = 0; ModOpIdx < UE_ARRAY_COUNT(Mods); ++ModOpIdx)
 	{
 		for (FAggregatorMod& Mod : Mods[ModOpIdx])
 		{
diff --git a/Engine/Plugins/Runtime/GoogleVR/GoogleVRHMD/Source/GoogleVRHMD/Private/GoogleVRHMDRender.cpp b/Engine/Plugins/Runtime/GoogleVR/GoogleVRHMD/Source/GoogleVRHMD/Private/GoogleVRHMDRender.cpp
--- a/Engine/Plugins/Runtime/GoogleVR/GoogleVRHMD/Source/GoogleVRHMD/Private/GoogleVRHMDRender.cpp
+++ b/Engine/Plugins/Runtime/GoogleVR/GoogleVRHMD/Source/GoogleVRHMD/Private/GoogleVRHMDRender.cpp
@@ -177,7 +177,7 @@ static void ResolvePendingRenderTarget(FRHICommandListImmediate& RHICmdList, FGr
 			const uint16 Indices[] = { 0, 1, 2, 2, 1, 3, 0, 4, 5 };
 
 			TResourceArray<uint16, INDEXBUFFER_ALIGNMENT> IndexBuffer;
-			uint32 InternalNumIndices = ARRAY_COUNT(Indices);
+			uint32 InternalNumIndices = UE_ARRAY_COUNT(Indices);
 			IndexBuffer.AddUninitialized(InternalNumIndices);
 			FMemory::Memcpy(IndexBuffer.GetData(), Indices, InternalNumIndices * sizeof(uint16));
 
diff --git a/Engine/Plugins/Runtime/Oculus/OculusAvatar/Source/Public/OvrAvatarManager.cpp b/Engine/Plugins/Runtime/Oculus/OculusAvatar/Source/Public/OvrAvatarManager.cpp
--- a/Engine/Plugins/Runtime/Oculus/OculusAvatar/Source/Public/OvrAvatarManager.cpp
+++ b/Engine/Plugins/Runtime/Oculus/OculusAvatar/Source/Public/OvrAvatarManager.cpp
@@ -203,7 +203,7 @@ void UOvrAvatarManager::InitializeSDK()
 	if (!bIsInitialized)
 	{
 #if WITH_EDITORONLY_DATA
-		for (int32 AssetIndex = 0; AssetIndex < ARRAY_COUNT(AssetList); ++AssetIndex)
+		for (int32 AssetIndex = 0; AssetIndex < UE_ARRAY_COUNT(AssetList); ++AssetIndex)
 		{
 			UObject* AssetObj = AssetList[AssetIndex].TryLoad();
 			if (AssetObj != nullptr)
diff --git a/Engine/Plugins/Runtime/Oculus/OculusVR/Source/OculusHMD/Private/OculusAssetManager.cpp b/Engine/Plugins/Runtime/Oculus/OculusVR/Source/OculusHMD/Private/OculusAssetManager.cpp
--- a/Engine/Plugins/Runtime/Oculus/OculusVR/Source/OculusHMD/Private/OculusAssetManager.cpp
+++ b/Engine/Plugins/Runtime/Oculus/OculusVR/Source/OculusHMD/Private/OculusAssetManager.cpp
@@ -55,7 +55,7 @@ public:
 void FOculusAssetDirectory::LoadForCook()
 {
 	FOculusAssetRepo& AssetRepro = FOculusAssetRepo::Get();
-	for (int32 AssetIndex = 0; AssetIndex < ARRAY_COUNT(FOculusAssetDirectory::AssetListing); ++AssetIndex)
+	for (int32 AssetIndex = 0; AssetIndex < UE_ARRAY_COUNT(FOculusAssetDirectory::AssetListing); ++AssetIndex)
 	{
 		AssetRepro.LoadAndAdd(FOculusAssetDirectory::AssetListing[AssetIndex]);
 	}
diff --git a/Engine/Plugins/Runtime/Oculus/OculusVR/Source/OculusInput/Private/OculusInput.cpp b/Engine/Plugins/Runtime/Oculus/OculusVR/Source/OculusInput/Private/OculusInput.cpp
--- a/Engine/Plugins/Runtime/Oculus/OculusVR/Source/OculusInput/Private/OculusInput.cpp
+++ b/Engine/Plugins/Runtime/Oculus/OculusVR/Source/OculusInput/Private/OculusInput.cpp
@@ -359,7 +359,7 @@ void FOculusInput::SendControllerEvents()
 
 				for (FOculusTouchControllerPair& ControllerPair : ControllerPairs)
 				{
-					for( int32 HandIndex = 0; HandIndex < ARRAY_COUNT( ControllerPair.ControllerStates ); ++HandIndex )
+					for( int32 HandIndex = 0; HandIndex < UE_ARRAY_COUNT( ControllerPair.ControllerStates ); ++HandIndex )
 					{
 						FOculusTouchControllerState& State = ControllerPair.ControllerStates[ HandIndex ];
 
diff --git a/Engine/Plugins/Runtime/PhysXVehicles/Source/PhysXVehicles/Private/WheeledVehicleMovementComponent.cpp b/Engine/Plugins/Runtime/PhysXVehicles/Source/PhysXVehicles/Private/WheeledVehicleMovementComponent.cpp
--- a/Engine/Plugins/Runtime/PhysXVehicles/Source/PhysXVehicles/Private/WheeledVehicleMovementComponent.cpp
+++ b/Engine/Plugins/Runtime/PhysXVehicles/Source/PhysXVehicles/Private/WheeledVehicleMovementComponent.cpp
@@ -1738,7 +1738,7 @@ void UWheeledVehicleMovementComponent::DrawDebug(UCanvas* Canvas, float& YL, flo
 		for (uint32 w = 0; w < PVehicle->mWheelsSimData.getNbWheels(); ++w)
 		{
 			float CurX = 4;
-			for (uint32 i = 0; i < ARRAY_COUNT(GraphChannels); ++i)
+			for (uint32 i = 0; i < UE_ARRAY_COUNT(GraphChannels); ++i)
 			{
 				float OutX = GraphWidth;
 				DrawTelemetryGraph(GraphChannels[i], TelemetryData->getWheelGraph(w), Canvas, CurX, YPos, GraphWidth, GraphHeight, OutX);
diff --git a/Engine/Plugins/Runtime/Steam/SteamController/Source/SteamController/Private/SteamControllerPrivate.h b/Engine/Plugins/Runtime/Steam/SteamController/Source/SteamController/Private/SteamControllerPrivate.h
--- a/Engine/Plugins/Runtime/Steam/SteamController/Source/SteamController/Private/SteamControllerPrivate.h
+++ b/Engine/Plugins/Runtime/Steam/SteamController/Source/SteamController/Private/SteamControllerPrivate.h
@@ -14,8 +14,8 @@
 #pragma warning(disable:4996)
 #endif
 
-#pragma push_macro("ARRAY_COUNT")
-#undef ARRAY_COUNT
+#pragma push_macro("UE_ARRAY_COUNT")
+#undef UE_ARRAY_COUNT
 
 #if STEAMSDK_FOUND == 0
 #error Steam SDK not located.  Expected to be found in Engine/Source/ThirdParty/Steamworks/{SteamVersion}
@@ -32,7 +32,7 @@
  	MSVC_PRAGMA( warning( pop ) )
 #endif	// USING_CODE_ANALYSIS
 
-#pragma pop_macro("ARRAY_COUNT")
+#pragma pop_macro("UE_ARRAY_COUNT")
 
 #ifdef _MSC_VER
 #pragma warning(pop)
diff --git a/Engine/Plugins/Runtime/WindowsMixedReality/Source/WindowsMixedRealityRHI/Private/WindowsMixedRealityRHI.cpp b/Engine/Plugins/Runtime/WindowsMixedReality/Source/WindowsMixedRealityRHI/Private/WindowsMixedRealityRHI.cpp
--- a/Engine/Plugins/Runtime/WindowsMixedReality/Source/WindowsMixedRealityRHI/Private/WindowsMixedRealityRHI.cpp
+++ b/Engine/Plugins/Runtime/WindowsMixedReality/Source/WindowsMixedRealityRHI/Private/WindowsMixedRealityRHI.cpp
@@ -83,7 +83,7 @@ void FWindowsMixedRealityRHIModule::StartupModule()
 			NULL,
 			DeviceCreationFlags,
 			RequestedFeatureLevels,
-			ARRAY_COUNT(RequestedFeatureLevels),
+			UE_ARRAY_COUNT(RequestedFeatureLevels),
 			D3D11_SDK_VERSION,
 			&D3DDevice,
 			&OutFeatureLevel,
diff --git a/Engine/Plugins/Slate/SlateRemote/Source/SlateRemote/Private/Server/SlateRemoteServer.cpp b/Engine/Plugins/Slate/SlateRemote/Source/SlateRemote/Private/Server/SlateRemoteServer.cpp
--- a/Engine/Plugins/Slate/SlateRemote/Source/SlateRemote/Private/Server/SlateRemoteServer.cpp
+++ b/Engine/Plugins/Slate/SlateRemote/Source/SlateRemote/Private/Server/SlateRemoteServer.cpp
@@ -191,7 +191,7 @@ void FSlateRemoteServer::ProcessTouchMessage( const FSlateRemoteServerMessage& M
 	}
 	else
 	{
-		//checkf(Message.Handle < ARRAY_COUNT(LastTouchPositions), TEXT("Received handle %d, but it's too big for our array"), Message.Handle);
+		//checkf(Message.Handle < UE_ARRAY_COUNT(LastTouchPositions), TEXT("Received handle %d, but it's too big for our array"), Message.Handle);
 	}
 }
 
diff --git a/Engine/Plugins/Tests/EditorTests/Source/EditorTests/Private/UnrealEd/EditorBuildPromotionTests.cpp b/Engine/Plugins/Tests/EditorTests/Source/EditorTests/Private/UnrealEd/EditorBuildPromotionTests.cpp
--- a/Engine/Plugins/Tests/EditorTests/Source/EditorTests/Private/UnrealEd/EditorBuildPromotionTests.cpp
+++ b/Engine/Plugins/Tests/EditorTests/Source/EditorTests/Private/UnrealEd/EditorBuildPromotionTests.cpp
@@ -1288,7 +1288,7 @@ namespace BuildPromotionTestHelper
 				const FString SurroundChannels[] = { TEXT("_fl"), TEXT("_fr"), TEXT("_fc"), TEXT("_lf"), TEXT("_sl"), TEXT("_sr"), TEXT("_bl"), TEXT("_br") };
 
 				USoundWave* ImportedSound = NULL;
-				for (int32 ChannelID = 0; ChannelID < ARRAY_COUNT(SurroundChannels); ++ChannelID)
+				for (int32 ChannelID = 0; ChannelID < UE_ARRAY_COUNT(SurroundChannels); ++ChannelID)
 				{
 					const FString ChannelFileName = FString::Printf(TEXT("%s%s.WAV"),*BaseFileName,*SurroundChannels[ChannelID]);
 					if (FPaths::FileExists(ChannelFileName))
@@ -1477,7 +1477,7 @@ namespace BuildPromotionTestHelper
 					TEXT("Black"),	TEXT("(R=0.0f,G=0.0f,B=0.0f)"),
 					TEXT("White"),	TEXT("(R=1.0f,G=1.0f,B=1.0f)") };
 
-				const int32 ChosenIndex = FMath::RandHelper(ARRAY_COUNT(AvailableColors) / 2);
+				const int32 ChosenIndex = FMath::RandHelper(UE_ARRAY_COUNT(AvailableColors) / 2);
 				ChosenMaterialColor = AvailableColors[ChosenIndex * 2];
 				const FString ColorValue = AvailableColors[(ChosenIndex * 2)+1];
 
diff --git a/Engine/Plugins/VirtualProduction/Takes/Source/TakeRecorderSources/Private/TakeRecorderSourcesModule.cpp b/Engine/Plugins/VirtualProduction/Takes/Source/TakeRecorderSources/Private/TakeRecorderSourcesModule.cpp
--- a/Engine/Plugins/VirtualProduction/Takes/Source/TakeRecorderSources/Private/TakeRecorderSourcesModule.cpp
+++ b/Engine/Plugins/VirtualProduction/Takes/Source/TakeRecorderSources/Private/TakeRecorderSourcesModule.cpp
@@ -449,7 +449,7 @@ public:
 		const TCHAR* Str = InStr;
 		EFilterType FilterType = EFilterType::None;
 		TCHAR Filter[128];
-		if (FParse::Token(Str, Filter, ARRAY_COUNT(Filter), 0))
+		if (FParse::Token(Str, Filter, UE_ARRAY_COUNT(Filter), 0))
 		{
 			FString const FilterStr = Filter;
 			if (FilterStr == TEXT("all"))
@@ -475,7 +475,7 @@ public:
 		if (FilterType == EFilterType::Actor || FilterType == EFilterType::Class)
 		{
 			TCHAR Specifier[128];
-			if (FParse::Token(Str, Specifier, ARRAY_COUNT(Specifier), 0))
+			if (FParse::Token(Str, Specifier, UE_ARRAY_COUNT(Specifier), 0))
 			{
 				FString const SpecifierStr = FString(Specifier).TrimStart();
 
diff --git a/Engine/Source/Developer/Android/AndroidTargetPlatform/Private/AndroidTargetPlatform.h b/Engine/Source/Developer/Android/AndroidTargetPlatform/Private/AndroidTargetPlatform.h
--- a/Engine/Source/Developer/Android/AndroidTargetPlatform/Private/AndroidTargetPlatform.h
+++ b/Engine/Source/Developer/Android/AndroidTargetPlatform/Private/AndroidTargetPlatform.h
@@ -512,7 +512,7 @@ public:
 			else
 			{
 				// perform any remapping away from defaults
-				for (int32 RemapIndex = 0; RemapIndex < ARRAY_COUNT(FormatRemap); ++RemapIndex)
+				for (int32 RemapIndex = 0; RemapIndex < UE_ARRAY_COUNT(FormatRemap); ++RemapIndex)
 				{
 					if (TextureFormatName == FormatRemap[RemapIndex][0])
 					{
@@ -543,7 +543,7 @@ public:
 
 		GetAllDefaultTextureFormats(this, OutFormats, false);
 
-		for (int32 RemapIndex = 0; RemapIndex < ARRAY_COUNT(FormatRemap); ++RemapIndex)
+		for (int32 RemapIndex = 0; RemapIndex < UE_ARRAY_COUNT(FormatRemap); ++RemapIndex)
 		{
 			OutFormats.Remove(FormatRemap[RemapIndex][0]);
 			OutFormats.AddUnique(FormatRemap[RemapIndex][1]);
diff --git a/Engine/Source/Developer/AudioFormatOpus/Private/AudioFormatOpus.cpp b/Engine/Source/Developer/AudioFormatOpus/Private/AudioFormatOpus.cpp
--- a/Engine/Source/Developer/AudioFormatOpus/Private/AudioFormatOpus.cpp
+++ b/Engine/Source/Developer/AudioFormatOpus/Private/AudioFormatOpus.cpp
@@ -428,7 +428,7 @@ public:
 		};
 
 		// look for the next highest valid rate
-		for (int32 Index = ARRAY_COUNT(ValidSampleRates) - 2; Index >= 0; Index--)
+		for (int32 Index = UE_ARRAY_COUNT(ValidSampleRates) - 2; Index >= 0; Index--)
 		{
 			if (SampleRate > ValidSampleRates[Index])
 			{
diff --git a/Engine/Source/Developer/AutomationController/Private/AutomationControllerManger.cpp b/Engine/Source/Developer/AutomationController/Private/AutomationControllerManger.cpp
--- a/Engine/Source/Developer/AutomationController/Private/AutomationControllerManger.cpp
+++ b/Engine/Source/Developer/AutomationController/Private/AutomationControllerManger.cpp
@@ -428,7 +428,7 @@ FString FAutomationControllerManager::SlugString(const FString& DisplayString) c
 	// Convert the display label, which may consist of just about any possible character, into a
 	// suitable name for a UObject (remove whitespace, certain symbols, etc.)
 	{
-		for ( int32 BadCharacterIndex = 0; BadCharacterIndex < ARRAY_COUNT(INVALID_OBJECTNAME_CHARACTERS) - 1; ++BadCharacterIndex )
+		for ( int32 BadCharacterIndex = 0; BadCharacterIndex < UE_ARRAY_COUNT(INVALID_OBJECTNAME_CHARACTERS) - 1; ++BadCharacterIndex )
 		{
 			const TCHAR TestChar[2] = { INVALID_OBJECTNAME_CHARACTERS[BadCharacterIndex], 0 };
 			const int32 NumReplacedChars = GeneratedName.ReplaceInline(TestChar, TEXT(""));
diff --git a/Engine/Source/Developer/CrashDebugHelper/Private/Windows/WindowsPlatformStackWalkExt.cpp b/Engine/Source/Developer/CrashDebugHelper/Private/Windows/WindowsPlatformStackWalkExt.cpp
--- a/Engine/Source/Developer/CrashDebugHelper/Private/Windows/WindowsPlatformStackWalkExt.cpp
+++ b/Engine/Source/Developer/CrashDebugHelper/Private/Windows/WindowsPlatformStackWalkExt.cpp
@@ -256,7 +256,7 @@ void FWindowsPlatformStackWalkExt::SetSymbolPathsFromModules()
 
 	TCHAR SymbolPath[16384] = { 0 };
 
-	Symbol->GetSymbolPathWide( SymbolPath, ARRAY_COUNT(SymbolPath), NULL );
+	Symbol->GetSymbolPathWide( SymbolPath, UE_ARRAY_COUNT(SymbolPath), NULL );
 	TArray<FString> SymbolPaths;
 	FString( SymbolPath ).ParseIntoArray(SymbolPaths, TEXT(";"), true );
 
@@ -266,7 +266,7 @@ void FWindowsPlatformStackWalkExt::SetSymbolPathsFromModules()
 		UE_LOG( LogCrashDebugHelper, Log, TEXT( "    %s" ), *It );
 	}
 
-	Symbol->GetImagePathWide( SymbolPath, ARRAY_COUNT( SymbolPath ), NULL );
+	Symbol->GetImagePathWide( SymbolPath, UE_ARRAY_COUNT( SymbolPath ), NULL );
 	TArray<FString> ImagePaths;
 	FString( SymbolPath ).ParseIntoArray( ImagePaths, TEXT( ";" ), true );
 	
@@ -415,7 +415,7 @@ void FWindowsPlatformStackWalkExt::GetExceptionInfo()
 	ULONG ProcessID = 0;
 	ULONG ThreadId = 0;
 	TCHAR Description[MAX_PATH] = { 0 };
-	Control->GetLastEventInformationWide( &ExceptionType, &ProcessID, &ThreadId, NULL, 0, NULL, Description, ARRAY_COUNT( Description ), NULL );
+	Control->GetLastEventInformationWide( &ExceptionType, &ProcessID, &ThreadId, NULL, 0, NULL, Description, UE_ARRAY_COUNT( Description ), NULL );
 
 	Exception.Code = ExceptionType;
 	Exception.ProcessId = ProcessID;
diff --git a/Engine/Source/Developer/DesktopPlatform/Private/Mac/MacNativeFeedbackContext.cpp b/Engine/Source/Developer/DesktopPlatform/Private/Mac/MacNativeFeedbackContext.cpp
--- a/Engine/Source/Developer/DesktopPlatform/Private/Mac/MacNativeFeedbackContext.cpp
+++ b/Engine/Source/Developer/DesktopPlatform/Private/Mac/MacNativeFeedbackContext.cpp
@@ -314,8 +314,8 @@ void FMacNativeFeedbackContext::Serialize( const TCHAR* Data, ELogVerbosity::Typ
 					
 					// turn the string into a bunch of 0's and 1's
 					TCHAR String[9];
-					FMemory::Memset(String, 0, sizeof(TCHAR) * ARRAY_COUNT(String));
-					FCString::Strncpy(String, Data, ARRAY_COUNT(String));
+					FMemory::Memset(String, 0, sizeof(TCHAR) * UE_ARRAY_COUNT(String));
+					FCString::Strncpy(String, Data, UE_ARRAY_COUNT(String));
 					for (TCHAR* S = String; *S; S++)
 					{
 						*S -= '0';
diff --git a/Engine/Source/Developer/GameplayDebugger/Private/GameplayDebuggerLocalController.cpp b/Engine/Source/Developer/GameplayDebugger/Private/GameplayDebuggerLocalController.cpp
--- a/Engine/Source/Developer/GameplayDebugger/Private/GameplayDebuggerLocalController.cpp
+++ b/Engine/Source/Developer/GameplayDebugger/Private/GameplayDebuggerLocalController.cpp
@@ -68,10 +68,10 @@ void UGameplayDebuggerLocalController::Initialize(AGameplayDebuggerCategoryRepli
 		SettingsCDO->CategorySlot5, SettingsCDO->CategorySlot6, SettingsCDO->CategorySlot7, SettingsCDO->CategorySlot8, SettingsCDO->CategorySlot9 };
 
 	bool bIsNumpadOnly = true;
-	for (int32 Idx = 0; Idx < ARRAY_COUNT(CategorySlots); Idx++)
+	for (int32 Idx = 0; Idx < UE_ARRAY_COUNT(CategorySlots); Idx++)
 	{
 		bool bHasPattern = false;
-		for (int32 PatternIdx = 0; PatternIdx < ARRAY_COUNT(NumpadKeys); PatternIdx++)
+		for (int32 PatternIdx = 0; PatternIdx < UE_ARRAY_COUNT(NumpadKeys); PatternIdx++)
 		{
 			if (CategorySlots[Idx] == NumpadKeys[PatternIdx])
 			{
diff --git a/Engine/Source/Developer/GameplayDebugger/Private/GameplayDebuggerTypes.cpp b/Engine/Source/Developer/GameplayDebugger/Private/GameplayDebuggerTypes.cpp
--- a/Engine/Source/Developer/GameplayDebugger/Private/GameplayDebuggerTypes.cpp
+++ b/Engine/Source/Developer/GameplayDebugger/Private/GameplayDebuggerTypes.cpp
@@ -445,7 +445,7 @@ void FGameplayDebuggerCanvasContext::PrintAt(float PosX, float PosY, const FColo
 \
 	/* first, try using the stack buffer */ \
 	Buffer = StackBuffer; \
-	GET_VARARGS_RESULT( Buffer, ARRAY_COUNT(StackBuffer), ARRAY_COUNT(StackBuffer) - 1, Fmt, Fmt, Result ); \
+	GET_VARARGS_RESULT( Buffer, UE_ARRAY_COUNT(StackBuffer), UE_ARRAY_COUNT(StackBuffer) - 1, Fmt, Fmt, Result ); \
 \
 	/* if that fails, then use heap allocation to make enough space */ \
 	while(Result == -1) \
diff --git a/Engine/Source/Developer/IOS/IOSTargetPlatform/Private/IOSTargetPlatform.cpp b/Engine/Source/Developer/IOS/IOSTargetPlatform/Private/IOSTargetPlatform.cpp
--- a/Engine/Source/Developer/IOS/IOSTargetPlatform/Private/IOSTargetPlatform.cpp
+++ b/Engine/Source/Developer/IOS/IOSTargetPlatform/Private/IOSTargetPlatform.cpp
@@ -625,7 +625,7 @@ void FIOSTargetPlatform::GetTextureFormats( const UTexture* Texture, TArray< TAr
 		TArray<FName> TextureFormatNamesASTC(TextureFormatNames);
 		for (FName& TextureFormatName : TextureFormatNamesASTC)
 		{
-			for (int32 RemapIndex = 0; RemapIndex < ARRAY_COUNT(FormatRemap); RemapIndex += 3)
+			for (int32 RemapIndex = 0; RemapIndex < UE_ARRAY_COUNT(FormatRemap); RemapIndex += 3)
 			{
 				if (TextureFormatName == FormatRemap[RemapIndex])
 				{
@@ -642,7 +642,7 @@ void FIOSTargetPlatform::GetTextureFormats( const UTexture* Texture, TArray< TAr
 		TArray<FName> TextureFormatNamesPVRTC(TextureFormatNames);
 		for (FName& TextureFormatName : TextureFormatNamesPVRTC)
 		{
-			for (int32 RemapIndex = 0; RemapIndex < ARRAY_COUNT(FormatRemap); RemapIndex += 3)
+			for (int32 RemapIndex = 0; RemapIndex < UE_ARRAY_COUNT(FormatRemap); RemapIndex += 3)
 			{
 				if (TextureFormatName == FormatRemap[RemapIndex])
 				{
@@ -674,7 +674,7 @@ void FIOSTargetPlatform::GetAllTextureFormats(TArray<FName>& OutFormats) const
 
 	GetAllDefaultTextureFormats(this, OutFormats, false);
 
-	for (int32 RemapIndex = 0; RemapIndex < ARRAY_COUNT(FormatRemap); RemapIndex += 3)
+	for (int32 RemapIndex = 0; RemapIndex < UE_ARRAY_COUNT(FormatRemap); RemapIndex += 3)
 	{
 		OutFormats.Remove(FormatRemap[RemapIndex+0]);
 	}
@@ -682,14 +682,14 @@ void FIOSTargetPlatform::GetAllTextureFormats(TArray<FName>& OutFormats) const
 	// include the formats we want (use ASTC first so that it is preferred at runtime if they both exist and it's supported)
 	if (bIncludeASTC)
 	{
-		for (int32 RemapIndex = 0; RemapIndex < ARRAY_COUNT(FormatRemap); RemapIndex += 3)
+		for (int32 RemapIndex = 0; RemapIndex < UE_ARRAY_COUNT(FormatRemap); RemapIndex += 3)
 		{
 			OutFormats.AddUnique(FormatRemap[RemapIndex + 2]);
 		}
 	}
 	if (bIncludePVRTC)
 	{
-		for (int32 RemapIndex = 0; RemapIndex < ARRAY_COUNT(FormatRemap); RemapIndex += 3)
+		for (int32 RemapIndex = 0; RemapIndex < UE_ARRAY_COUNT(FormatRemap); RemapIndex += 3)
 		{
 			OutFormats.AddUnique(FormatRemap[RemapIndex + 1]);
 		}
diff --git a/Engine/Source/Developer/Linux/LinuxTargetPlatform/Private/LinuxTargetDevice.h b/Engine/Source/Developer/Linux/LinuxTargetPlatform/Private/LinuxTargetDevice.h
--- a/Engine/Source/Developer/Linux/LinuxTargetPlatform/Private/LinuxTargetDevice.h
+++ b/Engine/Source/Developer/Linux/LinuxTargetPlatform/Private/LinuxTargetDevice.h
@@ -266,7 +266,7 @@ public:
 		char ReadLinkCmd[ReadLinkSize] = { 0 };
 		FCStringAnsi::Sprintf(ReadLinkCmd, "/proc/%lld/exe", ProcessId);
 		char ProcessPath[UNIX_MAX_PATH + 1] = { 0 };
-		int32 Ret = readlink(ReadLinkCmd, ProcessPath, ARRAY_COUNT(ProcessPath) - 1);
+		int32 Ret = readlink(ReadLinkCmd, ProcessPath, UE_ARRAY_COUNT(ProcessPath) - 1);
 		if (Ret != -1)
 		{
 			struct stat st;
diff --git a/Engine/Source/Developer/Lumin/LuminTargetPlatform/Private/LuminTargetPlatform.cpp b/Engine/Source/Developer/Lumin/LuminTargetPlatform/Private/LuminTargetPlatform.cpp
--- a/Engine/Source/Developer/Lumin/LuminTargetPlatform/Private/LuminTargetPlatform.cpp
+++ b/Engine/Source/Developer/Lumin/LuminTargetPlatform/Private/LuminTargetPlatform.cpp
@@ -242,7 +242,7 @@ void FLuminTargetPlatform::GetTextureFormats(const UTexture* InTexture, TArray<
 		}
 		else
 		{
-			for (int32 RemapIndex = 0; RemapIndex < ARRAY_COUNT(FormatRemap); ++RemapIndex)
+			for (int32 RemapIndex = 0; RemapIndex < UE_ARRAY_COUNT(FormatRemap); ++RemapIndex)
 			{
 				if (TextureFormatName == FormatRemap[RemapIndex][0])
 				{
@@ -259,13 +259,13 @@ void FLuminTargetPlatform::GetAllTextureFormats(TArray<FName>& OutFormats) const
 {
 	GetAllDefaultTextureFormats(this, OutFormats, false);
 
-	for (int32 RemapIndex = 0; RemapIndex < ARRAY_COUNT(FormatRemap); RemapIndex++)
+	for (int32 RemapIndex = 0; RemapIndex < UE_ARRAY_COUNT(FormatRemap); RemapIndex++)
 	{
 		OutFormats.Remove(FormatRemap[RemapIndex][0]);
 	}
 
 	// include the formats we want
-	for (int32 RemapIndex = 0; RemapIndex < ARRAY_COUNT(FormatRemap); RemapIndex++)
+	for (int32 RemapIndex = 0; RemapIndex < UE_ARRAY_COUNT(FormatRemap); RemapIndex++)
 	{
 		OutFormats.AddUnique(FormatRemap[RemapIndex][1]);
 	}
diff --git a/Engine/Source/Developer/MaterialUtilities/Private/MaterialUtilities.cpp b/Engine/Source/Developer/MaterialUtilities/Private/MaterialUtilities.cpp
--- a/Engine/Source/Developer/MaterialUtilities/Private/MaterialUtilities.cpp
+++ b/Engine/Source/Developer/MaterialUtilities/Private/MaterialUtilities.cpp
@@ -1625,7 +1625,7 @@ void FMaterialUtilities::AnalyzeMaterial(UMaterialInterface* InMaterial, const s
 	PropertyBeingBaked[MP_Metallic] = InMaterialSettings.bOpacityMap;
 	PropertyBeingBaked[MP_EmissiveColor] = InMaterialSettings.bEmissiveMap;
 
-	for (int32 PropertyIndex = 0; PropertyIndex < ARRAY_COUNT(PropertyBeingBaked); ++PropertyIndex)
+	for (int32 PropertyIndex = 0; PropertyIndex < UE_ARRAY_COUNT(PropertyBeingBaked); ++PropertyIndex)
 	{
 		if (PropertyBeingBaked[PropertyIndex])
 		{
@@ -2627,7 +2627,7 @@ bool FMaterialUtilities::RenderMaterialPropertyToTexture(struct FMaterialMergeDa
 
 	FMaterialRenderProxy* MaterialProxy = nullptr;
 
-	check(InMaterialProperty >= 0 && InMaterialProperty < ARRAY_COUNT(InMaterialData.ProxyCache->Proxies));
+	check(InMaterialProperty >= 0 && InMaterialProperty < UE_ARRAY_COUNT(InMaterialData.ProxyCache->Proxies));
 	if (InMaterialData.ProxyCache->Proxies[InMaterialProperty])
 	{
 		MaterialProxy = InMaterialData.ProxyCache->Proxies[InMaterialProperty];
@@ -2821,7 +2821,7 @@ FExportMaterialProxyCache::~FExportMaterialProxyCache()
 
 void FExportMaterialProxyCache::Release()
 {
-	for (int32 PropertyIndex = 0; PropertyIndex < ARRAY_COUNT(Proxies); PropertyIndex++)
+	for (int32 PropertyIndex = 0; PropertyIndex < UE_ARRAY_COUNT(Proxies); PropertyIndex++)
 	{
 		FMaterialRenderProxy* Proxy = Proxies[PropertyIndex];
 		if (Proxy)
diff --git a/Engine/Source/Developer/ShaderCompilerCommon/Private/ShaderCompilerCommon.cpp b/Engine/Source/Developer/ShaderCompilerCommon/Private/ShaderCompilerCommon.cpp
--- a/Engine/Source/Developer/ShaderCompilerCommon/Private/ShaderCompilerCommon.cpp
+++ b/Engine/Source/Developer/ShaderCompilerCommon/Private/ShaderCompilerCommon.cpp
@@ -1201,7 +1201,7 @@ namespace CrossCompiler
 	};
 
 	/** Compile time check to verify that the GL mapping tables are up-to-date. */
-	static_assert(SF_NumFrequencies == ARRAY_COUNT(FrequencyStringTable), "NumFrequencies changed. Please update tables.");
+	static_assert(SF_NumFrequencies == UE_ARRAY_COUNT(FrequencyStringTable), "NumFrequencies changed. Please update tables.");
 
 	const TCHAR* GetFrequencyName(EShaderFrequency Frequency)
 	{
diff --git a/Engine/Source/Developer/ShaderFormatOpenGL/Private/OpenGLShaderCompiler.cpp b/Engine/Source/Developer/ShaderFormatOpenGL/Private/OpenGLShaderCompiler.cpp
--- a/Engine/Source/Developer/ShaderFormatOpenGL/Private/OpenGLShaderCompiler.cpp
+++ b/Engine/Source/Developer/ShaderFormatOpenGL/Private/OpenGLShaderCompiler.cpp
@@ -489,7 +489,7 @@ GLenum GLFrequencyTable[] =
 	GLenum(0), // SF_RayCallable
 };
 
-static_assert(ARRAY_COUNT(GLFrequencyTable) == SF_NumFrequencies, "Frequency table size mismatch.");
+static_assert(UE_ARRAY_COUNT(GLFrequencyTable) == SF_NumFrequencies, "Frequency table size mismatch.");
 
 static inline bool IsDigit(TCHAR Char)
 {
diff --git a/Engine/Source/Developer/ShaderPreprocessor/Private/PreprocessorPrivate.h b/Engine/Source/Developer/ShaderPreprocessor/Private/PreprocessorPrivate.h
--- a/Engine/Source/Developer/ShaderPreprocessor/Private/PreprocessorPrivate.h
+++ b/Engine/Source/Developer/ShaderPreprocessor/Private/PreprocessorPrivate.h
@@ -26,7 +26,7 @@ inline EMessageType FilterMcppError(const FString& ErrorMsg)
 		TEXT("with no newline, supplemented newline"),
 		TEXT("Converted [CR+LF] to [LF]")
 	};
-	const int32 FilteredSubstringCount = ARRAY_COUNT(SubstringsToFilter);
+	const int32 FilteredSubstringCount = UE_ARRAY_COUNT(SubstringsToFilter);
 
 	if (ErrorMsg.Contains(TEXT("UE4SHADERMETADATA")))
 	{
diff --git a/Engine/Source/Developer/SlateFileDialogs/Private/SlateFileDlgWindow.cpp b/Engine/Source/Developer/SlateFileDialogs/Private/SlateFileDlgWindow.cpp
--- a/Engine/Source/Developer/SlateFileDialogs/Private/SlateFileDlgWindow.cpp
+++ b/Engine/Source/Developer/SlateFileDialogs/Private/SlateFileDlgWindow.cpp
@@ -874,14 +874,14 @@ void SSlateFileOpenDlg::BuildDirectoryPath()
 			BuiltPath = BuiltPath + TEXT("/") + AbsPath.Left(Idx);
 		}
 
-		FCString::Strcpy(Temp, ARRAY_COUNT(Temp), &AbsPath[Idx < AbsPath.Len() - 1 ? Idx + 1 : Idx]);
+		FCString::Strcpy(Temp, UE_ARRAY_COUNT(Temp), &AbsPath[Idx < AbsPath.Len() - 1 ? Idx + 1 : Idx]);
 
 		DirectoryNodesArray.Add(FDirNode(AbsPath.Left(Idx), nullptr));
 	}
 	else if (PLATFORM_LINUX)
 	{
 		// start with system base directory
-		FCString::Strcpy(Temp, ARRAY_COUNT(Temp), *AbsPath);
+		FCString::Strcpy(Temp, UE_ARRAY_COUNT(Temp), *AbsPath);
 
 		BuiltPath = "/";
 		DirectoryNodesArray.Add(FDirNode(FString(TEXT("/")), nullptr));
@@ -1501,7 +1501,7 @@ void SSlateFileOpenDlg::ParseFilters()
 			if (FilterNameArray.Num() == 0)
 			{
 				TCHAR Temp[MAX_FILTER_LENGTH] = {0};
-				FCString::Strcpy(Temp, ARRAY_COUNT(Temp), *Filters);
+				FCString::Strcpy(Temp, UE_ARRAY_COUNT(Temp), *Filters);
 
 				// break path into tokens
 				TCHAR *ContextStr = nullptr;
@@ -1551,7 +1551,7 @@ bool SSlateFileOpenDlg::GetFilterExtension(FString &OutString)
 
 	// make a copy of filter string that we can modify
 	TCHAR Temp[MAX_FILTER_LENGTH] = {0};
-	FCString::Strcpy(Temp, ARRAY_COUNT(Temp), *(*FilterNameArray[FilterIndex].Get()));
+	FCString::Strcpy(Temp, UE_ARRAY_COUNT(Temp), *(*FilterNameArray[FilterIndex].Get()));
 
 	// find start of extension
 	TCHAR *FilterExt = FCString::Strchr(Temp, '.');
diff --git a/Engine/Source/Developer/StandaloneRenderer/Private/Windows/D3D/SlateD3DShaders.cpp b/Engine/Source/Developer/StandaloneRenderer/Private/Windows/D3D/SlateD3DShaders.cpp
--- a/Engine/Source/Developer/StandaloneRenderer/Private/Windows/D3D/SlateD3DShaders.cpp
+++ b/Engine/Source/Developer/StandaloneRenderer/Private/Windows/D3D/SlateD3DShaders.cpp
@@ -388,7 +388,7 @@ FSlateDefaultVS::FSlateDefaultVS()
 		{ "COLOR",		0, DXGI_FORMAT_B8G8R8A8_UNORM,		0, D3D11_APPEND_ALIGNED_ELEMENT,	D3D11_INPUT_PER_VERTEX_DATA,	0 },
 	};
 
-	Create( FString::Printf( TEXT("%s/StandaloneRenderer/D3D/SlateDefaultVertexShader.hlsl"), FPlatformProcess::ShaderDir() ), TEXT("Main"), TEXT("vs_4_0"), Layout, ARRAY_COUNT(Layout) );
+	Create( FString::Printf( TEXT("%s/StandaloneRenderer/D3D/SlateDefaultVertexShader.hlsl"), FPlatformProcess::ShaderDir() ), TEXT("Main"), TEXT("vs_4_0"), Layout, UE_ARRAY_COUNT(Layout) );
 }
 
 void FSlateDefaultVS::SetViewProjection( const FMatrix& ViewProjectionMatrix )
diff --git a/Engine/Source/Developer/TextureFormatASTC/Private/TextureFormatASTC.cpp b/Engine/Source/Developer/TextureFormatASTC/Private/TextureFormatASTC.cpp
--- a/Engine/Source/Developer/TextureFormatASTC/Private/TextureFormatASTC.cpp
+++ b/Engine/Source/Developer/TextureFormatASTC/Private/TextureFormatASTC.cpp
@@ -348,7 +348,7 @@ public:
 
 	virtual void GetSupportedFormats(TArray<FName>& OutFormats) const override
 	{
-		for (int32 i = 0; i < ARRAY_COUNT(GSupportedTextureFormatNames); ++i)
+		for (int32 i = 0; i < UE_ARRAY_COUNT(GSupportedTextureFormatNames); ++i)
 		{
 			OutFormats.Add(GSupportedTextureFormatNames[i]);
 		}
diff --git a/Engine/Source/Developer/TextureFormatAndroid/Private/TextureFormatAndroid.cpp b/Engine/Source/Developer/TextureFormatAndroid/Private/TextureFormatAndroid.cpp
--- a/Engine/Source/Developer/TextureFormatAndroid/Private/TextureFormatAndroid.cpp
+++ b/Engine/Source/Developer/TextureFormatAndroid/Private/TextureFormatAndroid.cpp
@@ -143,7 +143,7 @@ class FTextureFormatAndroid : public ITextureFormat
 
 	virtual void GetSupportedFormats(TArray<FName>& OutFormats) const override
 	{
-		for (int32 i = 0; i < ARRAY_COUNT(GSupportedTextureFormatNames); ++i)
+		for (int32 i = 0; i < UE_ARRAY_COUNT(GSupportedTextureFormatNames); ++i)
 		{
 			OutFormats.Add(GSupportedTextureFormatNames[i]);
 		}
diff --git a/Engine/Source/Developer/TraceAnalysis/Private/ControlClient.cpp b/Engine/Source/Developer/TraceAnalysis/Private/ControlClient.cpp
--- a/Engine/Source/Developer/TraceAnalysis/Private/ControlClient.cpp
+++ b/Engine/Source/Developer/TraceAnalysis/Private/ControlClient.cpp
@@ -134,7 +134,7 @@ void FControlClient::FormatAndSend(const TCHAR* Format, ...)
     TCHAR Buffer[512];
 	va_list Args;
 	va_start(Args, Format);
-	int Length = FCString::GetVarArgs(Buffer, ARRAY_COUNT(Buffer), Format, Args);
+	int Length = FCString::GetVarArgs(Buffer, UE_ARRAY_COUNT(Buffer), Format, Args);
     if (Length > sizeof(Buffer))
     {
         Length = sizeof(Buffer);
diff --git a/Engine/Source/Developer/VulkanShaderFormat/Private/VulkanShaderCompiler.cpp b/Engine/Source/Developer/VulkanShaderFormat/Private/VulkanShaderCompiler.cpp
--- a/Engine/Source/Developer/VulkanShaderFormat/Private/VulkanShaderCompiler.cpp
+++ b/Engine/Source/Developer/VulkanShaderFormat/Private/VulkanShaderCompiler.cpp
@@ -252,7 +252,7 @@ static uint32 GetTypeComponents(const FString& Type)
 	static const FString TypePrefix[] = { "f", "i", "u" };
 	uint32 Components = 0;
 	int32 PrefixLength = 0;
-	for (uint32 i = 0; i<ARRAY_COUNT(TypePrefix); i++)
+	for (uint32 i = 0; i<UE_ARRAY_COUNT(TypePrefix); i++)
 	{
 		const FString& Prefix = TypePrefix[i];
 		const int32 CmpLength = Type.Contains(Prefix, ESearchCase::CaseSensitive, ESearchDir::FromStart);
diff --git a/Engine/Source/Developer/Windows/LiveCoding/Private/LiveCodingModule.cpp b/Engine/Source/Developer/Windows/LiveCoding/Private/LiveCodingModule.cpp
--- a/Engine/Source/Developer/Windows/LiveCoding/Private/LiveCodingModule.cpp
+++ b/Engine/Source/Developer/Windows/LiveCoding/Private/LiveCodingModule.cpp
@@ -362,7 +362,7 @@ void FLiveCodingModule::UpdateModules()
 	{
 #if IS_MONOLITHIC
 		wchar_t FullFilePath[WINDOWS_MAX_PATH];
-		verify(GetModuleFileName(hInstance, FullFilePath, ARRAY_COUNT(FullFilePath)));
+		verify(GetModuleFileName(hInstance, FullFilePath, UE_ARRAY_COUNT(FullFilePath)));
 		LppEnableModule(FullFilePath);
 #else
 		TArray<FModuleStatus> ModuleStatuses;
diff --git a/Engine/Source/Developer/Windows/LiveCodingServer/Private/External/LC_ServerCommandThread.cpp b/Engine/Source/Developer/Windows/LiveCodingServer/Private/External/LC_ServerCommandThread.cpp
--- a/Engine/Source/Developer/Windows/LiveCodingServer/Private/External/LC_ServerCommandThread.cpp
+++ b/Engine/Source/Developer/Windows/LiveCodingServer/Private/External/LC_ServerCommandThread.cpp
@@ -789,7 +789,7 @@ void ServerCommandThread::CompileChanges(bool didAllProcessesMakeProgress)
 				// If this file has a .lc.obj suffix, temporarily replace the original .obj file while generating the patch.
 				// It'd be nice to track this explicitly inside Live++ and just load the new file, but it requires a lot of changes and would make upgrades difficult.
 				static const TCHAR Suffix[] = TEXT(".lc.obj");
-				static const size_t SuffixLen = ARRAY_COUNT(Suffix) - 1;
+				static const size_t SuffixLen = UE_ARRAY_COUNT(Suffix) - 1;
 				if (NormalizedObjectFile.length() >= SuffixLen && _wcsicmp(NormalizedObjectFile.c_str() + NormalizedObjectFile.length() - SuffixLen, Suffix) == 0)
 				{
 					// Get the original filename
diff --git a/Engine/Source/Developer/Windows/ShaderFormatD3D/Private/D3DShaderCompiler.cpp b/Engine/Source/Developer/Windows/ShaderFormatD3D/Private/D3DShaderCompiler.cpp
--- a/Engine/Source/Developer/Windows/ShaderFormatD3D/Private/D3DShaderCompiler.cpp
+++ b/Engine/Source/Developer/Windows/ShaderFormatD3D/Private/D3DShaderCompiler.cpp
@@ -544,7 +544,7 @@ static void D3DCreateDXCArguments(TArray<const WCHAR*>& OutArgs, const WCHAR* Ex
 		L"0", L"1", L"2", L"3", L"4", L"5", L"6", L"7", L"8", L"9"
 	};
 
-	if (AutoBindingSpace < ARRAY_COUNT(DigitStrings))
+	if (AutoBindingSpace < UE_ARRAY_COUNT(DigitStrings))
 	{
 		OutArgs.Add(L"/auto-binding-space");
 		OutArgs.Add(DigitStrings[AutoBindingSpace]);
diff --git a/Engine/Source/Editor/AnimGraph/Private/K2Node_TransitionRuleGetter.cpp b/Engine/Source/Editor/AnimGraph/Private/K2Node_TransitionRuleGetter.cpp
--- a/Engine/Source/Editor/AnimGraph/Private/K2Node_TransitionRuleGetter.cpp
+++ b/Engine/Source/Editor/AnimGraph/Private/K2Node_TransitionRuleGetter.cpp
@@ -172,7 +172,7 @@ void UK2Node_TransitionRuleGetter::GetStateSpecificAnimTransitionSchemaMenuActio
 	TArray<UK2Node*> AssetPlayers;
 	StateNode->BoundGraph->GetNodesOfClassEx<UAnimGraphNode_Base, UK2Node>(/*out*/ AssetPlayers);
 
-	for (int32 TypeIndex = 0; TypeIndex < ARRAY_COUNT(SequenceSpecificGetters); ++TypeIndex)
+	for (int32 TypeIndex = 0; TypeIndex < UE_ARRAY_COUNT(SequenceSpecificGetters); ++TypeIndex)
 	{
 		for (auto NodeIt = AssetPlayers.CreateConstIterator(); NodeIt; ++NodeIt)
 		{
@@ -252,7 +252,7 @@ void UK2Node_TransitionRuleGetter::GetNonStateSpecificMenuActions(FBlueprintActi
 			ETransitionGetter::CurrentState_GetBlendWeight
 		};
 
-		for (int32 TypeIndex = 0; TypeIndex < ARRAY_COUNT(NonSpecificGetters); ++TypeIndex)
+		for (int32 TypeIndex = 0; TypeIndex < UE_ARRAY_COUNT(NonSpecificGetters); ++TypeIndex)
 		{
 			auto UiSpecOverride = [](const FBlueprintActionContext& /*Context*/, const IBlueprintNodeBinder::FBindingSet& Bindings, FBlueprintActionUiSpec* UiSpecOut, TEnumAsByte<ETransitionGetter::Type> InGetterType)
 			{
diff --git a/Engine/Source/Editor/ContentBrowser/Private/AssetViewTypes.h b/Engine/Source/Editor/ContentBrowser/Private/AssetViewTypes.h
--- a/Engine/Source/Editor/ContentBrowser/Private/AssetViewTypes.h
+++ b/Engine/Source/Editor/ContentBrowser/Private/AssetViewTypes.h
@@ -80,11 +80,11 @@ struct FAssetViewAsset : public FAssetViewItem
 		TextFilterUtils::FNameBufferWithNumber NameBuffer(Data.AssetName);
 		if (NameBuffer.IsWide())
 		{
-			FCStringWide::Strncpy(FirstFewAssetNameCharacters, NameBuffer.GetWideNamePtr(), ARRAY_COUNT(FirstFewAssetNameCharacters));
+			FCStringWide::Strncpy(FirstFewAssetNameCharacters, NameBuffer.GetWideNamePtr(), UE_ARRAY_COUNT(FirstFewAssetNameCharacters));
 		}
 		else
 		{
-			int32 NumChars = FMath::Min<int32>(FCStringAnsi::Strlen(NameBuffer.GetAnsiNamePtr()), ARRAY_COUNT(FirstFewAssetNameCharacters) - 1);
+			int32 NumChars = FMath::Min<int32>(FCStringAnsi::Strlen(NameBuffer.GetAnsiNamePtr()), UE_ARRAY_COUNT(FirstFewAssetNameCharacters) - 1);
 			FPlatformString::Convert(FirstFewAssetNameCharacters, NumChars, NameBuffer.GetAnsiNamePtr(), NumChars);
 			FirstFewAssetNameCharacters[NumChars] = 0;
 		}
diff --git a/Engine/Source/Editor/ContentBrowser/Private/AssetViewWidgets.cpp b/Engine/Source/Editor/ContentBrowser/Private/AssetViewWidgets.cpp
--- a/Engine/Source/Editor/ContentBrowser/Private/AssetViewWidgets.cpp
+++ b/Engine/Source/Editor/ContentBrowser/Private/AssetViewWidgets.cpp
@@ -1370,7 +1370,7 @@ void SAssetViewItem::CacheDisplayTags()
 				if (ValueString.StartsWith(StringToRemove))
 				{
 					// Remove the class path for native classes, and also remove Engine. for engine classes
-					const int32 SizeOfPrefix = ARRAY_COUNT(StringToRemove) - 1;
+					const int32 SizeOfPrefix = UE_ARRAY_COUNT(StringToRemove) - 1;
 					ValueString = ValueString.Mid(SizeOfPrefix, ValueString.Len() - SizeOfPrefix).Replace(TEXT("Engine."), TEXT(""));
 				}
 
diff --git a/Engine/Source/Editor/ContentBrowser/Private/ContentBrowserSingleton.cpp b/Engine/Source/Editor/ContentBrowser/Private/ContentBrowserSingleton.cpp
--- a/Engine/Source/Editor/ContentBrowser/Private/ContentBrowserSingleton.cpp
+++ b/Engine/Source/Editor/ContentBrowser/Private/ContentBrowserSingleton.cpp
@@ -47,7 +47,7 @@ FContentBrowserSingleton::FContentBrowserSingleton()
 		ContentBrowserIcon,
 		true);
 
-	for ( int32 BrowserIdx = 0; BrowserIdx < ARRAY_COUNT(ContentBrowserTabIDs); BrowserIdx++ )
+	for ( int32 BrowserIdx = 0; BrowserIdx < UE_ARRAY_COUNT(ContentBrowserTabIDs); BrowserIdx++ )
 	{
 		const FName TabID = FName(*FString::Printf(TEXT("ContentBrowserTab%d"), BrowserIdx + 1));
 		ContentBrowserTabIDs[BrowserIdx] = TabID;
@@ -79,7 +79,7 @@ FContentBrowserSingleton::~FContentBrowserSingleton()
 
 	if ( FSlateApplication::IsInitialized() )
 	{
-		for ( int32 BrowserIdx = 0; BrowserIdx < ARRAY_COUNT(ContentBrowserTabIDs); BrowserIdx++ )
+		for ( int32 BrowserIdx = 0; BrowserIdx < UE_ARRAY_COUNT(ContentBrowserTabIDs); BrowserIdx++ )
 		{
 			FGlobalTabmanager::Get()->UnregisterNomadTabSpawner( ContentBrowserTabIDs[BrowserIdx] );
 		}
@@ -557,7 +557,7 @@ FName FContentBrowserSingleton::SummonNewBrowser(bool bAllowLockedBrowsers)
 	}
 	
 	FName NewTabName;
-	for ( int32 BrowserIdx = 0; BrowserIdx < ARRAY_COUNT(ContentBrowserTabIDs); BrowserIdx++ )
+	for ( int32 BrowserIdx = 0; BrowserIdx < UE_ARRAY_COUNT(ContentBrowserTabIDs); BrowserIdx++ )
 	{
 		FName TestTabID = ContentBrowserTabIDs[BrowserIdx];
 		if ( !OpenBrowserIDs.Contains(TestTabID) && (bAllowLockedBrowsers || !IsLocked(TestTabID)) )
diff --git a/Engine/Source/Editor/ContentBrowser/Private/SAssetView.cpp b/Engine/Source/Editor/ContentBrowser/Private/SAssetView.cpp
--- a/Engine/Source/Editor/ContentBrowser/Private/SAssetView.cpp
+++ b/Engine/Source/Editor/ContentBrowser/Private/SAssetView.cpp
@@ -5002,7 +5002,7 @@ FReply SAssetView::HandleQuickJumpKeyDown(const TCHAR InCharacter, const bool bI
 	}
 
 	// Check for invalid characters
-	for(int InvalidCharIndex = 0; InvalidCharIndex < ARRAY_COUNT(INVALID_OBJECTNAME_CHARACTERS) - 1; ++InvalidCharIndex)
+	for(int InvalidCharIndex = 0; InvalidCharIndex < UE_ARRAY_COUNT(INVALID_OBJECTNAME_CHARACTERS) - 1; ++InvalidCharIndex)
 	{
 		if(InCharacter == INVALID_OBJECTNAME_CHARACTERS[InvalidCharIndex])
 		{
diff --git a/Engine/Source/Editor/CurveEditor/Private/CurveEditorHelpers.cpp b/Engine/Source/Editor/CurveEditor/Private/CurveEditorHelpers.cpp
--- a/Engine/Source/Editor/CurveEditor/Private/CurveEditorHelpers.cpp
+++ b/Engine/Source/Editor/CurveEditor/Private/CurveEditorHelpers.cpp
@@ -43,7 +43,7 @@ void ConstructYGridLines(const FCurveEditorScreenSpace& ViewSpace, uint8 InMinor
 	const float Order = FMath::Pow(10.f, FMath::FloorToInt(FMath::LogX(10.f, GridPixelSpacing / ViewSpace.PixelsPerOutput())));
 
 	static const int32 DesirableBases[]  = { 2, 5 };
-	static const int32 NumDesirableBases = ARRAY_COUNT(DesirableBases);
+	static const int32 NumDesirableBases = UE_ARRAY_COUNT(DesirableBases);
 
 	const int32 Scale = FMath::RoundToInt(GridPixelSpacing / ViewSpace.PixelsPerOutput() / Order);
 	int32 Base = DesirableBases[0];
diff --git a/Engine/Source/Editor/EditorStyle/Private/SlateEditorStyle.cpp b/Engine/Source/Editor/EditorStyle/Private/SlateEditorStyle.cpp
--- a/Engine/Source/Editor/EditorStyle/Private/SlateEditorStyle.cpp
+++ b/Engine/Source/Editor/EditorStyle/Private/SlateEditorStyle.cpp
@@ -6438,7 +6438,7 @@ void FSlateEditorStyle::FStyle::SetupClassIconsAndThumbnails()
 			TEXT("FoliageType_Actor")
 		};
 
-		for (int32 TypeIndex = 0; TypeIndex < ARRAY_COUNT(AssetTypes); ++TypeIndex)
+		for (int32 TypeIndex = 0; TypeIndex < UE_ARRAY_COUNT(AssetTypes); ++TypeIndex)
 		{
 			const TCHAR* Type = AssetTypes[TypeIndex];
 			Set( *FString::Printf(TEXT("ClassIcon.%s"), Type),		new IMAGE_BRUSH(FString::Printf(TEXT("Icons/AssetIcons/%s_%dx"), Type, 16), Icon16x16 ) );
diff --git a/Engine/Source/Editor/GameProjectGeneration/Private/GameProjectUtils.cpp b/Engine/Source/Editor/GameProjectGeneration/Private/GameProjectUtils.cpp
--- a/Engine/Source/Editor/GameProjectGeneration/Private/GameProjectUtils.cpp
+++ b/Engine/Source/Editor/GameProjectGeneration/Private/GameProjectUtils.cpp
@@ -2702,7 +2702,7 @@ GameProjectUtils::EProjectDuplicateResult GameProjectUtils::DuplicateProjectForU
 				TEXT("Saved/StagedBuilds"),
 			};
 
-			SourceDirectoriesToSkip.Reserve(ARRAY_COUNT(RelativeDirectoriesToSkip));
+			SourceDirectoriesToSkip.Reserve(UE_ARRAY_COUNT(RelativeDirectoriesToSkip));
 			for (const FString& RelativeDirectoryToSkip : RelativeDirectoriesToSkip)
 			{
 				SourceDirectoriesToSkip.Emplace(RootSourceDirectory / RelativeDirectoryToSkip);
diff --git a/Engine/Source/Editor/Kismet/Private/FindInBlueprintManager.cpp b/Engine/Source/Editor/Kismet/Private/FindInBlueprintManager.cpp
--- a/Engine/Source/Editor/Kismet/Private/FindInBlueprintManager.cpp
+++ b/Engine/Source/Editor/Kismet/Private/FindInBlueprintManager.cpp
@@ -1439,7 +1439,7 @@ FFindInBlueprintSearchManager::FFindInBlueprintSearchManager()
 	, AssetRegistryModule(nullptr)
 	, CachingObject(nullptr)
 {
-	for (int32 TabIdx = 0; TabIdx < ARRAY_COUNT(GlobalFindResultsTabIDs); TabIdx++)
+	for (int32 TabIdx = 0; TabIdx < UE_ARRAY_COUNT(GlobalFindResultsTabIDs); TabIdx++)
 	{
 		const FName TabID = FName(*FString::Printf(TEXT("GlobalFindResults_%02d"), TabIdx + 1));
 		GlobalFindResultsTabIDs[TabIdx] = TabID;
@@ -2590,7 +2590,7 @@ TSharedPtr<SFindInBlueprints> FFindInBlueprintSearchManager::OpenGlobalFindResul
 		}
 	}
 
-	for (int32 Idx = 0; Idx < ARRAY_COUNT(GlobalFindResultsTabIDs); ++Idx)
+	for (int32 Idx = 0; Idx < UE_ARRAY_COUNT(GlobalFindResultsTabIDs); ++Idx)
 	{
 		const FName GlobalTabId = GlobalFindResultsTabIDs[Idx];
 		if (!OpenGlobalTabIDs.Contains(GlobalTabId))
@@ -2643,7 +2643,7 @@ void FFindInBlueprintSearchManager::EnableGlobalFindResults(bool bEnable)
 			GlobalFindResultsIcon,
 			true);
 
-		for (int32 TabIdx = 0; TabIdx < ARRAY_COUNT(GlobalFindResultsTabIDs); TabIdx++)
+		for (int32 TabIdx = 0; TabIdx < UE_ARRAY_COUNT(GlobalFindResultsTabIDs); TabIdx++)
 		{
 			const FName TabID = GlobalFindResultsTabIDs[TabIdx];
 			if (!GlobalTabManager->HasTabSpawner(TabID))
@@ -2678,7 +2678,7 @@ void FFindInBlueprintSearchManager::EnableGlobalFindResults(bool bEnable)
 
 		GlobalFindResults.Empty();
 
-		for (int32 TabIdx = 0; TabIdx < ARRAY_COUNT(GlobalFindResultsTabIDs); TabIdx++)
+		for (int32 TabIdx = 0; TabIdx < UE_ARRAY_COUNT(GlobalFindResultsTabIDs); TabIdx++)
 		{
 			const FName TabID = GlobalFindResultsTabIDs[TabIdx];
 			if (GlobalTabManager->HasTabSpawner(TabID))
@@ -2699,7 +2699,7 @@ void FFindInBlueprintSearchManager::CloseOrphanedGlobalFindResultsTabs(TSharedPt
 {
 	if (TabManager.IsValid())
 	{
-		for (int32 TabIdx = 0; TabIdx < ARRAY_COUNT(GlobalFindResultsTabIDs); TabIdx++)
+		for (int32 TabIdx = 0; TabIdx < UE_ARRAY_COUNT(GlobalFindResultsTabIDs); TabIdx++)
 		{
 			const FName TabID = GlobalFindResultsTabIDs[TabIdx];
 			if (!FGlobalTabmanager::Get()->HasTabSpawner(TabID))
diff --git a/Engine/Source/Editor/LandscapeEditor/Private/LandscapeEditorDetailCustomization_NewLandscape.cpp b/Engine/Source/Editor/LandscapeEditor/Private/LandscapeEditorDetailCustomization_NewLandscape.cpp
--- a/Engine/Source/Editor/LandscapeEditor/Private/LandscapeEditorDetailCustomization_NewLandscape.cpp
+++ b/Engine/Source/Editor/LandscapeEditor/Private/LandscapeEditorDetailCustomization_NewLandscape.cpp
@@ -585,7 +585,7 @@ TSharedRef<SWidget> FLandscapeEditorDetailCustomization_NewLandscape::GetSection
 {
 	FMenuBuilder MenuBuilder(true, nullptr);
 
-	for (int32 i = 0; i < ARRAY_COUNT(FNewLandscapeUtils::SectionSizes); i++)
+	for (int32 i = 0; i < UE_ARRAY_COUNT(FNewLandscapeUtils::SectionSizes); i++)
 	{
 		MenuBuilder.AddMenuEntry(FText::Format(LOCTEXT("NxNQuads", "{0}\u00D7{0} Quads"), FText::AsNumber(FNewLandscapeUtils::SectionSizes[i])), FText::GetEmpty(),
 			FSlateIcon(), FExecuteAction::CreateStatic(&OnChangeSectionSize, PropertyHandle, FNewLandscapeUtils::SectionSizes[i]));
@@ -617,7 +617,7 @@ TSharedRef<SWidget> FLandscapeEditorDetailCustomization_NewLandscape::GetSection
 {
 	FMenuBuilder MenuBuilder(true, nullptr);
 
-	for (int32 i = 0; i < ARRAY_COUNT(FNewLandscapeUtils::NumSections); i++)
+	for (int32 i = 0; i < UE_ARRAY_COUNT(FNewLandscapeUtils::NumSections); i++)
 	{
 		FFormatNamedArguments Args;
 		Args.Add(TEXT("Width"), FNewLandscapeUtils::NumSections[i]);
diff --git a/Engine/Source/Editor/LandscapeEditor/Private/LandscapeEditorDetailCustomization_ResizeLandscape.cpp b/Engine/Source/Editor/LandscapeEditor/Private/LandscapeEditorDetailCustomization_ResizeLandscape.cpp
--- a/Engine/Source/Editor/LandscapeEditor/Private/LandscapeEditorDetailCustomization_ResizeLandscape.cpp
+++ b/Engine/Source/Editor/LandscapeEditor/Private/LandscapeEditorDetailCustomization_ResizeLandscape.cpp
@@ -305,7 +305,7 @@ TSharedRef<SWidget> FLandscapeEditorDetailCustomization_ResizeLandscape::GetSect
 {
 	FMenuBuilder MenuBuilder(true, NULL);
 
-	for (int32 i = 0; i < ARRAY_COUNT(SectionSizes); i++)
+	for (int32 i = 0; i < UE_ARRAY_COUNT(SectionSizes); i++)
 	{
 		MenuBuilder.AddMenuEntry(FText::Format(LOCTEXT("NxNQuads", "{0}x{0} Quads"), FText::AsNumber(SectionSizes[i])), FText::GetEmpty(), FSlateIcon(), FExecuteAction::CreateStatic(&OnChangeSectionSize, PropertyHandle, SectionSizes[i]));
 	}
@@ -372,7 +372,7 @@ TSharedRef<SWidget> FLandscapeEditorDetailCustomization_ResizeLandscape::GetSect
 {
 	FMenuBuilder MenuBuilder(true, NULL);
 
-	for (int32 i = 0; i < ARRAY_COUNT(NumSections); i++)
+	for (int32 i = 0; i < UE_ARRAY_COUNT(NumSections); i++)
 	{
 		FFormatNamedArguments Args;
 		Args.Add(TEXT("Width"), NumSections[i]);
diff --git a/Engine/Source/Editor/LandscapeEditor/Private/NewLandscapeUtils.cpp b/Engine/Source/Editor/LandscapeEditor/Private/NewLandscapeUtils.cpp
--- a/Engine/Source/Editor/LandscapeEditor/Private/NewLandscapeUtils.cpp
+++ b/Engine/Source/Editor/LandscapeEditor/Private/NewLandscapeUtils.cpp
@@ -28,9 +28,9 @@ void FNewLandscapeUtils::ChooseBestComponentSizeForImport(ULandscapeEditorObject
 	if(Width > 0 && Height > 0)
 	{
 		// Try to find a section size and number of sections that exactly matches the dimensions of the heightfield
-		for(int32 SectionSizesIdx = ARRAY_COUNT(SectionSizes) - 1; SectionSizesIdx >= 0; SectionSizesIdx--)
+		for(int32 SectionSizesIdx = UE_ARRAY_COUNT(SectionSizes) - 1; SectionSizesIdx >= 0; SectionSizesIdx--)
 		{
-			for(int32 NumSectionsIdx = ARRAY_COUNT(NumSections) - 1; NumSectionsIdx >= 0; NumSectionsIdx--)
+			for(int32 NumSectionsIdx = UE_ARRAY_COUNT(NumSections) - 1; NumSectionsIdx >= 0; NumSectionsIdx--)
 			{
 				int32 ss = SectionSizes[SectionSizesIdx];
 				int32 ns = NumSections[NumSectionsIdx];
@@ -58,7 +58,7 @@ void FNewLandscapeUtils::ChooseBestComponentSizeForImport(ULandscapeEditorObject
 			// if there was no exact match, try increasing the section size until we encompass the whole heightmap
 			const int32 CurrentSectionSize = UISettings->NewLandscape_QuadsPerSection;
 			const int32 CurrentNumSections = UISettings->NewLandscape_SectionsPerComponent;
-			for(int32 SectionSizesIdx = 0; SectionSizesIdx < ARRAY_COUNT(SectionSizes); SectionSizesIdx++)
+			for(int32 SectionSizesIdx = 0; SectionSizesIdx < UE_ARRAY_COUNT(SectionSizes); SectionSizesIdx++)
 			{
 				if(SectionSizes[SectionSizesIdx] < CurrentSectionSize)
 				{
@@ -83,8 +83,8 @@ void FNewLandscapeUtils::ChooseBestComponentSizeForImport(ULandscapeEditorObject
 		if(!bFoundMatch)
 		{
 			// if the heightmap is very large, fall back to using the largest values we support
-			const int32 MaxSectionSize = SectionSizes[ARRAY_COUNT(SectionSizes) - 1];
-			const int32 MaxNumSubSections = NumSections[ARRAY_COUNT(NumSections) - 1];
+			const int32 MaxSectionSize = SectionSizes[UE_ARRAY_COUNT(SectionSizes) - 1];
+			const int32 MaxNumSubSections = NumSections[UE_ARRAY_COUNT(NumSections) - 1];
 			const int32 ComponentsX = FMath::DivideAndRoundUp((Width - 1), MaxSectionSize * MaxNumSubSections);
 			const int32 ComponentsY = FMath::DivideAndRoundUp((Height - 1), MaxSectionSize * MaxNumSubSections);
 
diff --git a/Engine/Source/Editor/LevelEditor/Private/SSurfaceProperties.cpp b/Engine/Source/Editor/LevelEditor/Private/SSurfaceProperties.cpp
--- a/Engine/Source/Editor/LevelEditor/Private/SSurfaceProperties.cpp
+++ b/Engine/Source/Editor/LevelEditor/Private/SSurfaceProperties.cpp
@@ -66,7 +66,7 @@ void SSurfaceProperties::Construct( const FArguments& InArgs )
 	GConfig->GetBool(TEXT("SelectionDetails"), TEXT("UseRelativeScaling"), bUseRelativeScaling, GEditorPerProjectIni);
 
 	static const float ScalingValues[] = { 1.0f / 16, 1.0f / 8, 1.0f / 4, 1.0f / 2, 1, 2, 4, 8, 16 };
-	for(int Idx = 0; Idx < ARRAY_COUNT(ScalingValues); Idx++)
+	for(int Idx = 0; Idx < UE_ARRAY_COUNT(ScalingValues); Idx++)
 	{
 		ScalingFactors.Add(MakeShareable(new FString(FString::Printf(TEXT("%f"),ScalingValues[Idx]))));
 	}
@@ -213,7 +213,7 @@ TSharedRef<SWidget> SSurfaceProperties::ConstructTexturePan()
 
 	const static FText ButtonFields[] = {FText::FromString("1/256"), FText::FromString("1/64"), FText::FromString("1/16"), FText::FromString("1/4")};
 	const static int32 ButtonIncriments[] = {1,4,16,64};
-	for (int32 i = 0; i < ARRAY_COUNT(ButtonFields); i++)
+	for (int32 i = 0; i < UE_ARRAY_COUNT(ButtonFields); i++)
 	{
 		GridBox->AddSlot(i, 0)
 		.VAlign(VAlign_Fill)
@@ -245,7 +245,7 @@ TSharedRef<SWidget> SSurfaceProperties::ConstructTexturePan()
 	{
 		TSharedPtr<SComboButton> ComboButton;
 		TSharedPtr<SWidget> NumberBox;
-		GridBox->AddSlot(ARRAY_COUNT(ButtonFields),i)
+		GridBox->AddSlot(UE_ARRAY_COUNT(ButtonFields),i)
 		.VAlign(VAlign_Fill)
 		.HAlign(HAlign_Fill)
 		[
@@ -328,7 +328,7 @@ TSharedRef<SWidget> SSurfaceProperties::ConstructTextureRotate()
 	const static RotationAction RotationActions[] = {Rotate,Rotate,RotateCustom};
 	
 	// create rotation controls
-	for (int Idx = 0; Idx < ARRAY_COUNT(RotationValues); Idx++)
+	for (int Idx = 0; Idx < UE_ARRAY_COUNT(RotationValues); Idx++)
 	{
 		if (RotationActions[Idx] == RotateCustom)
 		{
@@ -412,7 +412,7 @@ TSharedRef<SWidget> SSurfaceProperties::ConstructTextureFlip()
 	const static TextureCoordChannel TextureCoordinateChannels[] = {UChannel,VChannel};
 
 	// create flip controls
-	for (int Idx = 0; Idx < ARRAY_COUNT(ButtonFields); Idx++)
+	for (int Idx = 0; Idx < UE_ARRAY_COUNT(ButtonFields); Idx++)
 	{
 		FlipBox->AddSlot()
 		.FillWidth(1)
@@ -468,7 +468,7 @@ TSharedRef<SWidget> SSurfaceProperties::ConstructTextureScale()
 	FString ControlLabels[] = {"U", "V"};
 	TextureCoordChannel Channels[] = {UChannel,VChannel};
 	
-	for(int idx = 0; idx < ARRAY_COUNT(ControlLabels); idx++)
+	for(int idx = 0; idx < UE_ARRAY_COUNT(ControlLabels); idx++)
 	{
 		Controls->AddSlot()
 		.HAlign(HAlign_Fill)
diff --git a/Engine/Source/Editor/Matinee/Private/Matinee.cpp b/Engine/Source/Editor/Matinee/Private/Matinee.cpp
--- a/Engine/Source/Editor/Matinee/Private/Matinee.cpp
+++ b/Engine/Source/Editor/Matinee/Private/Matinee.cpp
@@ -2958,7 +2958,7 @@ void FMatinee::RepairHierarchyProblems()
  */
 FString FMatinee::GetInterpEdFPSSnapSizeLocName( int32 StringIndex )
 {
-	check( StringIndex >= 0 && StringIndex < ARRAY_COUNT( InterpEdFPSSnapSizes ) );
+	check( StringIndex >= 0 && StringIndex < UE_ARRAY_COUNT( InterpEdFPSSnapSizes ) );
 
 	switch( StringIndex )
 	{
diff --git a/Engine/Source/Editor/Matinee/Private/MatineeDraw.cpp b/Engine/Source/Editor/Matinee/Private/MatineeDraw.cpp
--- a/Engine/Source/Editor/Matinee/Private/MatineeDraw.cpp
+++ b/Engine/Source/Editor/Matinee/Private/MatineeDraw.cpp
@@ -676,7 +676,7 @@ void FMatineeViewportClient::DrawTimeline(FViewport* InViewport, FCanvas* Canvas
 	int32 SelIndex = InterpEd->SnapSelectionIndex;
 
 	// Determine if time should be drawn including frames or keys
-	if(SelIndex == ARRAY_COUNT(FMatinee::InterpEdSnapSizes)+ARRAY_COUNT(FMatinee::InterpEdFPSSnapSizes))
+	if(SelIndex == UE_ARRAY_COUNT(FMatinee::InterpEdSnapSizes)+UE_ARRAY_COUNT(FMatinee::InterpEdFPSSnapSizes))
 	{
 		UInterpTrack* Track = NULL;
 		int32 SelKeyIndex = 0;
@@ -704,7 +704,7 @@ void FMatineeViewportClient::DrawTimeline(FViewport* InViewport, FCanvas* Canvas
 
 		StringSize( Font, XL, YL, *SnapPosString );
 	}
-	else if(SelIndex < ARRAY_COUNT(FMatinee::InterpEdFPSSnapSizes)+ARRAY_COUNT(FMatinee::InterpEdSnapSizes) && SelIndex >= ARRAY_COUNT(FMatinee::InterpEdSnapSizes))
+	else if(SelIndex < UE_ARRAY_COUNT(FMatinee::InterpEdFPSSnapSizes)+UE_ARRAY_COUNT(FMatinee::InterpEdSnapSizes) && SelIndex >= UE_ARRAY_COUNT(FMatinee::InterpEdSnapSizes))
 	{
 		// frames
 
@@ -719,7 +719,7 @@ void FMatineeViewportClient::DrawTimeline(FViewport* InViewport, FCanvas* Canvas
 		SnapPosString = FString::Printf( TEXT("%3.1f / %3.1f %s"), (1.0 / InterpEd->SnapAmount) * InterpEd->MatineeActor->InterpPosition, (1.0 / InterpEd->SnapAmount) * InterpEd->IData->InterpLength, *NSLOCTEXT("UnrealEd", "InterpEd_TimelineInfo_Frames", "frames").ToString() );
 		StringSize( Font, XL, YL, *SnapPosString );
 	}
-	else if(SelIndex < ARRAY_COUNT(FMatinee::InterpEdSnapSizes))
+	else if(SelIndex < UE_ARRAY_COUNT(FMatinee::InterpEdSnapSizes))
 	{
 		// seconds
 		SnapPosString = "";
diff --git a/Engine/Source/Editor/Matinee/Private/MatineeMenus.cpp b/Engine/Source/Editor/Matinee/Private/MatineeMenus.cpp
--- a/Engine/Source/Editor/Matinee/Private/MatineeMenus.cpp
+++ b/Engine/Source/Editor/Matinee/Private/MatineeMenus.cpp
@@ -4981,25 +4981,25 @@ void FMatinee::OnChangeSnapSize( TSharedPtr< FString > SelectedString, ESelectIn
 	bool bFound = SnapComboStrings.Find(SelectedString, NewSelection);
 	check(bFound);
 
-	if(NewSelection == ARRAY_COUNT(InterpEdSnapSizes)+ARRAY_COUNT(InterpEdFPSSnapSizes))
+	if(NewSelection == UE_ARRAY_COUNT(InterpEdSnapSizes)+UE_ARRAY_COUNT(InterpEdFPSSnapSizes))
 	{
 		bSnapToFrames = false;
 		bSnapToKeys = true;
 		SnapAmount = 1.0f / 30.0f;	// Shouldn't be used
 		CurveEd->SetInSnap(false, SnapAmount, bSnapToFrames);
 	}
-	else if(NewSelection<ARRAY_COUNT(InterpEdSnapSizes))	// see if they picked a second snap amount
+	else if(NewSelection<UE_ARRAY_COUNT(InterpEdSnapSizes))	// see if they picked a second snap amount
 	{
 		bSnapToFrames = false;
 		bSnapToKeys = false;
 		SnapAmount = InterpEdSnapSizes[NewSelection];
 		CurveEd->SetInSnap(bSnapEnabled, SnapAmount, bSnapToFrames);
 	}
-	else if(NewSelection<ARRAY_COUNT(InterpEdFPSSnapSizes)+ARRAY_COUNT(InterpEdSnapSizes))	// See if they picked a FPS snap amount.
+	else if(NewSelection<UE_ARRAY_COUNT(InterpEdFPSSnapSizes)+UE_ARRAY_COUNT(InterpEdSnapSizes))	// See if they picked a FPS snap amount.
 	{
 		bSnapToFrames = true;
 		bSnapToKeys = false;
-		SnapAmount = InterpEdFPSSnapSizes[NewSelection-ARRAY_COUNT(InterpEdSnapSizes)];
+		SnapAmount = InterpEdFPSSnapSizes[NewSelection-UE_ARRAY_COUNT(InterpEdSnapSizes)];
 		CurveEd->SetInSnap(bSnapEnabled, SnapAmount, bSnapToFrames);
 	}
 
diff --git a/Engine/Source/Editor/Matinee/Private/MatineeToolbars.cpp b/Engine/Source/Editor/Matinee/Private/MatineeToolbars.cpp
--- a/Engine/Source/Editor/Matinee/Private/MatineeToolbars.cpp
+++ b/Engine/Source/Editor/Matinee/Private/MatineeToolbars.cpp
@@ -21,14 +21,14 @@ FString FMatinee::GetToolbarSnapText()
 {
 	FString Text;
 	int32 StringIndex = 0;
-	if (SnapSelectionIndex < ARRAY_COUNT(InterpEdSnapSizes))
+	if (SnapSelectionIndex < UE_ARRAY_COUNT(InterpEdSnapSizes))
 	{
 		StringIndex = SnapSelectionIndex;
 		Text = FString::Printf(TEXT("%.2f"),FMatinee::InterpEdSnapSizes[StringIndex]);
 	}
-	else if (SnapSelectionIndex < (ARRAY_COUNT(InterpEdSnapSizes)+ARRAY_COUNT(InterpEdFPSSnapSizes)))
+	else if (SnapSelectionIndex < (UE_ARRAY_COUNT(InterpEdSnapSizes)+UE_ARRAY_COUNT(InterpEdFPSSnapSizes)))
 	{
-		StringIndex = SnapSelectionIndex - ARRAY_COUNT(InterpEdSnapSizes);
+		StringIndex = SnapSelectionIndex - UE_ARRAY_COUNT(InterpEdSnapSizes);
 		Text = GetInterpEdFPSSnapSizeLocName( StringIndex );
 	}
 	else
@@ -211,13 +211,13 @@ void FMatinee::ExtendToolbar()
 
 	// Append Second Snap Times
 	SnapComboStrings.Empty();
-	for(int32 i=0; i<ARRAY_COUNT(InterpEdSnapSizes); i++)
+	for(int32 i=0; i<UE_ARRAY_COUNT(InterpEdSnapSizes); i++)
 	{
 		FString SnapCaption = FString::Printf( TEXT("%1.2f"), InterpEdSnapSizes[i] );
 		SnapComboStrings.Add( MakeShareable( new FString(SnapCaption) ) );
 	}
 	// Append FPS Snap Times
-	for(int32 i=0; i<ARRAY_COUNT(InterpEdFPSSnapSizes); i++)
+	for(int32 i=0; i<UE_ARRAY_COUNT(InterpEdFPSSnapSizes); i++)
 	{
 		FString SnapCaption = GetInterpEdFPSSnapSizeLocName( i );
 		SnapComboStrings.Add( MakeShareable( new FString(SnapCaption) ) );
diff --git a/Engine/Source/Editor/MovieSceneTools/Private/TrackEditors/PropertyTrackEditors/EulerTransformPropertyTrackEditor.cpp b/Engine/Source/Editor/MovieSceneTools/Private/TrackEditors/PropertyTrackEditors/EulerTransformPropertyTrackEditor.cpp
--- a/Engine/Source/Editor/MovieSceneTools/Private/TrackEditors/PropertyTrackEditors/EulerTransformPropertyTrackEditor.cpp
+++ b/Engine/Source/Editor/MovieSceneTools/Private/TrackEditors/PropertyTrackEditors/EulerTransformPropertyTrackEditor.cpp
@@ -80,7 +80,7 @@ void FEulerTransformPropertyTrackEditor::GenerateKeysFromPropertyChanged( const
 
 	if (LeafPath.Len() > 0)
 	{
-		for (int32 ChannelIndex = 0; ChannelIndex < ARRAY_COUNT(ChannelNames); ++ChannelIndex)
+		for (int32 ChannelIndex = 0; ChannelIndex < UE_ARRAY_COUNT(ChannelNames); ++ChannelIndex)
 		{
 			// If it doesn't match the fully qualified path, and doesn't start with the leaf path, don't add a key
 			bool bMatchesQualifiedPath = FPlatformString::Stricmp(*QualifiedLeafPath, ChannelNames[ChannelIndex]) == 0;
diff --git a/Engine/Source/Editor/PixelInspector/Private/PixelInspector.cpp b/Engine/Source/Editor/PixelInspector/Private/PixelInspector.cpp
--- a/Engine/Source/Editor/PixelInspector/Private/PixelInspector.cpp
+++ b/Engine/Source/Editor/PixelInspector/Private/PixelInspector.cpp
@@ -660,7 +660,7 @@ namespace PixelInspector
 
 	void SPixelInspector::ReadBackRequestData()
 	{
-		for (int RequestIndex = 0; RequestIndex < ARRAY_COUNT(Requests); ++RequestIndex)
+		for (int RequestIndex = 0; RequestIndex < UE_ARRAY_COUNT(Requests); ++RequestIndex)
 		{
 			FPixelInspectorRequest& Request = Requests[RequestIndex];
 			if (Request.RequestComplete == false && Request.RenderingCommandSend == true)
diff --git a/Engine/Source/Editor/SequenceRecorder/Private/SequenceRecorderModule.cpp b/Engine/Source/Editor/SequenceRecorder/Private/SequenceRecorderModule.cpp
--- a/Engine/Source/Editor/SequenceRecorder/Private/SequenceRecorderModule.cpp
+++ b/Engine/Source/Editor/SequenceRecorder/Private/SequenceRecorderModule.cpp
@@ -292,7 +292,7 @@ class FSequenceRecorderModule : public ISequenceRecorder, private FSelfRegisteri
 		// parse actor name
 		TCHAR ActorName[128];
 		AActor* FoundActor = nullptr;
-		if (FParse::Token(Str, ActorName, ARRAY_COUNT(ActorName), 0))
+		if (FParse::Token(Str, ActorName, UE_ARRAY_COUNT(ActorName), 0))
 		{
 			FoundActor = FindActorByName(FString(ActorName), InWorld);
 		}
@@ -303,7 +303,7 @@ class FSequenceRecorderModule : public ISequenceRecorder, private FSelfRegisteri
 			if (SkelComp)
 			{
 				TCHAR AssetPath[256];
-				FParse::Token(Str, AssetPath, ARRAY_COUNT(AssetPath), 0);
+				FParse::Token(Str, AssetPath, UE_ARRAY_COUNT(AssetPath), 0);
 				FString const AssetName = FPackageName::GetLongPackageAssetName(AssetPath);
 				return FAnimationRecorderManager::Get().RecordAnimation(SkelComp, AssetPath, AssetName, GetDefault<USequenceRecorderSettings>()->DefaultAnimationSettings);
 			}
@@ -366,7 +366,7 @@ class FSequenceRecorderModule : public ISequenceRecorder, private FSelfRegisteri
 		TCHAR ActorName[128];
 		AActor* FoundActor = nullptr;
 		bool bStopAll = false;
-		if (FParse::Token(Str, ActorName, ARRAY_COUNT(ActorName), 0))
+		if (FParse::Token(Str, ActorName, UE_ARRAY_COUNT(ActorName), 0))
 		{
 			FString const ActorNameStr = FString(ActorName);
 
@@ -432,7 +432,7 @@ class FSequenceRecorderModule : public ISequenceRecorder, private FSelfRegisteri
 		const TCHAR* Str = InStr;
 		EFilterType FilterType = EFilterType::None;
 		TCHAR Filter[128];
-		if(FParse::Token(Str, Filter, ARRAY_COUNT(Filter), 0))
+		if(FParse::Token(Str, Filter, UE_ARRAY_COUNT(Filter), 0))
 		{
 			FString const FilterStr = Filter;
 			if (FilterStr == TEXT("all"))
@@ -456,7 +456,7 @@ class FSequenceRecorderModule : public ISequenceRecorder, private FSelfRegisteri
 		if(FilterType == EFilterType::Actor || FilterType == EFilterType::Class)
 		{
 			TCHAR Specifier[128];
-			if(FParse::Token(Str, Specifier, ARRAY_COUNT(Specifier), 0))
+			if(FParse::Token(Str, Specifier, UE_ARRAY_COUNT(Specifier), 0))
 			{
 				FString const SpecifierStr = FString(Specifier).TrimStart();
 				if(FilterType == EFilterType::Actor)
diff --git a/Engine/Source/Editor/Sequencer/Private/DisplayNodes/SequencerDisplayNode.cpp b/Engine/Source/Editor/Sequencer/Private/DisplayNodes/SequencerDisplayNode.cpp
--- a/Engine/Source/Editor/Sequencer/Private/DisplayNodes/SequencerDisplayNode.cpp
+++ b/Engine/Source/Editor/Sequencer/Private/DisplayNodes/SequencerDisplayNode.cpp
@@ -62,8 +62,8 @@ namespace SequencerNodeConstants
 		5, // Anything else
 	};
 
-	static_assert(ARRAY_COUNT(DefaultSortBias) == (__underlying_type(EDisplayNodeSortType))EDisplayNodeSortType::NUM, "Mismatched type/bias count");
-	static_assert(ARRAY_COUNT(ObjectBindingSortBias) == (__underlying_type(EDisplayNodeSortType))EDisplayNodeSortType::NUM, "Mismatched type/bias count");
+	static_assert(UE_ARRAY_COUNT(DefaultSortBias) == (__underlying_type(EDisplayNodeSortType))EDisplayNodeSortType::NUM, "Mismatched type/bias count");
+	static_assert(UE_ARRAY_COUNT(ObjectBindingSortBias) == (__underlying_type(EDisplayNodeSortType))EDisplayNodeSortType::NUM, "Mismatched type/bias count");
 
 	inline bool SortChildrenWithBias(const TSharedRef<FSequencerDisplayNode>& A, const TSharedRef<FSequencerDisplayNode>& B, const uint8* SortBias)
 	{
diff --git a/Engine/Source/Editor/UnrealEd/Private/Commandlets/AssetRegistryGenerator.cpp b/Engine/Source/Editor/UnrealEd/Private/Commandlets/AssetRegistryGenerator.cpp
--- a/Engine/Source/Editor/UnrealEd/Private/Commandlets/AssetRegistryGenerator.cpp
+++ b/Engine/Source/Editor/UnrealEd/Private/Commandlets/AssetRegistryGenerator.cpp
@@ -196,13 +196,13 @@ public:
 			return true;
 
 		int32 ExtIndex = 0;
-		for (; ExtIndex < ARRAY_COUNT(Extensions); ++ExtIndex)
+		for (; ExtIndex < UE_ARRAY_COUNT(Extensions); ++ExtIndex)
 		{
 			if (0 == FCString::Stricmp(Extension, Extensions[ExtIndex]))
 				break;
 		}
 
-		if (ExtIndex >= ARRAY_COUNT(Extensions))
+		if (ExtIndex >= UE_ARRAY_COUNT(Extensions))
 			return true;
 
 		int32 LengthWithoutExtension = Extension - FilenameOrDirectory;
diff --git a/Engine/Source/Editor/UnrealEd/Private/Commandlets/AudioTestCommandlet.cpp b/Engine/Source/Editor/UnrealEd/Private/Commandlets/AudioTestCommandlet.cpp
--- a/Engine/Source/Editor/UnrealEd/Private/Commandlets/AudioTestCommandlet.cpp
+++ b/Engine/Source/Editor/UnrealEd/Private/Commandlets/AudioTestCommandlet.cpp
@@ -172,7 +172,7 @@ static AudioTestInfo AudioTestInfoList[] =
 	{ "system", "voice_manager",	"SourcePath or Directory",	1, TestAudioSystemVoiceManager		}, // AUDIO_TEST_SYSTEM_VOICE_MANAGER
 	{ "system", "soundfile_manager", "None",					1, TestAudioSystemSoundFileManager  }, // AUDIO_TEST_SYSTEM_SOUND_MANAGER
 };
-static_assert(ARRAY_COUNT(AudioTestInfoList) == AUDIO_TESTS, "Mismatched info list and test enumeration");
+static_assert(UE_ARRAY_COUNT(AudioTestInfoList) == AUDIO_TESTS, "Mismatched info list and test enumeration");
 
 void UAudioTestCommandlet::PrintUsage()
 {
diff --git a/Engine/Source/Editor/UnrealEd/Private/CookOnTheFlyServer.cpp b/Engine/Source/Editor/UnrealEd/Private/CookOnTheFlyServer.cpp
--- a/Engine/Source/Editor/UnrealEd/Private/CookOnTheFlyServer.cpp
+++ b/Engine/Source/Editor/UnrealEd/Private/CookOnTheFlyServer.cpp
@@ -276,7 +276,7 @@ void OutputHierarchyTimers(const FHierarchicalTimerInfo* TimerInfo, int32 Depth)
 	FString TimerName(TimerInfo->Name);
 
 	static const TCHAR LeftPad[] = TEXT("                                ");
-	const SIZE_T PadOffset = FMath::Max<int>(ARRAY_COUNT(LeftPad) - 1 - Depth * 2, 0);
+	const SIZE_T PadOffset = FMath::Max<int>(UE_ARRAY_COUNT(LeftPad) - 1 - Depth * 2, 0);
 
 	UE_LOG(LogCook, Display, TEXT("  %s%s: %.3fs (%u)"), &LeftPad[PadOffset], *TimerName, TimerInfo->Length, TimerInfo->HitCount);
 
diff --git a/Engine/Source/Editor/UnrealEd/Private/DebugToolExec.cpp b/Engine/Source/Editor/UnrealEd/Private/DebugToolExec.cpp
--- a/Engine/Source/Editor/UnrealEd/Private/DebugToolExec.cpp
+++ b/Engine/Source/Editor/UnrealEd/Private/DebugToolExec.cpp
@@ -102,7 +102,7 @@ bool FDebugToolExec::Exec( UWorld* InWorld, const TCHAR* Cmd, FOutputDevice& Ar
 			if( ParseObject<UClass>( Cmd, TEXT("CLASS="), Class, ANY_PACKAGE ) == false )
 			{
 				TCHAR ClassName[256];
-				if ( FParse::Token(Cmd,ClassName,ARRAY_COUNT(ClassName), 1) )
+				if ( FParse::Token(Cmd,ClassName,UE_ARRAY_COUNT(ClassName), 1) )
 				{
 					Class = FindObject<UClass>( ANY_PACKAGE, ClassName);
 				}
diff --git a/Engine/Source/Editor/UnrealEd/Private/EditorDirectories.cpp b/Engine/Source/Editor/UnrealEd/Private/EditorDirectories.cpp
--- a/Engine/Source/Editor/UnrealEd/Private/EditorDirectories.cpp
+++ b/Engine/Source/Editor/UnrealEd/Private/EditorDirectories.cpp
@@ -14,7 +14,7 @@ FEditorDirectories& FEditorDirectories::Get()
 
 void FEditorDirectories::LoadLastDirectories()
 {
-	for( int32 CurDirectoryIndex = 0; CurDirectoryIndex < ARRAY_COUNT( LastDir ); ++CurDirectoryIndex )
+	for( int32 CurDirectoryIndex = 0; CurDirectoryIndex < UE_ARRAY_COUNT( LastDir ); ++CurDirectoryIndex )
 	{
 		LastDir[CurDirectoryIndex].Reset();
 	}
@@ -35,7 +35,7 @@ void FEditorDirectories::LoadLastDirectories()
 
 	// Set up some defaults if they're note defined in the ini
 	const FString DefaultDir = FPaths::ProjectContentDir();
-	for( int32 CurDirectoryIndex = 0; CurDirectoryIndex < ARRAY_COUNT( LastDir ); ++CurDirectoryIndex )
+	for( int32 CurDirectoryIndex = 0; CurDirectoryIndex < UE_ARRAY_COUNT( LastDir ); ++CurDirectoryIndex )
 	{
 		if (LastDir[ CurDirectoryIndex ].IsEmpty())
 		{
@@ -81,7 +81,7 @@ void FEditorDirectories::SaveLastDirectories()
 
 FString FEditorDirectories::GetLastDirectory( const ELastDirectory::Type InLastDir ) const
 {
-	if ( InLastDir >= 0 && InLastDir < ARRAY_COUNT( LastDir ) )
+	if ( InLastDir >= 0 && InLastDir < UE_ARRAY_COUNT( LastDir ) )
 	{
 		return LastDir[InLastDir];
 	}
@@ -90,7 +90,7 @@ FString FEditorDirectories::GetLastDirectory( const ELastDirectory::Type InLastD
 
 void FEditorDirectories::SetLastDirectory( const ELastDirectory::Type InLastDir, const FString& InLastStr )
 {
-	if ( InLastDir >= 0 && InLastDir < ARRAY_COUNT( LastDir ) )
+	if ( InLastDir >= 0 && InLastDir < UE_ARRAY_COUNT( LastDir ) )
 	{
 		LastDir[InLastDir] = InLastStr;
 	}
diff --git a/Engine/Source/Editor/UnrealEd/Private/EditorEngine.cpp b/Engine/Source/Editor/UnrealEd/Private/EditorEngine.cpp
--- a/Engine/Source/Editor/UnrealEd/Private/EditorEngine.cpp
+++ b/Engine/Source/Editor/UnrealEd/Private/EditorEngine.cpp
@@ -1009,7 +1009,7 @@ void UEditorEngine::Init(IEngineLoop* InEngineLoop)
 			TEXT("EditorInteractiveToolsFramework")
 		};
 
-		FScopedSlowTask ModuleSlowTask(ARRAY_COUNT(ModuleNames));
+		FScopedSlowTask ModuleSlowTask(UE_ARRAY_COUNT(ModuleNames));
 		for (const TCHAR* ModuleName : ModuleNames)
 		{
 			ModuleSlowTask.EnterProgressFrame(1);
diff --git a/Engine/Source/Editor/UnrealEd/Private/EditorServer.cpp b/Engine/Source/Editor/UnrealEd/Private/EditorServer.cpp
--- a/Engine/Source/Editor/UnrealEd/Private/EditorServer.cpp
+++ b/Engine/Source/Editor/UnrealEd/Private/EditorServer.cpp
@@ -2501,7 +2501,7 @@ bool UEditorEngine::Map_Load(const TCHAR* Str, FOutputDevice& Ar)
 				const int32 MAX_STREAMLVL_SIZE = 16384;  // max cmd line size (16kb)
 				TCHAR StreamLvlBuf[MAX_STREAMLVL_SIZE]; //There can be a lot of streaming levels with very large path names
 
-				if(FParse::Value(Str, TEXT("STREAMLVL="), StreamLvlBuf, ARRAY_COUNT(StreamLvlBuf)))
+				if(FParse::Value(Str, TEXT("STREAMLVL="), StreamLvlBuf, UE_ARRAY_COUNT(StreamLvlBuf)))
 				{
 					TCHAR *ContextStr = NULL;
 					
diff --git a/Engine/Source/Editor/UnrealEd/Private/Factories/EditorFactories.cpp b/Engine/Source/Editor/UnrealEd/Private/Factories/EditorFactories.cpp
--- a/Engine/Source/Editor/UnrealEd/Private/Factories/EditorFactories.cpp
+++ b/Engine/Source/Editor/UnrealEd/Private/Factories/EditorFactories.cpp
@@ -2169,11 +2169,11 @@ public:
 	{
 		Ar << H.Manufacturer << H.Version << H.Encoding << H.BitsPerPixel;
 		Ar << H.XMin << H.YMin << H.XMax << H.YMax << H.XDotsPerInch << H.YDotsPerInch;
-		for( int32 i=0; i<ARRAY_COUNT(H.OldColorMap); i++ )
+		for( int32 i=0; i<UE_ARRAY_COUNT(H.OldColorMap); i++ )
 			Ar << H.OldColorMap[i];
 		Ar << H.Reserved1 << H.NumPlanes;
 		Ar << H.BytesPerLine << H.PaletteType << H.HScreenSize << H.VScreenSize;
-		for( int32 i=0; i<ARRAY_COUNT(H.Reserved2); i++ )
+		for( int32 i=0; i<UE_ARRAY_COUNT(H.Reserved2); i++ )
 			Ar << H.Reserved2[i];
 		return Ar;
 	}
diff --git a/Engine/Source/Editor/UnrealEd/Private/Fbx/FbxMainImport.cpp b/Engine/Source/Editor/UnrealEd/Private/Fbx/FbxMainImport.cpp
--- a/Engine/Source/Editor/UnrealEd/Private/Fbx/FbxMainImport.cpp
+++ b/Engine/Source/Editor/UnrealEd/Private/Fbx/FbxMainImport.cpp
@@ -1633,7 +1633,7 @@ ANSICHAR* FFbxImporter::MakeName(const ANSICHAR* Name) const
 	
 	FCStringAnsi::Strcpy(TmpName, len + 1, Name);
 
-	for ( int32 i = 0; i < ARRAY_COUNT(SpecialChars); i++ )
+	for ( int32 i = 0; i < UE_ARRAY_COUNT(SpecialChars); i++ )
 	{
 		ANSICHAR* CharPtr = TmpName;
 		while ( (CharPtr = FCStringAnsi::Strchr(CharPtr,SpecialChars[i])) != NULL )
