Added clamping to the cell width and cell height settings when extracting sprites from a texture, to prevent a hang if they were ever set to 0 (also adjusted their default values so resetting to default won't take them to 0)

diff --git a/Engine/Plugins/2D/Paper2D/Source/Paper2DEditor/Private/ExtractSprites/PaperExtractSpritesSettings.h b/Engine/Plugins/2D/Paper2D/Source/Paper2DEditor/Private/ExtractSprites/PaperExtractSpritesSettings.h
--- a/Engine/Plugins/2D/Paper2D/Source/Paper2DEditor/Private/ExtractSprites/PaperExtractSpritesSettings.h
+++ b/Engine/Plugins/2D/Paper2D/Source/Paper2DEditor/Private/ExtractSprites/PaperExtractSpritesSettings.h
@@ -67,11 +67,11 @@ public:
 
 	// The width of each sprite in grid mode
 	UPROPERTY(Category = Grid, EditAnywhere, meta = (UIMin = 1, ClampMin = 1))
-	int32 CellWidth;
+	int32 CellWidth = 64;
 
 	// The height of each sprite in grid mode
 	UPROPERTY(Category = Grid, EditAnywhere, meta = (UIMin = 1, ClampMin = 1))
-	int32 CellHeight;
+	int32 CellHeight = 64;
 
 	// Number of cells extracted horizontally. Can be used to limit the number of sprites extracted. Set to 0 to extract all sprites
 	UPROPERTY(Category = Grid, EditAnywhere, meta = (UIMin = 0, ClampMin = 0))
diff --git a/Engine/Plugins/2D/Paper2D/Source/Paper2DEditor/Private/ExtractSprites/SPaperExtractSpritesDialog.cpp b/Engine/Plugins/2D/Paper2D/Source/Paper2DEditor/Private/ExtractSprites/SPaperExtractSpritesDialog.cpp
--- a/Engine/Plugins/2D/Paper2D/Source/Paper2DEditor/Private/ExtractSprites/SPaperExtractSpritesDialog.cpp
+++ b/Engine/Plugins/2D/Paper2D/Source/Paper2DEditor/Private/ExtractSprites/SPaperExtractSpritesDialog.cpp
@@ -360,14 +360,17 @@ void SPaperExtractSpritesDialog::PreviewExtractedSprites()
 			const int32 TextureWidth = TextureSize.X;
 			const int32 TextureHeight = TextureSize.Y;
 
+			const int32 EffectiveCellWidth = FMath::Clamp(ExtractSpriteGridSettings->CellWidth, 1, TextureWidth);
+			const int32 EffectiveCellHeight = FMath::Clamp(ExtractSpriteGridSettings->CellHeight, 1, TextureHeight);
+
 			int NumExtractedCellsY = 0;
-			for (int32 Y = ExtractSpriteGridSettings->MarginY; Y + ExtractSpriteGridSettings->CellHeight <= TextureHeight; Y += ExtractSpriteGridSettings->CellHeight + ExtractSpriteGridSettings->SpacingY)
+			for (int32 Y = ExtractSpriteGridSettings->MarginY; Y + EffectiveCellHeight <= TextureHeight; Y += EffectiveCellHeight + ExtractSpriteGridSettings->SpacingY)
 			{
 				int NumExtractedCellsX = 0;
-				for (int32 X = ExtractSpriteGridSettings->MarginX; X + ExtractSpriteGridSettings->CellWidth <= TextureWidth; X += ExtractSpriteGridSettings->CellWidth + ExtractSpriteGridSettings->SpacingX)
+				for (int32 X = ExtractSpriteGridSettings->MarginX; X + EffectiveCellWidth <= TextureWidth; X += EffectiveCellWidth + ExtractSpriteGridSettings->SpacingX)
 				{
 					FPaperExtractedSprite* Sprite = new(ExtractedSprites)FPaperExtractedSprite();
-					Sprite->Rect = FIntRect(X, Y, X + ExtractSpriteGridSettings->CellWidth, Y + ExtractSpriteGridSettings->CellHeight);
+					Sprite->Rect = FIntRect(X, Y, X + EffectiveCellWidth, Y + EffectiveCellHeight);
 					
 					Sprite->Name = NamingTemplate;
 					Sprite->Name.ReplaceInline(TEXT("{0}"), *FString::Printf(TEXT("%d"), ExtractedRectIndex));
