Use an index buffer for Paper2D render proxies, so batches can share a single vertex buffer again.

Also added methods for pre-allocating the vertex and index buffers used by DynamicMeshBuilder internally, and deleted an unused function.

#rnx

diff --git a/Engine/Plugins/2D/Paper2D/Source/Paper2D/Private/PaperRenderSceneProxy.cpp b/Engine/Plugins/2D/Paper2D/Source/Paper2D/Private/PaperRenderSceneProxy.cpp
--- a/Engine/Plugins/2D/Paper2D/Source/Paper2D/Private/PaperRenderSceneProxy.cpp
+++ b/Engine/Plugins/2D/Paper2D/Source/Paper2D/Private/PaperRenderSceneProxy.cpp
@@ -14,7 +14,6 @@ static TAutoConsoleVariable<int32> CVarDrawSpritesAsTwoSided(TEXT("r.Paper2D.Dra
 
 DECLARE_CYCLE_STAT(TEXT("Get Batch Mesh"), STAT_PaperRender_GetBatchMesh, STATGROUP_Paper2D);
 DECLARE_CYCLE_STAT(TEXT("Get New Batch Meshes"), STAT_PaperRender_GetNewBatchMeshes, STATGROUP_Paper2D);
-DECLARE_CYCLE_STAT(TEXT("Convert Batches"), STAT_PaperRender_ConvertBatches, STATGROUP_Paper2D);
 DECLARE_CYCLE_STAT(TEXT("SpriteProxy GDME"), STAT_PaperRenderSceneProxy_GetDynamicMeshElements, STATGROUP_Paper2D);
 
 FPackedNormal FPaperSpriteTangents::PackedNormalX(FVector(1.0f, 0.0f, 0.0f));
@@ -277,15 +276,16 @@ void FPaperRenderSceneProxy::GetNewBatchMeshes(const FSceneView* View, int32 Vie
 
 	if (Vertices.Num())
 	{
-		for (const FSpriteRenderSection& Batch : BatchedSections)
+		FDynamicMeshBuilder DynamicMeshBuilder(View->GetFeatureLevel());
+		DynamicMeshBuilder.AddVertices(Vertices);
+		DynamicMeshBuilder.ReserveTriangles(Vertices.Num() / 3);
+		for (int32 i = 0; i < Vertices.Num(); i += 3)
 		{
-			FDynamicMeshBuilder DynamicMeshBuilder(View->GetFeatureLevel());
-
-			for (int32 i = Batch.VertexOffset; i < Batch.VertexOffset + Batch.NumVertices; ++i)
-			{
-				DynamicMeshBuilder.AddVertex(Vertices[i]);
-			}
+			DynamicMeshBuilder.AddTriangle(i, i + 1, i + 2);
+		}
 
+		for (const FSpriteRenderSection& Batch : BatchedSections)
+		{
 			if (Batch.IsValid())
 			{
 				FMaterialRenderProxy* ParentMaterialProxy = Batch.Material->GetRenderProxy();
@@ -320,8 +320,10 @@ void FPaperRenderSceneProxy::GetNewBatchMeshes(const FSceneView* View, int32 Vie
 				Settings.CastShadow = bCastShadow;
 				Settings.bDisableBackfaceCulling = bDrawTwoSided;
 
-
-				DynamicMeshBuilder.GetMesh(GetLocalToWorld(), TextureOverrideMaterialProxy, DPG, Settings, nullptr, ViewIndex, Collector);
+				FDynamicMeshDrawOffset DrawOffset;
+				DrawOffset.FirstIndex = Batch.VertexOffset;
+				DrawOffset.NumPrimitives = Batch.NumVertices / 3;
+				DynamicMeshBuilder.GetMesh(GetLocalToWorld(), TextureOverrideMaterialProxy, DPG, Settings, &DrawOffset, ViewIndex, Collector);
 			}
 		}
 	}
@@ -516,25 +518,3 @@ void FPaperRenderSceneProxy::SetTransientTextureOverride_RenderThread(const UTex
 	}
 }
 #endif
-
-void FPaperRenderSceneProxy::ConvertBatchesToNewStyle(TArray<FSpriteDrawCallRecord>& SourceBatches)
-{
-	SCOPE_CYCLE_COUNTER(STAT_PaperRender_ConvertBatches);
-
-	Vertices.Reset();
-	BatchedSections.Reset();
-
-	for (const FSpriteDrawCallRecord& SourceBatch : SourceBatches)
-	{
-		if (SourceBatch.IsValid())
-		{
-			FSpriteRenderSection& DestBatch = *new (BatchedSections) FSpriteRenderSection();
-
-			DestBatch.BaseTexture = SourceBatch.BaseTexture;
-			DestBatch.AdditionalTextures = SourceBatch.AdditionalTextures;
-			DestBatch.Material = Material;
-
-			DestBatch.AddTriangles(SourceBatch, Vertices);
-		}
-	}
-}
diff --git a/Engine/Plugins/2D/Paper2D/Source/Paper2D/Private/PaperRenderSceneProxy.h b/Engine/Plugins/2D/Paper2D/Source/Paper2D/Private/PaperRenderSceneProxy.h
--- a/Engine/Plugins/2D/Paper2D/Source/Paper2D/Private/PaperRenderSceneProxy.h
+++ b/Engine/Plugins/2D/Paper2D/Source/Paper2D/Private/PaperRenderSceneProxy.h
@@ -148,8 +148,6 @@ protected:
 
 	bool IsCollisionView(const FEngineShowFlags& EngineShowFlags, bool& bDrawSimpleCollision, bool& bDrawComplexCollision) const;
 
-	void ConvertBatchesToNewStyle(TArray<FSpriteDrawCallRecord>& SourceBatches);
-
 	virtual void DebugDrawCollision(const FSceneView* View, int32 ViewIndex, FMeshElementCollector& Collector, bool bDrawSolid) const;
 	virtual void DebugDrawBodySetup(const FSceneView* View, int32 ViewIndex, FMeshElementCollector& Collector, UBodySetup* BodySetup, const FMatrix& GeomTransform, const FLinearColor& CollisionColor, bool bDrawSolid) const;
 protected:
diff --git a/Engine/Source/Runtime/Engine/Private/DynamicMeshBuilder.cpp b/Engine/Source/Runtime/Engine/Private/DynamicMeshBuilder.cpp
--- a/Engine/Source/Runtime/Engine/Private/DynamicMeshBuilder.cpp
+++ b/Engine/Source/Runtime/Engine/Private/DynamicMeshBuilder.cpp
@@ -638,6 +638,7 @@ void FDynamicMeshBuilder::AddTriangles(const TArray<uint32> &InIndices)
 	}
 	else
 	{
+		IndexBuffer->Indices.Reserve(IndexBuffer->Indices.Num() + InIndices.Num());
 		for (int32 Index : InIndices)
 		{
 			IndexBuffer->Indices.Add(static_cast<DynamicMeshIndexType>(Index));
@@ -645,6 +646,18 @@ void FDynamicMeshBuilder::AddTriangles(const TArray<uint32> &InIndices)
 	}
 }
 
+/** Pre-allocate space for the given number of vertices. */
+void FDynamicMeshBuilder::ReserveVertices(int32 InNumVertices)
+{
+	VertexBuffer->Vertices.Reserve(InNumVertices);
+}
+
+/** Pre-allocate space for the given number of triangles. */
+void FDynamicMeshBuilder::ReserveTriangles(int32 InNumTriangles)
+{
+	IndexBuffer->Indices.Reserve(3 * InNumTriangles);
+}
+
 FMeshBuilderOneFrameResources::~FMeshBuilderOneFrameResources()
 {
 	if (VertexBuffer)
diff --git a/Engine/Source/Runtime/Engine/Public/DynamicMeshBuilder.h b/Engine/Source/Runtime/Engine/Public/DynamicMeshBuilder.h
--- a/Engine/Source/Runtime/Engine/Public/DynamicMeshBuilder.h
+++ b/Engine/Source/Runtime/Engine/Public/DynamicMeshBuilder.h
@@ -172,6 +172,12 @@ public:
 	/** Add many indices to the mesh. */
 	ENGINE_API void AddTriangles(const TArray<uint32> &InIndices);
 
+	/** Pre-allocate space for the given number of vertices. */
+	ENGINE_API void ReserveVertices(int32 InNumVertices);
+
+	/** Pre-allocate space for the given number of triangles. */
+	ENGINE_API void ReserveTriangles(int32 InNumTriangles);
+
 	/** Adds a mesh of what's been built so far to the collector. */
 	ENGINE_API void GetMesh(const FMatrix& LocalToWorld, const FMaterialRenderProxy* MaterialRenderProxy, uint8 DepthPriorityGroup, bool bDisableBackfaceCulling, bool bReceivesDecals, int32 ViewIndex, FMeshElementCollector& Collector);
 	ENGINE_API void GetMesh(const FMatrix& LocalToWorld, const FMaterialRenderProxy* MaterialRenderProxy, uint8 DepthPriorityGroup, bool bDisableBackfaceCulling, bool bReceivesDecals, bool bUseSelectionOutline, int32 ViewIndex, 
