Paper2D: Fixed an issue where grouped sprites with material or texture changes would render all sprites after the split incorrectly

diff --git a/Engine/Plugins/2D/Paper2D/Source/Paper2D/Private/GroupedSpriteSceneProxy.cpp b/Engine/Plugins/2D/Paper2D/Source/Paper2D/Private/GroupedSpriteSceneProxy.cpp
--- a/Engine/Plugins/2D/Paper2D/Source/Paper2D/Private/GroupedSpriteSceneProxy.cpp
+++ b/Engine/Plugins/2D/Paper2D/Source/Paper2D/Private/GroupedSpriteSceneProxy.cpp
@@ -58,35 +58,90 @@ FGroupedSpriteSceneProxy::FGroupedSpriteSceneProxy(UPaperGroupedSpriteComponent*
 		BodySetups.Reserve(NumInstances);
 	}
 
+	// Create all the sections first so we can generate indices correctly
+	TArray<int32> SectionIndicies;
+	SectionIndicies.AddZeroed(NumInstances);
+
 	for (int32 InstanceIndex = 0; InstanceIndex < NumInstances; ++InstanceIndex)
 	{
-		const FSpriteInstanceData InstanceData = InComponent->PerInstanceSpriteData[InstanceIndex];
+		const FSpriteInstanceData& InstanceData = InComponent->PerInstanceSpriteData[InstanceIndex];
 
-		UBodySetup* BodySetup = nullptr;
 		if (UPaperSprite* SourceSprite = InstanceData.SourceSprite)
 		{
-			FSpriteDrawCallRecord Record;
-			Record.BuildFromSprite(SourceSprite);
-
+			UTexture2D* BaseTexture = SourceSprite->GetBakedTexture();
+			FAdditionalSpriteTextureArray AdditionalTextures;
+			SourceSprite->GetBakedAdditionalSourceTextures(/*out*/ AdditionalTextures);
 			UMaterialInterface* SpriteMaterial = InComponent->GetMaterial(InstanceData.MaterialIndex);
 
-			FSpriteRenderSection& Section = FindOrAddSection(Record, SpriteMaterial);
-			
-			const int32 NumNewVerts = Record.RenderVerts.Num();
-			Section.NumVertices += NumNewVerts;
+			FSpriteRenderSection* FoundSection = nullptr;
+			int32 FoundSectionIndex = INDEX_NONE;
+
+			for (int32 SectionIndex = BatchedSections.Num() - 1; SectionIndex >= 0; --SectionIndex)
+			{
+				FSpriteRenderSection& TestSection = BatchedSections[SectionIndex];
+				if (TestSection.Material == SpriteMaterial)
+				{
+					if (TestSection.BaseTexture == BaseTexture)
+					{
+						if (TestSection.AdditionalTextures == AdditionalTextures)
+						{
+							FoundSection = &TestSection;
+							FoundSectionIndex = SectionIndex;
+							break;
+						}
+					}
+				}
+			}
+
+			if (FoundSectionIndex == INDEX_NONE)
+			{
+				FoundSectionIndex = BatchedSections.Num();
+
+				// Didn't find a matching section, create one
+				FoundSection = new (BatchedSections) FSpriteRenderSection();
+				FoundSection->Material = SpriteMaterial;
+				FoundSection->BaseTexture = BaseTexture;
+				FoundSection->AdditionalTextures = AdditionalTextures;
+			}
+
+			SectionIndicies[InstanceIndex] = FoundSectionIndex;
+			FoundSection->NumVertices += SourceSprite->BakedRenderData.Num();
+		}
+	}
+
+	int32 RunningVertCount = 0;
+	for (FSpriteRenderSection& Section : BatchedSections)
+	{
+		Section.VertexOffset = RunningVertCount;
+		RunningVertCount += Section.NumVertices;
+		Section.NumVertices = 0;
+	}
+	check(Vertices.Num() == 0);
+	Vertices.AddUninitialized(RunningVertCount);
+
+	int32 InstanceIndex = 0;
+	for (const FSpriteInstanceData& InstanceData : InComponent->PerInstanceSpriteData)
+	{
+		UBodySetup* BodySetup = nullptr;
+		if (UPaperSprite* SourceSprite = InstanceData.SourceSprite)
+		{
+			const int32 SectionIndex = SectionIndicies[InstanceIndex];
+			FSpriteRenderSection& Section = BatchedSections[SectionIndex];
+			FDynamicMeshVertex* VertexWritePtr = Vertices.GetData() + Section.VertexOffset + Section.NumVertices;
+			Section.NumVertices += SourceSprite->BakedRenderData.Num();
 
 			const FPackedNormal TangentX = InstanceData.Transform.GetUnitAxis(EAxis::X);
 			FPackedNormal TangentZ = InstanceData.Transform.GetUnitAxis(EAxis::Y);
 			TangentZ.Vector.W = (InstanceData.Transform.Determinant() < 0.0f) ? -127 : 127;
 
 			const FColor VertColor(InstanceData.VertexColor);
-			for (const FVector4& SourceVert : Record.RenderVerts)
+			for (const FVector4& SourceVert : SourceSprite->BakedRenderData)
 			{
 				const FVector LocalPos((PaperAxisX * SourceVert.X) + (PaperAxisY * SourceVert.Y));
 				const FVector ComponentSpacePos = InstanceData.Transform.TransformPosition(LocalPos);
 				const FVector2D UV(SourceVert.Z, SourceVert.W);
 
-				new (Vertices) FDynamicMeshVertex(ComponentSpacePos, TangentX.ToFVector(), TangentZ.ToFVector(), UV, VertColor);
+				*VertexWritePtr++ = FDynamicMeshVertex(ComponentSpacePos, TangentX.ToFVector(), TangentZ.ToFVector(), UV, VertColor);
 			}
 
 			BodySetup = SourceSprite->BodySetup;
@@ -97,6 +152,7 @@ FGroupedSpriteSceneProxy::FGroupedSpriteSceneProxy(UPaperGroupedSpriteComponent*
 			BodySetupTransforms.Add(InstanceData.Transform);
 			BodySetups.Add(BodySetup);
 		}
+		++InstanceIndex;
 	}
 }
 
