Paper2D optimizations:
- Changed how texture param rewrites works for paper2d rendering so it doesn't have to rebuild a uniform buffer every frame when using MICs or non-varying MIDs
- Removed the legacy Paper2D rendering path so all components use batched sections
- Improved how flipbooks handle sprites with material splits

diff --git a/Engine/Plugins/2D/Paper2D/Source/Paper2D/Private/PaperFlipbookComponent.cpp b/Engine/Plugins/2D/Paper2D/Source/Paper2D/Private/PaperFlipbookComponent.cpp
--- a/Engine/Plugins/2D/Paper2D/Source/Paper2D/Private/PaperFlipbookComponent.cpp
+++ b/Engine/Plugins/2D/Paper2D/Source/Paper2D/Private/PaperFlipbookComponent.cpp
@@ -93,11 +93,12 @@ FPrimitiveSceneProxy* UPaperFlipbookComponent::CreateSceneProxy()
 
 	CalculateCurrentFrame();
 	UPaperSprite* SpriteToSend = GetSpriteAtCachedIndex();
+	const int32 SplitIndex = SpriteToSend ? SpriteToSend->AlternateMaterialSplitIndex : INDEX_NONE;
 
 	FSpriteDrawCallRecord DrawCall;
 	DrawCall.BuildFromSprite(SpriteToSend);
 	DrawCall.Color = SpriteColor.ToFColor(/*bSRGB=*/ false);
-	NewProxy->SetDrawCall_RenderThread(DrawCall);
+	NewProxy->SetSprite_RenderThread(DrawCall, SplitIndex);
 	return NewProxy;
 }
 
@@ -356,13 +357,14 @@ void UPaperFlipbookComponent::SendRenderDynamicData_Concurrent()
 		FSpriteDrawCallRecord DrawCall;
 		DrawCall.BuildFromSprite(SpriteToSend);
 		DrawCall.Color = SpriteColor.ToFColor(/*bSRGB=*/ false);
+		const int32 SplitIndex = SpriteToSend ? SpriteToSend->AlternateMaterialSplitIndex : INDEX_NONE;
 
-		FPaperRenderSceneProxy* InSceneProxy = (FPaperRenderSceneProxy*)SceneProxy;
+		FPaperFlipbookSceneProxy* InSceneProxy = (FPaperFlipbookSceneProxy*)SceneProxy;
 		UBodySetup* InBodySetup = CachedBodySetup;
 		ENQUEUE_RENDER_COMMAND(FSendPaperRenderComponentDynamicData)(
-			[InSceneProxy, DrawCall, InBodySetup](FRHICommandListImmediate& RHICmdList)
+			[InSceneProxy, DrawCall, InBodySetup, SplitIndex](FRHICommandListImmediate& RHICmdList)
 			{
-				InSceneProxy->SetDrawCall_RenderThread(DrawCall);
+				InSceneProxy->SetSprite_RenderThread(DrawCall, SplitIndex);
 				InSceneProxy->SetBodySetup_RenderThread(InBodySetup);
 			});
 	}
diff --git a/Engine/Plugins/2D/Paper2D/Source/Paper2D/Private/PaperFlipbookSceneProxy.cpp b/Engine/Plugins/2D/Paper2D/Source/Paper2D/Private/PaperFlipbookSceneProxy.cpp
--- a/Engine/Plugins/2D/Paper2D/Source/Paper2D/Private/PaperFlipbookSceneProxy.cpp
+++ b/Engine/Plugins/2D/Paper2D/Source/Paper2D/Private/PaperFlipbookSceneProxy.cpp
@@ -7,12 +7,8 @@
 // FPaperFlipbookSceneProxy
 
 FPaperFlipbookSceneProxy::FPaperFlipbookSceneProxy(const UPaperFlipbookComponent* InComponent)
-	: FPaperRenderSceneProxy(InComponent)
+	: FPaperRenderSceneProxy_SpriteBase(InComponent)
 {
-	//@TODO: PAPER2D: WireframeColor = RenderComp->GetWireframeColor();
-
-	Material = InComponent->GetMaterial(0);
-	MaterialRelevance = InComponent->GetMaterialRelevance(GetScene().GetFeatureLevel());
 }
 
 SIZE_T FPaperFlipbookSceneProxy::GetTypeHash() const
diff --git a/Engine/Plugins/2D/Paper2D/Source/Paper2D/Private/PaperFlipbookSceneProxy.h b/Engine/Plugins/2D/Paper2D/Source/Paper2D/Private/PaperFlipbookSceneProxy.h
--- a/Engine/Plugins/2D/Paper2D/Source/Paper2D/Private/PaperFlipbookSceneProxy.h
+++ b/Engine/Plugins/2D/Paper2D/Source/Paper2D/Private/PaperFlipbookSceneProxy.h
@@ -8,7 +8,7 @@
 //////////////////////////////////////////////////////////////////////////
 // FPaperFlipbookSceneProxy
 
-class FPaperFlipbookSceneProxy final : public FPaperRenderSceneProxy
+class FPaperFlipbookSceneProxy final : public FPaperRenderSceneProxy_SpriteBase
 {
 public:
 	SIZE_T GetTypeHash() const override;
diff --git a/Engine/Plugins/2D/Paper2D/Source/Paper2D/Private/PaperRenderSceneProxy.cpp b/Engine/Plugins/2D/Paper2D/Source/Paper2D/Private/PaperRenderSceneProxy.cpp
--- a/Engine/Plugins/2D/Paper2D/Source/Paper2D/Private/PaperRenderSceneProxy.cpp
+++ b/Engine/Plugins/2D/Paper2D/Source/Paper2D/Private/PaperRenderSceneProxy.cpp
@@ -7,12 +7,11 @@
 #include "PhysicsEngine/BodySetup.h"
 #include "EngineGlobals.h"
 #include "Components/PrimitiveComponent.h"
+#include "Components/MeshComponent.h"
 #include "Engine/Engine.h"
-#include "DynamicMeshBuilder.h"
 
 static TAutoConsoleVariable<int32> CVarDrawSpritesAsTwoSided(TEXT("r.Paper2D.DrawTwoSided"), 1, TEXT("Draw sprites as two sided."));
 
-DECLARE_CYCLE_STAT(TEXT("Get Batch Mesh"), STAT_PaperRender_GetBatchMesh, STATGROUP_Paper2D);
 DECLARE_CYCLE_STAT(TEXT("Get New Batch Meshes"), STAT_PaperRender_GetNewBatchMeshes, STATGROUP_Paper2D);
 DECLARE_CYCLE_STAT(TEXT("SpriteProxy GDME"), STAT_PaperRenderSceneProxy_GetDynamicMeshElements, STATGROUP_Paper2D);
 
@@ -36,11 +35,14 @@ void FPaperSpriteTangents::SetTangentsFromPaperAxes()
 class FSpriteTextureOverrideRenderProxy : public FDynamicPrimitiveResource, public FMaterialRenderProxy
 {
 public:
-	const FMaterialRenderProxy* const Parent;
+	const FMaterialRenderProxy* Parent;
 	const UTexture* BaseTexture;
 	FAdditionalSpriteTextureArray AdditionalTextures;
 	UE_EXPAND_IF_WITH_EDITOR(const FPaperRenderSceneProxyTextureOverrideMap& TextureOverrideList);
 
+	// Used to track if we need to rebuild our texture param setting proxy
+	int32 ParentMaterialSerialNumber;
+
 public:
 	/** Initialization constructor. */
 	FSpriteTextureOverrideRenderProxy(const FMaterialRenderProxy* InParent, const UTexture* InBaseTexture, FAdditionalSpriteTextureArray InAdditionalTextures UE_EXPAND_IF_WITH_EDITOR(, const FPaperRenderSceneProxyTextureOverrideMap& InTextureOverrideList))
@@ -48,12 +50,41 @@ public:
 		, BaseTexture(InBaseTexture)
 		, AdditionalTextures(InAdditionalTextures)
 		UE_EXPAND_IF_WITH_EDITOR(, TextureOverrideList(InTextureOverrideList))
-	{}
+		, ParentMaterialSerialNumber(INDEX_NONE)
+	{
+	}
 
 	virtual ~FSpriteTextureOverrideRenderProxy()
 	{
 	}
 
+	void CheckValidity(FMaterialRenderProxy* InCurrentParent)
+	{
+		if (InCurrentParent != Parent)
+		{
+			Parent = InCurrentParent;
+			ParentMaterialSerialNumber = INDEX_NONE;
+		}
+
+		if (ParentMaterialSerialNumber != Parent->GetExpressionCacheSerialNumber())
+		{
+			// Not valid, need to rebuild
+			CacheUniformExpressions(/*bRecreateUniformBuffer=*/ true);
+			ParentMaterialSerialNumber = Parent->GetExpressionCacheSerialNumber();
+		}
+	}
+
+	void Reinitialize(const FMaterialRenderProxy* InParent, const UTexture* InBaseTexture, FAdditionalSpriteTextureArray InAdditionalTextures)
+	{
+		if ((InParent != Parent) || (InBaseTexture != BaseTexture) || (InAdditionalTextures != AdditionalTextures))
+		{
+			Parent = InParent;
+			BaseTexture = InBaseTexture;
+			AdditionalTextures = InAdditionalTextures;
+			ParentMaterialSerialNumber = INDEX_NONE;
+		}
+	}
+
 	// FDynamicPrimitiveResource interface.
 	virtual void InitPrimitiveResource() override
 	{
@@ -136,7 +167,6 @@ const FName FSpriteTextureOverrideRenderProxy::AdditionalTextureParameterRootNam
 
 FPaperRenderSceneProxy::FPaperRenderSceneProxy(const UPrimitiveComponent* InComponent)
 	: FPrimitiveSceneProxy(InComponent)
-	, Material(nullptr)
 	, Owner(InComponent->GetOwner())
 	, MyBodySetup(const_cast<UPrimitiveComponent*>(InComponent)->GetBodySetup())
 	, bCastShadow(InComponent->CastShadow)
@@ -155,6 +185,13 @@ SIZE_T FPaperRenderSceneProxy::GetTypeHash() const
 
 FPaperRenderSceneProxy::~FPaperRenderSceneProxy()
 {
+	for (FSpriteTextureOverrideRenderProxy* Proxy : MaterialTextureOverrideProxies)
+	{
+		if (Proxy != nullptr)
+		{
+			Proxy->ReleasePrimitiveResource();
+		}
+	}
 }
 
 void FPaperRenderSceneProxy::DebugDrawBodySetup(const FSceneView* View, int32 ViewIndex, FMeshElementCollector& Collector, UBodySetup* BodySetup, const FMatrix& GeomTransformMatrix, const FLinearColor& CollisionColor, bool bDrawSolid) const
@@ -189,6 +226,20 @@ void FPaperRenderSceneProxy::DebugDrawBodySetup(const FSceneView* View, int32 Vi
 	}
 }
 
+void FPaperRenderSceneProxy::RecreateCachedRenderData()
+{
+	int32 BatchIndex = 0;
+	for (FSpriteTextureOverrideRenderProxy* Proxy : MaterialTextureOverrideProxies)
+	{
+		if ((Proxy != nullptr) && BatchedSections.IsValidIndex(BatchIndex))
+		{
+			const FSpriteRenderSection& Section = BatchedSections[BatchIndex];
+			Proxy->Reinitialize(Section.Material->GetRenderProxy(), Section.BaseTexture, Section.AdditionalTextures);
+		}
+		++BatchIndex;
+	}
+}
+
 void FPaperRenderSceneProxy::DebugDrawCollision(const FSceneView* View, int32 ViewIndex, FMeshElementCollector& Collector, bool bDrawSolid) const
 {
 	if (MyBodySetup != nullptr)
@@ -230,7 +281,6 @@ void FPaperRenderSceneProxy::GetDynamicMeshElements(const TArray<const FSceneVie
 		}
 	}
 
-
 	for (int32 ViewIndex = 0; ViewIndex < Views.Num(); ViewIndex++)
 	{
 		if (VisibilityMap & (1 << ViewIndex))
@@ -255,10 +305,6 @@ void FPaperRenderSceneProxy::GetDynamicMeshElements(const TArray<const FSceneVie
 
 void FPaperRenderSceneProxy::GetDynamicMeshElementsForView(const FSceneView* View, int32 ViewIndex, FMeshElementCollector& Collector) const
 {
-	if (Material != nullptr)
-	{
-		GetBatchMesh(View, Material, BatchedSprites, ViewIndex, Collector);
-	}
 	GetNewBatchMeshes(View, ViewIndex, Collector);
 }
 
@@ -274,6 +320,12 @@ void FPaperRenderSceneProxy::GetNewBatchMeshes(const FSceneView* View, int32 Vie
 	const uint8 DPG = GetDepthPriorityGroup(View);
 	const bool bIsWireframeView = View->Family->EngineShowFlags.Wireframe;
 
+	if (MaterialTextureOverrideProxies.Num() < BatchedSections.Num())
+	{
+		MaterialTextureOverrideProxies.AddDefaulted(BatchedSections.Num() - MaterialTextureOverrideProxies.Num());
+	}
+	FSpriteTextureOverrideRenderProxy** pCurrentMaterialProxy = MaterialTextureOverrideProxies.GetData();
+
 	if (Vertices.Num())
 	{
 		FDynamicMeshBuilder DynamicMeshBuilder(View->GetFeatureLevel());
@@ -313,9 +365,12 @@ void FPaperRenderSceneProxy::GetNewBatchMeshes(const FSceneView* View, int32 Vie
 					Settings.bCanApplyViewModeOverrides = false;
 				}
 
-				// Create a texture override material proxy and register it as a dynamic resource so that it won't be deleted until the rendering thread has finished with it
-				FSpriteTextureOverrideRenderProxy* TextureOverrideMaterialProxy = new FSpriteTextureOverrideRenderProxy(ParentMaterialProxy, Batch.BaseTexture, Batch.AdditionalTextures UE_EXPAND_IF_WITH_EDITOR(, TextureOverrideList));
-				Collector.RegisterOneFrameMaterialProxy(TextureOverrideMaterialProxy);
+				// Create a texture override material proxy or make sure our existing one is valid
+				if (*pCurrentMaterialProxy == nullptr)
+				{
+					*pCurrentMaterialProxy = new FSpriteTextureOverrideRenderProxy(ParentMaterialProxy, Batch.BaseTexture, Batch.AdditionalTextures UE_EXPAND_IF_WITH_EDITOR(, TextureOverrideList));
+				}
+				(*pCurrentMaterialProxy)->CheckValidity(ParentMaterialProxy);
 
 				Settings.CastShadow = bCastShadow;
 				Settings.bDisableBackfaceCulling = bDrawTwoSided;
@@ -323,8 +378,10 @@ void FPaperRenderSceneProxy::GetNewBatchMeshes(const FSceneView* View, int32 Vie
 				FDynamicMeshDrawOffset DrawOffset;
 				DrawOffset.FirstIndex = Batch.VertexOffset;
 				DrawOffset.NumPrimitives = Batch.NumVertices / 3;
-				DynamicMeshBuilder.GetMesh(GetLocalToWorld(), TextureOverrideMaterialProxy, DPG, Settings, &DrawOffset, ViewIndex, Collector);
+				DynamicMeshBuilder.GetMesh(GetLocalToWorld(), *pCurrentMaterialProxy, DPG, Settings, &DrawOffset, ViewIndex, Collector);
 			}
+
+			++pCurrentMaterialProxy;
 		}
 	}
 }
@@ -335,81 +392,6 @@ public:
 	TArray<FDynamicMeshVertex> Vertices;
 };
 
-void FPaperRenderSceneProxy::GetBatchMesh(const FSceneView* View, class UMaterialInterface* BatchMaterial, const TArray<FSpriteDrawCallRecord>& Batch, int32 ViewIndex, FMeshElementCollector& Collector) const
-{
-	SCOPE_CYCLE_COUNTER(STAT_PaperRender_GetBatchMesh);
-
-	const uint8 DPG = GetDepthPriorityGroup(View);
-
-	const bool bIsWireframeView = View->Family->EngineShowFlags.Wireframe;
-
-	FMatrix PreviousLocalToWorld;
-
-	if (!GetScene().GetPreviousLocalToWorld(GetPrimitiveSceneInfo(), PreviousLocalToWorld))
-	{
-		PreviousLocalToWorld = GetLocalToWorld();
-	}
-
-	for (const FSpriteDrawCallRecord& Record : Batch)
-	{
-		if (Record.IsValid())
-		{
-			const FColor SpriteColor(Record.Color);
-			const FVector EffectiveOrigin = Record.Destination;
-
-			FPaperVertexArray& VertexArray = Collector.AllocateOneFrameResource<FPaperVertexArray>();
-			VertexArray.Vertices.Empty(Record.RenderVerts.Num());
-
-			for (int32 SVI = 0; SVI < Record.RenderVerts.Num(); ++SVI)
-			{
-				const FVector4& SourceVert = Record.RenderVerts[SVI];
-				const FVector Pos((PaperAxisX * SourceVert.X) + (PaperAxisY * SourceVert.Y) + EffectiveOrigin);
-				const FVector2D UV(SourceVert.Z, SourceVert.W);
-
-				new (VertexArray.Vertices) FDynamicMeshVertex(Pos, FPaperSpriteTangents::PackedNormalX.ToFVector(), FPaperSpriteTangents::PackedNormalZ.ToFVector(), UV, SpriteColor);
-			}
-
-			FDynamicMeshBuilder DynamicMeshBuilder(View->GetFeatureLevel());
-			DynamicMeshBuilder.AddVertices(VertexArray.Vertices);
-
-			FDynamicMeshBuilderSettings Settings;
-			Settings.bCanApplyViewModeOverrides = true;
-			Settings.bUseWireframeSelectionColoring = IsSelected();
-
-			FMaterialRenderProxy* ParentMaterialProxy = BatchMaterial->GetRenderProxy();
-
-			// Implementing our own wireframe coloring as the automatic one (controlled by Mesh.bCanApplyViewModeOverrides) only supports per-FPrimitiveSceneProxy WireframeColor
-			if (bIsWireframeView)
-			{
-				const FLinearColor EffectiveWireframeColor = (BatchMaterial->GetBlendMode() != BLEND_Opaque) ? GetWireframeColor() : FLinearColor::Green;
-
-				auto WireframeMaterialInstance = new FColoredMaterialRenderProxy(
-					GEngine->WireframeMaterial->GetRenderProxy(),
-					GetSelectionColor(EffectiveWireframeColor, IsSelected(), IsHovered(), false)
-					);
-
-				Collector.RegisterOneFrameMaterialProxy(WireframeMaterialInstance);
-
-				ParentMaterialProxy = WireframeMaterialInstance;
-
-				Settings.bWireframe = true;
-				// We are applying our own wireframe override
-				Settings.bCanApplyViewModeOverrides = false;
-			}
-
-			// Create a texture override material proxy and register it as a dynamic resource so that it won't be deleted until the rendering thread has finished with it
-			FSpriteTextureOverrideRenderProxy* TextureOverrideMaterialProxy = new FSpriteTextureOverrideRenderProxy(ParentMaterialProxy, Record.BaseTexture, Record.AdditionalTextures UE_EXPAND_IF_WITH_EDITOR(, TextureOverrideList));
-			Collector.RegisterOneFrameMaterialProxy(TextureOverrideMaterialProxy);
-		
-			Settings.bDisableBackfaceCulling = bDrawTwoSided;
-			Settings.bReceivesDecals = true;
-			Settings.bUseSelectionOutline = true;
-			Settings.CastShadow = bCastShadow;
-			DynamicMeshBuilder.GetMesh(GetLocalToWorld(), PreviousLocalToWorld, TextureOverrideMaterialProxy, DPG, Settings, nullptr, ViewIndex, Collector);
-		}
-	}
-}
-
 FPrimitiveViewRelevance FPaperRenderSceneProxy::GetViewRelevance(const FSceneView* View) const
 {
 	const FEngineShowFlags& EngineShowFlags = View->Family->EngineShowFlags;
@@ -467,16 +449,6 @@ bool FPaperRenderSceneProxy::IsUsingDistanceCullFade() const
 	return MaterialRelevance.bUsesDistanceCullFade;
 }
 
-void FPaperRenderSceneProxy::SetDrawCall_RenderThread(const FSpriteDrawCallRecord& NewDynamicData)
-{
-	SCOPE_CYCLE_COUNTER(STAT_PaperRender_SetSpriteRT);
-
-	BatchedSprites.Empty();
-
-	FSpriteDrawCallRecord& Record = *new (BatchedSprites) FSpriteDrawCallRecord;
-	Record = NewDynamicData;
-}
-
 void FPaperRenderSceneProxy::SetBodySetup_RenderThread(UBodySetup* NewSetup)
 {
 	MyBodySetup = NewSetup;
@@ -518,3 +490,55 @@ void FPaperRenderSceneProxy::SetTransientTextureOverride_RenderThread(const UTex
 	}
 }
 #endif
+
+//////////////////////////////////////////////////////////////////////////
+// FPaperRenderSceneProxy_SpriteBase
+
+FPaperRenderSceneProxy_SpriteBase::FPaperRenderSceneProxy_SpriteBase(const UMeshComponent* InComponent)
+	: FPaperRenderSceneProxy(InComponent)
+{
+	Material = InComponent->GetMaterial(0);
+	if (Material == nullptr)
+	{
+		Material = UMaterial::GetDefaultMaterial(MD_Surface);
+	}
+
+	AlternateMaterial = InComponent->GetMaterial(1);
+	if (AlternateMaterial == nullptr)
+	{
+		AlternateMaterial = UMaterial::GetDefaultMaterial(MD_Surface);
+	}
+
+	MaterialRelevance = InComponent->GetMaterialRelevance(GetScene().GetFeatureLevel());
+}
+
+void FPaperRenderSceneProxy_SpriteBase::SetSprite_RenderThread(const FSpriteDrawCallRecord& NewDynamicData, int32 SplitIndex)
+{
+	SCOPE_CYCLE_COUNTER(STAT_PaperRender_SetSpriteRT);
+
+	BatchedSections.Reset();
+	Vertices.Reset();
+
+	if (SplitIndex != INDEX_NONE)
+	{
+		BatchedSections.AddDefaulted();
+		BatchedSections.AddDefaulted();
+
+		FSpriteRenderSection& Section = BatchedSections[0];
+		Section.Material = Material;
+		Section.AddVerticesFromDrawCallRecord(NewDynamicData, 0, SplitIndex, Vertices);
+
+		FSpriteRenderSection& AlternateSection = BatchedSections[1];
+		AlternateSection.Material = AlternateMaterial;
+		AlternateSection.AddVerticesFromDrawCallRecord(NewDynamicData, SplitIndex, NewDynamicData.RenderVerts.Num() - SplitIndex, Vertices);
+	}
+	else
+	{
+		FSpriteRenderSection& Section = BatchedSections[BatchedSections.AddDefaulted()];
+		Section.Material = Material;
+		Section.AddVerticesFromDrawCallRecord(NewDynamicData, 0, NewDynamicData.RenderVerts.Num(), Vertices);
+	}
+
+	RecreateCachedRenderData();
+}
+
diff --git a/Engine/Plugins/2D/Paper2D/Source/Paper2D/Private/PaperRenderSceneProxy.h b/Engine/Plugins/2D/Paper2D/Source/Paper2D/Private/PaperRenderSceneProxy.h
--- a/Engine/Plugins/2D/Paper2D/Source/Paper2D/Private/PaperRenderSceneProxy.h
+++ b/Engine/Plugins/2D/Paper2D/Source/Paper2D/Private/PaperRenderSceneProxy.h
@@ -17,6 +17,7 @@
 class FMeshElementCollector;
 class UBodySetup;
 class UPrimitiveComponent;
+class FSpriteTextureOverrideRenderProxy;
 
 #if WITH_EDITOR
 typedef TMap<const UTexture*, const UTexture*> FPaperRenderSceneProxyTextureOverrideMap;
@@ -66,7 +67,7 @@ struct PAPER2D_API FSpriteRenderSection
 	}
 
 	template <typename SourceArrayType>
-	void AddTriangles(const FSpriteDrawCallRecord& Record, SourceArrayType& Vertices)
+	void AddVerticesFromDrawCallRecord(const FSpriteDrawCallRecord& Record, int32 StartIndexWithinRecord, int32 NumVertsToCopy, SourceArrayType& Vertices)
 	{
 		if (NumVertices == 0)
 		{
@@ -81,13 +82,14 @@ struct PAPER2D_API FSpriteRenderSection
 			// Note: Not checking AdditionalTextures for now, since checking BaseTexture should catch most bugs
 		}
 
-		const int32 NumNewVerts = Record.RenderVerts.Num();
-		NumVertices += NumNewVerts;
-		Vertices.Reserve(Vertices.Num() + NumNewVerts);
+		NumVertices += NumVertsToCopy;
+		Vertices.Reserve(Vertices.Num() + NumVertsToCopy);
 
 		const FColor VertColor(Record.Color);
-		for (const FVector4& SourceVert : Record.RenderVerts)
+		for (int32 VertexIndex = StartIndexWithinRecord; VertexIndex < StartIndexWithinRecord + NumVertsToCopy; ++VertexIndex)
 		{
+			const FVector4& SourceVert = Record.RenderVerts[VertexIndex];
+			
 			const FVector Pos((PaperAxisX * SourceVert.X) + (PaperAxisY * SourceVert.Y) + Record.Destination);
 			const FVector2D UV(SourceVert.Z, SourceVert.W);
 
@@ -133,7 +135,6 @@ public:
 	virtual bool IsUsingDistanceCullFade() const override;
 	// End of FPrimitiveSceneProxy interface.
 
-	void SetDrawCall_RenderThread(const FSpriteDrawCallRecord& NewDynamicData);
 	void SetBodySetup_RenderThread(UBodySetup* NewSetup);
 
 #if WITH_EDITOR
@@ -143,20 +144,19 @@ public:
 protected:
 	virtual void GetDynamicMeshElementsForView(const FSceneView* View, int32 ViewIndex, FMeshElementCollector& Collector) const;
 
-	void GetBatchMesh(const FSceneView* View, UMaterialInterface* BatchMaterial, const TArray<FSpriteDrawCallRecord>& Batch, int32 ViewIndex, FMeshElementCollector& Collector) const;
 	void GetNewBatchMeshes(const FSceneView* View, int32 ViewIndex, FMeshElementCollector& Collector) const;
 
 	bool IsCollisionView(const FEngineShowFlags& EngineShowFlags, bool& bDrawSimpleCollision, bool& bDrawComplexCollision) const;
 
 	virtual void DebugDrawCollision(const FSceneView* View, int32 ViewIndex, FMeshElementCollector& Collector, bool bDrawSolid) const;
 	virtual void DebugDrawBodySetup(const FSceneView* View, int32 ViewIndex, FMeshElementCollector& Collector, UBodySetup* BodySetup, const FMatrix& GeomTransform, const FLinearColor& CollisionColor, bool bDrawSolid) const;
+
+	// Call this if you modify BatchedSections or Vertices after the proxy has already been created
+	void RecreateCachedRenderData();
 protected:
 	TArray<FSpriteRenderSection> BatchedSections;
 	TArray<FDynamicMeshVertex> Vertices;
-
-	// Old style
-	TArray<FSpriteDrawCallRecord> BatchedSprites;
-	class UMaterialInterface* Material;
+	mutable TArray<FSpriteTextureOverrideRenderProxy*> MaterialTextureOverrideProxies;
 
 	//
 	AActor* Owner;
@@ -174,3 +174,18 @@ protected:
 	// The texture override list
 	UE_EXPAND_IF_WITH_EDITOR(FPaperRenderSceneProxyTextureOverrideMap TextureOverrideList);
 };
+
+//////////////////////////////////////////////////////////////////////////
+// FPaperRenderSceneProxy_SpriteBase - common base class for sprites and flipbooks (which build from sprites)
+
+class FPaperRenderSceneProxy_SpriteBase : public FPaperRenderSceneProxy
+{
+public:
+	FPaperRenderSceneProxy_SpriteBase(const UMeshComponent* InComponent);
+
+	void SetSprite_RenderThread(const FSpriteDrawCallRecord& NewDynamicData, int32 SplitIndex);
+
+public:
+	UMaterialInterface* Material;
+	UMaterialInterface* AlternateMaterial;
+};
diff --git a/Engine/Plugins/2D/Paper2D/Source/Paper2D/Private/PaperSpriteSceneProxy.cpp b/Engine/Plugins/2D/Paper2D/Source/Paper2D/Private/PaperSpriteSceneProxy.cpp
--- a/Engine/Plugins/2D/Paper2D/Source/Paper2D/Private/PaperSpriteSceneProxy.cpp
+++ b/Engine/Plugins/2D/Paper2D/Source/Paper2D/Private/PaperSpriteSceneProxy.cpp
@@ -12,25 +12,10 @@
 // FPaperSpriteSceneProxy
 
 FPaperSpriteSceneProxy::FPaperSpriteSceneProxy(UPaperSpriteComponent* InComponent)
-	: FPaperRenderSceneProxy(InComponent)
+	: FPaperRenderSceneProxy_SpriteBase(InComponent)
 	, BodySetup(InComponent->GetBodySetup())
 {
 	SetWireframeColor(InComponent->GetWireframeColor());
-
-	Material = InComponent->GetMaterial(0);
-	if (Material == nullptr)
-	{
-		Material = UMaterial::GetDefaultMaterial(MD_Surface);
-	}
-
-	AlternateMaterial = InComponent->GetMaterial(1);
-	if (AlternateMaterial == nullptr)
-	{
-		AlternateMaterial = UMaterial::GetDefaultMaterial(MD_Surface);
-	}
-
-	MaterialSplitIndex = INDEX_NONE;
-	MaterialRelevance = InComponent->GetMaterialRelevance(GetScene().GetFeatureLevel());
 }
 
 SIZE_T FPaperSpriteSceneProxy::GetTypeHash() const
@@ -95,47 +80,3 @@ void FPaperSpriteSceneProxy::GetDynamicMeshElements(const TArray<const FSceneVie
 
 	FPaperRenderSceneProxy::GetDynamicMeshElements(Views, ViewFamily, VisibilityMap, Collector);
 }
-
-void FPaperSpriteSceneProxy::SetSprite_RenderThread(const FSpriteDrawCallRecord& NewDynamicData, int32 SplitIndex)
-{
-	SCOPE_CYCLE_COUNTER(STAT_PaperRender_SetSpriteRT);
-
-	BatchedSprites.Empty();
-	AlternateBatchedSprites.Empty();
-
-	if (SplitIndex != INDEX_NONE)
-	{
-		FSpriteDrawCallRecord& Record = *new (BatchedSprites) FSpriteDrawCallRecord;
-		FSpriteDrawCallRecord& AltRecord = *new (AlternateBatchedSprites) FSpriteDrawCallRecord;
-		
-		Record.Color = NewDynamicData.Color;
-		Record.Destination = NewDynamicData.Destination;
-		Record.BaseTexture = NewDynamicData.BaseTexture;
-		Record.AdditionalTextures = NewDynamicData.AdditionalTextures;
-		Record.RenderVerts.Append(NewDynamicData.RenderVerts.GetData(), SplitIndex);
-
-		AltRecord.Color = NewDynamicData.Color;
-		AltRecord.Destination = NewDynamicData.Destination;
-		AltRecord.BaseTexture = NewDynamicData.BaseTexture;
-		AltRecord.AdditionalTextures = NewDynamicData.AdditionalTextures;
-		AltRecord.RenderVerts.Append(NewDynamicData.RenderVerts.GetData() + SplitIndex, NewDynamicData.RenderVerts.Num() - SplitIndex);
-	}
-	else
-	{
-		FSpriteDrawCallRecord& Record = *new (BatchedSprites) FSpriteDrawCallRecord;
-		Record = NewDynamicData;
-	}
-}
-
-void FPaperSpriteSceneProxy::GetDynamicMeshElementsForView(const FSceneView* View, int32 ViewIndex, FMeshElementCollector& Collector) const
-{
-	if (Material != nullptr)
-	{
-		GetBatchMesh(View, Material, BatchedSprites, ViewIndex, Collector);
-	}
-
-	if ((AlternateMaterial != nullptr) && (AlternateBatchedSprites.Num() > 0))
-	{
-		GetBatchMesh(View, AlternateMaterial, AlternateBatchedSprites, ViewIndex, Collector);
-	}
-}
diff --git a/Engine/Plugins/2D/Paper2D/Source/Paper2D/Private/PaperSpriteSceneProxy.h b/Engine/Plugins/2D/Paper2D/Source/Paper2D/Private/PaperSpriteSceneProxy.h
--- a/Engine/Plugins/2D/Paper2D/Source/Paper2D/Private/PaperSpriteSceneProxy.h
+++ b/Engine/Plugins/2D/Paper2D/Source/Paper2D/Private/PaperSpriteSceneProxy.h
@@ -13,7 +13,7 @@ class UPaperSpriteComponent;
 //////////////////////////////////////////////////////////////////////////
 // FPaperSpriteSceneProxy
 
-class FPaperSpriteSceneProxy final : public FPaperRenderSceneProxy
+class FPaperSpriteSceneProxy final : public FPaperRenderSceneProxy_SpriteBase
 {
 public:
 	SIZE_T GetTypeHash() const override;
@@ -24,17 +24,6 @@ public:
 	virtual void GetDynamicMeshElements(const TArray<const FSceneView*>& Views, const FSceneViewFamily& ViewFamily, uint32 VisibilityMap, FMeshElementCollector& Collector) const override;
 	// End of FPrimitiveSceneProxy interface
 
-	void SetSprite_RenderThread(const struct FSpriteDrawCallRecord& NewDynamicData, int32 SplitIndex);
-
-protected:
-
-	// FPaperRenderSceneProxy interface
-	virtual void GetDynamicMeshElementsForView(const FSceneView* View, int32 ViewIndex, FMeshElementCollector& Collector) const override;
-	// End of FPaperRenderSceneProxy interface
-
 protected:
-	UMaterialInterface* AlternateMaterial;
-	int32 MaterialSplitIndex;
 	const UBodySetup* BodySetup;
-	TArray<FSpriteDrawCallRecord> AlternateBatchedSprites;
 };
diff --git a/Engine/Plugins/2D/Paper2D/Source/Paper2D/Private/PaperTileMapRenderSceneProxy.cpp b/Engine/Plugins/2D/Paper2D/Source/Paper2D/Private/PaperTileMapRenderSceneProxy.cpp
--- a/Engine/Plugins/2D/Paper2D/Source/Paper2D/Private/PaperTileMapRenderSceneProxy.cpp
+++ b/Engine/Plugins/2D/Paper2D/Source/Paper2D/Private/PaperTileMapRenderSceneProxy.cpp
@@ -32,11 +32,6 @@ FPaperTileMapRenderSceneProxy::FPaperTileMapRenderSceneProxy(const UPaperTileMap
 
 	SetWireframeColor(InComponent->GetWireframeColor());
 	TileMap = InComponent->TileMap;
-	Material = InComponent->GetMaterial(0);
-	if (Material == nullptr)
-	{
-		Material = UMaterial::GetDefaultMaterial(MD_Surface);
-	}
 	MaterialRelevance = InComponent->GetMaterialRelevance(GetScene().GetFeatureLevel());
 
 #if WITH_EDITORONLY_DATA
diff --git a/Engine/Plugins/2D/Paper2D/Source/Paper2D/Private/Terrain/PaperTerrainComponent.cpp b/Engine/Plugins/2D/Paper2D/Source/Paper2D/Private/Terrain/PaperTerrainComponent.cpp
--- a/Engine/Plugins/2D/Paper2D/Source/Paper2D/Private/Terrain/PaperTerrainComponent.cpp
+++ b/Engine/Plugins/2D/Paper2D/Source/Paper2D/Private/Terrain/PaperTerrainComponent.cpp
@@ -19,8 +19,6 @@
 #define PAPER_USE_MATERIAL_SLOPES 1
 #define PAPER_TERRAIN_DRAW_DEBUG 0
 
-DECLARE_CYCLE_STAT(TEXT("Terrain Spline Proxy"), STAT_TerrainSpline_GetDynamicMeshElements, STATGROUP_Paper2D);
-
 //////////////////////////////////////////////////////////////////////////
 
 static FBox2D GetSpriteRenderDataBounds2D(const TArray<FVector4>& Data)
@@ -71,41 +69,32 @@ public:
 	}
 
 	FPaperTerrainSceneProxy(const UPaperTerrainComponent* InComponent, const TArray<FPaperTerrainSpriteGeometry>& InDrawingData);
-
-protected:
-	TArray<FPaperTerrainSpriteGeometry> DrawingData;
-protected:
-	// FPaperRenderSceneProxy interface
-	virtual void GetDynamicMeshElementsForView(const FSceneView* View, int32 ViewIndex, FMeshElementCollector& Collector) const override;
-	// End of FPaperRenderSceneProxy interface
 };
 
 FPaperTerrainSceneProxy::FPaperTerrainSceneProxy(const UPaperTerrainComponent* InComponent, const TArray<FPaperTerrainSpriteGeometry>& InDrawingData)
 	: FPaperRenderSceneProxy(InComponent)
 {
-	DrawingData = InDrawingData;
+	//@TODO: PAPER2D: Build/store the data in the 'new' batch format directly rather than converting it here
+	BatchedSections.AddDefaulted(InDrawingData.Num());
+	FSpriteRenderSection* DstBatch = BatchedSections.GetData();
+	for (const FPaperTerrainSpriteGeometry& SrcBatch : InDrawingData)
+	{
+		DstBatch->Material = SrcBatch.Material;
+		for (const FSpriteDrawCallRecord& SrcRecord : SrcBatch.Records)
+		{
+			DstBatch->AddVerticesFromDrawCallRecord(SrcRecord, 0, SrcRecord.RenderVerts.Num(), Vertices);
+		}
+		++DstBatch;
+	}
 
 	// Combine the material relevance for all materials
-	for (const FPaperTerrainSpriteGeometry& Batch : DrawingData)
+	for (const FPaperTerrainSpriteGeometry& Batch : InDrawingData)
 	{
 		const UMaterialInterface* MaterialInterface = (Batch.Material != nullptr) ? Batch.Material : UMaterial::GetDefaultMaterial(MD_Surface);
 		MaterialRelevance |= MaterialInterface->GetRelevance_Concurrent(GetScene().GetFeatureLevel());
 	}
 }
 
-void FPaperTerrainSceneProxy::GetDynamicMeshElementsForView(const FSceneView* View, int32 ViewIndex, FMeshElementCollector& Collector) const
-{
-	SCOPE_CYCLE_COUNTER(STAT_TerrainSpline_GetDynamicMeshElements);
-
-	for (const FPaperTerrainSpriteGeometry& Batch : DrawingData)
-	{
-		if (Batch.Material != nullptr)
-		{
-			GetBatchMesh(View, Batch.Material, Batch.Records, ViewIndex, Collector);
-		}
-	}
-}
-
 //////////////////////////////////////////////////////////////////////////
 // UPaperTerrainComponent
 
