Removed spurious forward declaration of non-existent type FPaperSpriteVertex
#rnx

diff --git a/Engine/Plugins/2D/Paper2D/Source/Paper2D/Classes/PaperTileMapComponent.h b/Engine/Plugins/2D/Paper2D/Source/Paper2D/Classes/PaperTileMapComponent.h
--- a/Engine/Plugins/2D/Paper2D/Source/Paper2D/Classes/PaperTileMapComponent.h
+++ b/Engine/Plugins/2D/Paper2D/Source/Paper2D/Classes/PaperTileMapComponent.h
@@ -13,7 +13,6 @@ class FPaperTileMapRenderSceneProxy;
 class FPrimitiveSceneProxy;
 class UPaperTileMap;
 class UTexture;
-struct FPaperSpriteVertex;
 struct FSpriteRenderSection;
 
 /**
