Added UMeshComponent forward declaration to fix non-unity error
#rnx

diff --git a/Engine/Plugins/2D/Paper2D/Source/Paper2D/Private/PaperRenderSceneProxy.h b/Engine/Plugins/2D/Paper2D/Source/Paper2D/Private/PaperRenderSceneProxy.h
--- a/Engine/Plugins/2D/Paper2D/Source/Paper2D/Private/PaperRenderSceneProxy.h
+++ b/Engine/Plugins/2D/Paper2D/Source/Paper2D/Private/PaperRenderSceneProxy.h
@@ -18,6 +18,7 @@ class FMeshElementCollector;
 class UBodySetup;
 class UPrimitiveComponent;
 class FSpriteTextureOverrideRenderProxy;
+class UMeshComponent;
 
 #if WITH_EDITOR
 typedef TMap<const UTexture*, const UTexture*> FPaperRenderSceneProxyTextureOverrideMap;
