Paper2D: Rendering optimizations to use prebuilt vertex buffers
- Still always using the dynamic draw path to avoid various edge cases for now, and does a little more copying of vertex data than necessary
- This is controlled by r.Paper2D.UsePrebuiltVertexBuffers, which defaults to 1 and may become the only option in the future

Vertex buffer code adapted from:
PR #6300: Paper2d optimization (Contributed by [REDACTED_PERSON])

diff --git a/Engine/Plugins/2D/Paper2D/Source/Paper2D/Private/PaperRenderSceneProxy.cpp b/Engine/Plugins/2D/Paper2D/Source/Paper2D/Private/PaperRenderSceneProxy.cpp
--- a/Engine/Plugins/2D/Paper2D/Source/Paper2D/Private/PaperRenderSceneProxy.cpp
+++ b/Engine/Plugins/2D/Paper2D/Source/Paper2D/Private/PaperRenderSceneProxy.cpp
@@ -11,6 +11,7 @@
 #include "Engine/Engine.h"
 
 static TAutoConsoleVariable<int32> CVarDrawSpritesAsTwoSided(TEXT("r.Paper2D.DrawTwoSided"), 1, TEXT("Draw sprites as two sided."));
+static TAutoConsoleVariable<int32> CVarDrawSpritesUsingPrebuiltVertexBuffers(TEXT("r.Paper2D.UsePrebuiltVertexBuffers"), 1, TEXT("Draw sprites using prebuilt vertex buffers."));
 
 DECLARE_CYCLE_STAT(TEXT("Get New Batch Meshes"), STAT_PaperRender_GetNewBatchMeshes, STATGROUP_Paper2D);
 DECLARE_CYCLE_STAT(TEXT("SpriteProxy GDME"), STAT_PaperRenderSceneProxy_GetDynamicMeshElements, STATGROUP_Paper2D);
@@ -167,6 +168,7 @@ const FName FSpriteTextureOverrideRenderProxy::AdditionalTextureParameterRootNam
 
 FPaperRenderSceneProxy::FPaperRenderSceneProxy(const UPrimitiveComponent* InComponent)
 	: FPrimitiveSceneProxy(InComponent)
+	, VertexFactory(InComponent->GetWorld()->FeatureLevel)
 	, Owner(InComponent->GetOwner())
 	, MyBodySetup(const_cast<UPrimitiveComponent*>(InComponent)->GetBodySetup())
 	, bCastShadow(InComponent->CastShadow)
@@ -175,6 +177,7 @@ FPaperRenderSceneProxy::FPaperRenderSceneProxy(const UPrimitiveComponent* InComp
 	SetWireframeColor(FLinearColor::White);
 
 	bDrawTwoSided = CVarDrawSpritesAsTwoSided.GetValueOnGameThread() != 0;
+	bSpritesUseVertexBufferPath = CVarDrawSpritesUsingPrebuiltVertexBuffers.GetValueOnGameThread() != 0;
 }
 
 SIZE_T FPaperRenderSceneProxy::GetTypeHash() const
@@ -192,6 +195,8 @@ FPaperRenderSceneProxy::~FPaperRenderSceneProxy()
 			Proxy->ReleasePrimitiveResource();
 		}
 	}
+	VertexBuffer.ReleaseResource();
+	VertexFactory.ReleaseResource();
 }
 
 void FPaperRenderSceneProxy::DebugDrawBodySetup(const FSceneView* View, int32 ViewIndex, FMeshElementCollector& Collector, UBodySetup* BodySetup, const FMatrix& GeomTransformMatrix, const FLinearColor& CollisionColor, bool bDrawSolid) const
@@ -238,6 +243,36 @@ void FPaperRenderSceneProxy::RecreateCachedRenderData()
 		}
 		++BatchIndex;
 	}
+
+	if (bSpritesUseVertexBufferPath && (Vertices.Num() > 0))
+	{
+		VertexBuffer.Vertices = Vertices;
+
+		//We want the proxy to update the buffer, if it has been already initialized
+		if (VertexBuffer.IsInitialized())
+		{
+			const bool bFactoryRequiresReInitialization = VertexBuffer.CommitRequiresBufferRecreation();
+			VertexBuffer.CommitVertexData();
+
+			//When the buffer reallocates, the factory needs to bind the buffers and SRV again, we just init again.
+			if (bFactoryRequiresReInitialization)
+			{
+				VertexFactory.Init(&VertexBuffer);
+			}
+		}
+	}
+}
+
+void FPaperRenderSceneProxy::CreateRenderThreadResources()
+{
+	if (bSpritesUseVertexBufferPath && (Vertices.Num() > 0))
+	{
+		VertexBuffer.Vertices = Vertices;
+
+		// Init the resources
+		VertexBuffer.InitResource();
+		VertexFactory.Init(&VertexBuffer);
+	}
 }
 
 void FPaperRenderSceneProxy::DebugDrawCollision(const FSceneView* View, int32 ViewIndex, FMeshElementCollector& Collector, bool bDrawSolid) const
@@ -305,7 +340,14 @@ void FPaperRenderSceneProxy::GetDynamicMeshElements(const TArray<const FSceneVie
 
 void FPaperRenderSceneProxy::GetDynamicMeshElementsForView(const FSceneView* View, int32 ViewIndex, FMeshElementCollector& Collector) const
 {
-	GetNewBatchMeshes(View, ViewIndex, Collector);
+	if (bSpritesUseVertexBufferPath)
+	{
+		GetNewBatchMeshesPrebuilt(View, ViewIndex, Collector);
+	}
+	else
+	{
+		GetNewBatchMeshes(View, ViewIndex, Collector);
+	}
 }
 
 void FPaperRenderSceneProxy::GetNewBatchMeshes(const FSceneView* View, int32 ViewIndex, FMeshElementCollector& Collector) const
@@ -320,12 +362,7 @@ void FPaperRenderSceneProxy::GetNewBatchMeshes(const FSceneView* View, int32 Vie
 	const uint8 DPG = GetDepthPriorityGroup(View);
 	const bool bIsWireframeView = View->Family->EngineShowFlags.Wireframe;
 
-	if (MaterialTextureOverrideProxies.Num() < BatchedSections.Num())
-	{
-		MaterialTextureOverrideProxies.AddDefaulted(BatchedSections.Num() - MaterialTextureOverrideProxies.Num());
-	}
-	FSpriteTextureOverrideRenderProxy** pCurrentMaterialProxy = MaterialTextureOverrideProxies.GetData();
-
+	int32 SectionIndex = 0;
 	if (Vertices.Num())
 	{
 		FDynamicMeshBuilder DynamicMeshBuilder(View->GetFeatureLevel());
@@ -366,11 +403,7 @@ void FPaperRenderSceneProxy::GetNewBatchMeshes(const FSceneView* View, int32 Vie
 				}
 
 				// Create a texture override material proxy or make sure our existing one is valid
-				if (*pCurrentMaterialProxy == nullptr)
-				{
-					*pCurrentMaterialProxy = new FSpriteTextureOverrideRenderProxy(ParentMaterialProxy, Batch.BaseTexture, Batch.AdditionalTextures UE_EXPAND_IF_WITH_EDITOR(, TextureOverrideList));
-				}
-				(*pCurrentMaterialProxy)->CheckValidity(ParentMaterialProxy);
+				FSpriteTextureOverrideRenderProxy* SectionMaterialProxy = GetCachedMaterialProxyForSection(SectionIndex, ParentMaterialProxy);
 
 				Settings.CastShadow = bCastShadow;
 				Settings.bDisableBackfaceCulling = bDrawTwoSided;
@@ -378,19 +411,89 @@ void FPaperRenderSceneProxy::GetNewBatchMeshes(const FSceneView* View, int32 Vie
 				FDynamicMeshDrawOffset DrawOffset;
 				DrawOffset.FirstIndex = Batch.VertexOffset;
 				DrawOffset.NumPrimitives = Batch.NumVertices / 3;
-				DynamicMeshBuilder.GetMesh(GetLocalToWorld(), *pCurrentMaterialProxy, DPG, Settings, &DrawOffset, ViewIndex, Collector);
+				DynamicMeshBuilder.GetMesh(GetLocalToWorld(), SectionMaterialProxy, DPG, Settings, &DrawOffset, ViewIndex, Collector);
 			}
 
-			++pCurrentMaterialProxy;
+			++SectionIndex;
 		}
 	}
 }
 
-class FPaperVertexArray : public FOneFrameResource
+void FPaperRenderSceneProxy::GetNewBatchMeshesPrebuilt(const FSceneView* View, int32 ViewIndex, FMeshElementCollector& Collector) const
 {
-public:
-	TArray<FDynamicMeshVertex> Vertices;
-};
+	const uint8 DPG = GetDepthPriorityGroup(View);
+	const bool bIsWireframeView = View->Family->EngineShowFlags.Wireframe;
+
+	//Go for each section, creating a batch and collecting it
+	for (int32 SectionIndex = 0; SectionIndex < BatchedSections.Num(); SectionIndex++)
+	{
+		FMeshBatch& Batch = Collector.AllocateMesh();
+
+		if (GetMeshElement(SectionIndex, DPG, IsSelected(), Batch))
+		{
+			// Implementing our own wireframe coloring as the automatic one (controlled by Mesh.bCanApplyViewModeOverrides) only supports per-FPrimitiveSceneProxy WireframeColor
+			if (bIsWireframeView)
+			{
+				const FSpriteRenderSection& Section = BatchedSections[SectionIndex];
+				const FLinearColor EffectiveWireframeColor = (Section.Material->GetBlendMode() != BLEND_Opaque) ? GetWireframeColor() : FLinearColor::Green;
+
+				auto WireframeMaterialInstance = new FColoredMaterialRenderProxy(
+					GEngine->WireframeMaterial->GetRenderProxy(),
+					GetSelectionColor(EffectiveWireframeColor, IsSelected(), IsHovered(), false)
+				);
+
+				Collector.RegisterOneFrameMaterialProxy(WireframeMaterialInstance);
+
+				// We are applying our own wireframe override
+				Batch.bCanApplyViewModeOverrides = false;
+				Batch.bWireframe = true;
+
+				// Create a texture override material proxy and register it as a dynamic resource so that it won't be deleted until the rendering thread has finished with it
+				Batch.MaterialRenderProxy = WireframeMaterialInstance;
+			}
+
+			Collector.AddMesh(ViewIndex, Batch);
+		}
+	}
+}
+
+bool FPaperRenderSceneProxy::GetMeshElement(int32 SectionIndex, uint8 DepthPriorityGroup, bool bIsSelected, FMeshBatch& OutMeshBatch) const
+{
+	check(bSpritesUseVertexBufferPath);
+	check(SectionIndex < BatchedSections.Num());
+
+	const FSpriteRenderSection& Section = BatchedSections[SectionIndex];
+	if (Section.IsValid())
+	{
+		OutMeshBatch.bCanApplyViewModeOverrides = true;
+		OutMeshBatch.bUseWireframeSelectionColoring = bIsSelected;
+
+		OutMeshBatch.LODIndex = 0;
+		OutMeshBatch.VertexFactory = &VertexFactory;
+		OutMeshBatch.LCI = nullptr;
+		OutMeshBatch.ReverseCulling = IsLocalToWorldDeterminantNegative() ? true : false;
+		OutMeshBatch.CastShadow = bCastShadow;
+		OutMeshBatch.DepthPriorityGroup = DepthPriorityGroup;
+		OutMeshBatch.Type = PT_TriangleList;
+		OutMeshBatch.bDisableBackfaceCulling = bDrawTwoSided;
+		OutMeshBatch.MaterialRenderProxy = GetCachedMaterialProxyForSection(SectionIndex, Section.Material->GetRenderProxy());
+
+		// Set up the FMeshBatchElement.
+		FMeshBatchElement& BatchElement = OutMeshBatch.Elements[0];
+		BatchElement.IndexBuffer = VertexBuffer.GetIndexPtr();
+		BatchElement.FirstIndex = Section.VertexOffset;
+		BatchElement.MinVertexIndex = Section.VertexOffset;
+		BatchElement.MaxVertexIndex = Section.VertexOffset + Section.NumVertices;
+		BatchElement.NumPrimitives = Section.NumVertices / 3;
+		BatchElement.VertexFactoryUserData = VertexFactory.GetUniformBuffer();
+
+		return true;
+	}
+	else
+	{
+		return false;
+	}
+}
 
 FPrimitiveViewRelevance FPaperRenderSceneProxy::GetViewRelevance(const FSceneView* View) const
 {
@@ -418,6 +521,7 @@ FPrimitiveViewRelevance FPaperRenderSceneProxy::GetViewRelevance(const FSceneVie
 	const bool bInCollisionView = IsCollisionView(EngineShowFlags, bDrawSimpleCollision, bDrawComplexCollision);
 #endif
 
+	Result.bStaticRelevance = false;
 	Result.bDynamicRelevance = true;
 
 	if (!EngineShowFlags.Materials
@@ -491,6 +595,25 @@ void FPaperRenderSceneProxy::SetTransientTextureOverride_RenderThread(const UTex
 }
 #endif
 
+FSpriteTextureOverrideRenderProxy* FPaperRenderSceneProxy::GetCachedMaterialProxyForSection(int32 SectionIndex, FMaterialRenderProxy* ParentMaterialProxy) const
+{
+	if (MaterialTextureOverrideProxies.Num() < BatchedSections.Num())
+	{
+		MaterialTextureOverrideProxies.AddDefaulted(BatchedSections.Num() - MaterialTextureOverrideProxies.Num());
+	}
+
+	FSpriteTextureOverrideRenderProxy*& Result = MaterialTextureOverrideProxies[SectionIndex];
+	if (Result == nullptr)
+	{
+		const FSpriteRenderSection& Section = BatchedSections[SectionIndex];
+		Result = new FSpriteTextureOverrideRenderProxy(ParentMaterialProxy, Section.BaseTexture, Section.AdditionalTextures UE_EXPAND_IF_WITH_EDITOR(, TextureOverrideList));
+	}
+
+	Result->CheckValidity(ParentMaterialProxy);
+
+	return Result;
+}
+
 //////////////////////////////////////////////////////////////////////////
 // FPaperRenderSceneProxy_SpriteBase
 
diff --git a/Engine/Plugins/2D/Paper2D/Source/Paper2D/Private/PaperRenderSceneProxy.h b/Engine/Plugins/2D/Paper2D/Source/Paper2D/Private/PaperRenderSceneProxy.h
--- a/Engine/Plugins/2D/Paper2D/Source/Paper2D/Private/PaperRenderSceneProxy.h
+++ b/Engine/Plugins/2D/Paper2D/Source/Paper2D/Private/PaperRenderSceneProxy.h
@@ -13,12 +13,14 @@
 #include "LocalVertexFactory.h"
 #include "Paper2DModule.h"
 #include "DynamicMeshBuilder.h"
+#include "PaperSpriteVertexBuffer.h"
 
 class FMeshElementCollector;
 class UBodySetup;
 class UPrimitiveComponent;
 class FSpriteTextureOverrideRenderProxy;
 class UMeshComponent;
+class FMaterialRenderProxy;
 
 #if WITH_EDITOR
 typedef TMap<const UTexture*, const UTexture*> FPaperRenderSceneProxyTextureOverrideMap;
@@ -134,6 +136,7 @@ public:
 	virtual uint32 GetMemoryFootprint() const override;
 	virtual bool CanBeOccluded() const override;
 	virtual bool IsUsingDistanceCullFade() const override;
+	virtual void CreateRenderThreadResources() override;
 	// End of FPrimitiveSceneProxy interface.
 
 	void SetBodySetup_RenderThread(UBodySetup* NewSetup);
@@ -145,7 +148,10 @@ public:
 protected:
 	virtual void GetDynamicMeshElementsForView(const FSceneView* View, int32 ViewIndex, FMeshElementCollector& Collector) const;
 
+	bool GetMeshElement(int32 SectionIndex, uint8 DepthPriorityGroup, bool bIsSelected, FMeshBatch& OutMeshBatch) const;
+
 	void GetNewBatchMeshes(const FSceneView* View, int32 ViewIndex, FMeshElementCollector& Collector) const;
+	void GetNewBatchMeshesPrebuilt(const FSceneView* View, int32 ViewIndex, FMeshElementCollector& Collector) const;
 
 	bool IsCollisionView(const FEngineShowFlags& EngineShowFlags, bool& bDrawSimpleCollision, bool& bDrawComplexCollision) const;
 
@@ -154,17 +160,24 @@ protected:
 
 	// Call this if you modify BatchedSections or Vertices after the proxy has already been created
 	void RecreateCachedRenderData();
+
+	FSpriteTextureOverrideRenderProxy* GetCachedMaterialProxyForSection(int32 SectionIndex, FMaterialRenderProxy* ParentMaterialProxy) const;
+
 protected:
 	TArray<FSpriteRenderSection> BatchedSections;
 	TArray<FDynamicMeshVertex> Vertices;
 	mutable TArray<FSpriteTextureOverrideRenderProxy*> MaterialTextureOverrideProxies;
 
+	FPaperSpriteVertexBuffer VertexBuffer;
+	FPaperSpriteVertexFactory VertexFactory;
+
 	//
 	AActor* Owner;
 	UBodySetup* MyBodySetup;
 
-	bool bDrawTwoSided;
-	bool bCastShadow;
+	uint8 bDrawTwoSided:1;
+	uint8 bCastShadow:1;
+	uint8 bSpritesUseVertexBufferPath:1;
 
 	// The view relevance for the associated material
 	FMaterialRelevance MaterialRelevance;
diff --git a/Engine/Plugins/2D/Paper2D/Source/Paper2D/Private/PaperSpriteVertexBuffer.cpp b/Engine/Plugins/2D/Paper2D/Source/Paper2D/Private/PaperSpriteVertexBuffer.cpp
new file mode 100644
--- /dev/null
+++ b/Engine/Plugins/2D/Paper2D/Source/Paper2D/Private/PaperSpriteVertexBuffer.cpp
@@ -0,0 +1,251 @@
+// Copyright 1998-2019 Epic Games, Inc. All Rights Reserved.
+
+#include "PaperSpriteVertexBuffer.h"
+#include "Materials/Material.h"
+#include "SceneManagement.h"
+#include "PhysicsEngine/BodySetup.h"
+#include "EngineGlobals.h"
+#include "Engine/Engine.h"
+#include "PaperSpriteComponent.h"
+
+//////////////////////////////////////////////////////////////////////////
+// FPaperSpriteVertexBuffer
+
+void FPaperSpriteVertexBuffer::SetDynamicUsage(bool bInDynamicUsage)
+{
+	bDynamicUsage = bInDynamicUsage;
+}
+
+void FPaperSpriteVertexBuffer::CreateBuffers(int32 InNumVertices)
+{
+	//Make sure we don't have dangling buffers
+	if (NumAllocatedVertices > 0)
+	{
+		ReleaseBuffers();
+	}
+
+	//The buffer will always be a shader resource, but they can be static/dynamic depending of the usage
+	uint32 Usage = BUF_ShaderResource | (bDynamicUsage ? BUF_Dynamic : BUF_Static);
+	NumAllocatedVertices = InNumVertices;
+
+	uint32 PositionSize = NumAllocatedVertices * sizeof(FVector);
+	// create vertex buffer
+	{
+		FRHIResourceCreateInfo CreateInfo;
+		PositionBuffer.VertexBufferRHI = RHICreateVertexBuffer(PositionSize, Usage, CreateInfo);
+		if (RHISupportsManualVertexFetch(GMaxRHIShaderPlatform))
+		{
+			PositionBufferSRV = RHICreateShaderResourceView(PositionBuffer.VertexBufferRHI, sizeof(float), PF_R32_FLOAT);
+		}
+
+	}
+
+	uint32 TangentSize = NumAllocatedVertices * 2 * sizeof(FPackedNormal);
+	// create vertex buffer
+	{
+		FRHIResourceCreateInfo CreateInfo;
+		TangentBuffer.VertexBufferRHI = RHICreateVertexBuffer(TangentSize, Usage, CreateInfo);
+		if (RHISupportsManualVertexFetch(GMaxRHIShaderPlatform))
+		{
+			TangentBufferSRV = RHICreateShaderResourceView(TangentBuffer.VertexBufferRHI, sizeof(FPackedNormal), PF_R8G8B8A8_SNORM);
+		}
+	}
+
+	uint32 TexCoordSize = NumAllocatedVertices * sizeof(FVector2D);
+	// create vertex buffer
+	{
+		FRHIResourceCreateInfo CreateInfo;
+		TexCoordBuffer.VertexBufferRHI = RHICreateVertexBuffer(TexCoordSize, Usage, CreateInfo);
+		if (RHISupportsManualVertexFetch(GMaxRHIShaderPlatform))
+		{
+			TexCoordBufferSRV = RHICreateShaderResourceView(TexCoordBuffer.VertexBufferRHI, sizeof(FVector2D), PF_G32R32F);
+		}
+	}
+
+	uint32 ColorSize = NumAllocatedVertices * sizeof(FColor);
+	// create vertex buffer
+	{
+		FRHIResourceCreateInfo CreateInfo;
+		ColorBuffer.VertexBufferRHI = RHICreateVertexBuffer(ColorSize, Usage, CreateInfo);
+		if (RHISupportsManualVertexFetch(GMaxRHIShaderPlatform))
+		{
+			ColorBufferSRV = RHICreateShaderResourceView(ColorBuffer.VertexBufferRHI, sizeof(FColor), PF_R8G8B8A8);
+		}
+	}
+
+	//Create Index Buffer
+	{
+		FRHIResourceCreateInfo CreateInfo;
+		IndexBuffer.IndexBufferRHI = RHICreateIndexBuffer(sizeof(uint32), Vertices.Num() * sizeof(uint32), Usage, CreateInfo);
+	}
+}
+
+void FPaperSpriteVertexBuffer::ReleaseBuffers()
+{
+	PositionBuffer.ReleaseRHI();
+	TangentBuffer.ReleaseRHI();
+	TexCoordBuffer.ReleaseRHI();
+	ColorBuffer.ReleaseRHI();
+	IndexBuffer.ReleaseRHI();
+
+	TangentBufferSRV.SafeRelease();
+	TexCoordBufferSRV.SafeRelease();
+	ColorBufferSRV.SafeRelease();
+	PositionBufferSRV.SafeRelease();
+
+	NumAllocatedVertices = 0;
+}
+
+void FPaperSpriteVertexBuffer::CommitVertexData()
+{
+	if (Vertices.Num())
+	{
+		//Check if we have to accommodate the buffer size
+		if (NumAllocatedVertices != Vertices.Num())
+		{
+			CreateBuffers(Vertices.Num());
+		}
+
+		//Lock vertices
+		FVector* PositionBufferData = nullptr;
+		uint32 PositionSize = Vertices.Num() * sizeof(FVector);
+		{
+			void* Data = RHILockVertexBuffer(PositionBuffer.VertexBufferRHI, 0, PositionSize, RLM_WriteOnly);
+			PositionBufferData = static_cast<FVector*>(Data);
+		}
+
+		FPackedNormal* TangentBufferData = nullptr;
+		uint32 TangentSize = Vertices.Num() * 2 * sizeof(FPackedNormal);
+		{
+			void* Data = RHILockVertexBuffer(TangentBuffer.VertexBufferRHI, 0, TangentSize, RLM_WriteOnly);
+			TangentBufferData = static_cast<FPackedNormal*>(Data);
+		}
+
+		FVector2D* TexCoordBufferData = nullptr;
+		uint32 TexCoordSize = Vertices.Num() * sizeof(FVector2D);
+		{
+			void* Data = RHILockVertexBuffer(TexCoordBuffer.VertexBufferRHI, 0, TexCoordSize, RLM_WriteOnly);
+			TexCoordBufferData = static_cast<FVector2D*>(Data);
+		}
+
+		FColor* ColorBufferData = nullptr;
+		uint32 ColorSize = Vertices.Num() * sizeof(FColor);
+		{
+			void* Data = RHILockVertexBuffer(ColorBuffer.VertexBufferRHI, 0, ColorSize, RLM_WriteOnly);
+			ColorBufferData = static_cast<FColor*>(Data);
+		}
+
+		uint32* IndexBufferData = nullptr;
+		uint32 IndexSize = Vertices.Num() * sizeof(uint32);
+		{
+			void* Data = RHILockIndexBuffer(IndexBuffer.IndexBufferRHI, 0, IndexSize, RLM_WriteOnly);
+			IndexBufferData = static_cast<uint32*>(Data);
+		}
+
+		//Fill verts
+		for (int32 i = 0; i < Vertices.Num(); i++)
+		{
+			PositionBufferData[i] = Vertices[i].Position;
+			TangentBufferData[2 * i + 0] = Vertices[i].TangentX;
+			TangentBufferData[2 * i + 1] = Vertices[i].TangentZ;
+			ColorBufferData[i] = Vertices[i].Color;
+			TexCoordBufferData[i] = Vertices[i].TextureCoordinate[0];
+			IndexBufferData[i] = i;
+		}
+
+		// Unlock the buffer.
+		RHIUnlockVertexBuffer(PositionBuffer.VertexBufferRHI);
+		RHIUnlockVertexBuffer(TangentBuffer.VertexBufferRHI);
+		RHIUnlockVertexBuffer(TexCoordBuffer.VertexBufferRHI);
+		RHIUnlockVertexBuffer(ColorBuffer.VertexBufferRHI);
+		RHIUnlockIndexBuffer(IndexBuffer.IndexBufferRHI);
+
+		//We clear the vertex data, as it isn't needed anymore
+		Vertices.Empty();
+	}
+}
+
+void FPaperSpriteVertexBuffer::InitRHI()
+{
+	//Automatically try to create the data and use it
+	CommitVertexData();
+}
+
+void FPaperSpriteVertexBuffer::ReleaseRHI()
+{
+	PositionBuffer.ReleaseRHI();
+	TangentBuffer.ReleaseRHI();
+	TexCoordBuffer.ReleaseRHI();
+	ColorBuffer.ReleaseRHI();
+	IndexBuffer.ReleaseRHI();
+
+	TangentBufferSRV.SafeRelease();
+	TexCoordBufferSRV.SafeRelease();
+	ColorBufferSRV.SafeRelease();
+	PositionBufferSRV.SafeRelease();
+}
+
+void FPaperSpriteVertexBuffer::InitResource()
+{
+	FRenderResource::InitResource();
+	PositionBuffer.InitResource();
+	TangentBuffer.InitResource();
+	TexCoordBuffer.InitResource();
+	ColorBuffer.InitResource();
+	IndexBuffer.InitResource();
+}
+
+void FPaperSpriteVertexBuffer::ReleaseResource()
+{
+	FRenderResource::ReleaseResource();
+	PositionBuffer.ReleaseResource();
+	TangentBuffer.ReleaseResource();
+	TexCoordBuffer.ReleaseResource();
+	ColorBuffer.ReleaseResource();
+	IndexBuffer.ReleaseResource();
+}
+
+//////////////////////////////////////////////////////////////////////////
+// FPaperSpriteVertexFactory
+
+FPaperSpriteVertexFactory::FPaperSpriteVertexFactory(ERHIFeatureLevel::Type FeatureLevel)
+	: FLocalVertexFactory(FeatureLevel, "FPaperSpriteVertexFactory")
+{
+}
+
+void FPaperSpriteVertexFactory::Init(const FPaperSpriteVertexBuffer* InVertexBuffer)
+{
+	if (IsInRenderingThread())
+	{
+		FLocalVertexFactory::FDataType VertexData;
+		VertexData.NumTexCoords = 1;
+
+		//SRV setup
+		VertexData.LightMapCoordinateIndex = 0;
+		VertexData.TangentsSRV = InVertexBuffer->TangentBufferSRV;
+		VertexData.TextureCoordinatesSRV = InVertexBuffer->TexCoordBufferSRV;
+		VertexData.ColorComponentsSRV = InVertexBuffer->ColorBufferSRV;
+		VertexData.PositionComponentSRV = InVertexBuffer->PositionBufferSRV;
+
+		// Vertex Streams
+		VertexData.PositionComponent = FVertexStreamComponent(&InVertexBuffer->PositionBuffer, 0, sizeof(FVector), VET_Float3, EVertexStreamUsage::Default);
+		VertexData.TangentBasisComponents[0] = FVertexStreamComponent(&InVertexBuffer->TangentBuffer, 0, 2 * sizeof(FPackedNormal), VET_PackedNormal, EVertexStreamUsage::ManualFetch);
+		VertexData.TangentBasisComponents[1] = FVertexStreamComponent(&InVertexBuffer->TangentBuffer, sizeof(FPackedNormal), 2 * sizeof(FPackedNormal), VET_PackedNormal, EVertexStreamUsage::ManualFetch);
+		VertexData.ColorComponent = FVertexStreamComponent(&InVertexBuffer->ColorBuffer, 0, sizeof(FColor), VET_Color, EVertexStreamUsage::ManualFetch);
+		VertexData.TextureCoordinates.Add(FVertexStreamComponent(&InVertexBuffer->TexCoordBuffer, 0, sizeof(FVector2D), VET_Float2, EVertexStreamUsage::ManualFetch));
+
+		SetData(VertexData);
+		VertexBuffer = InVertexBuffer;
+
+		InitResource();
+	}
+	else
+	{
+		FPaperSpriteVertexFactory* ThisFactory = this;
+		ENQUEUE_RENDER_COMMAND(SpriteVertexFactoryInit)(
+			[ThisFactory, InVertexBuffer](FRHICommandListImmediate& RHICmdList)
+		{
+			ThisFactory->Init(InVertexBuffer);
+		});
+	}
+}
diff --git a/Engine/Plugins/2D/Paper2D/Source/Paper2D/Private/PaperSpriteVertexBuffer.h b/Engine/Plugins/2D/Paper2D/Source/Paper2D/Private/PaperSpriteVertexBuffer.h
new file mode 100644
--- /dev/null
+++ b/Engine/Plugins/2D/Paper2D/Source/Paper2D/Private/PaperSpriteVertexBuffer.h
@@ -0,0 +1,100 @@
+// Copyright 1998-2019 Epic Games, Inc. All Rights Reserved.
+
+#pragma once
+
+#include "CoreMinimal.h"
+#include "Engine/EngineTypes.h"
+#include "LocalVertexFactory.h"
+#include "RenderResource.h"
+#include "PackedNormal.h"
+#include "SceneManagement.h"
+#include "DynamicMeshBuilder.h"
+
+//////////////////////////////////////////////////////////////////////////
+//
+
+/** A Paper2D sprite vertex. */
+using FPaperSpriteVertex = FDynamicMeshVertex;
+
+//////////////////////////////////////////////////////////////////////////
+// FPaperSpriteVertexBuffer
+
+/* By default this buffer is static, but can be configured to be dynamic prior to initialization. */
+class FPaperSpriteVertexBuffer : public FVertexBuffer
+{
+
+public:
+	//Buffers
+	FVertexBuffer PositionBuffer;
+	FVertexBuffer TangentBuffer;
+	FVertexBuffer TexCoordBuffer;
+	FVertexBuffer ColorBuffer;
+	FIndexBuffer IndexBuffer;
+
+	//SRVs for Manual Fetch on platforms that support it
+	FShaderResourceViewRHIRef TangentBufferSRV;
+	FShaderResourceViewRHIRef TexCoordBufferSRV;
+	FShaderResourceViewRHIRef ColorBufferSRV;
+	FShaderResourceViewRHIRef PositionBufferSRV;
+
+	//Vertex data
+	TArray<FPaperSpriteVertex> Vertices;
+
+	//Ctor
+	FPaperSpriteVertexBuffer()
+		: bDynamicUsage(true)
+		, NumAllocatedVertices(0)
+	{}
+
+	/* Marks this buffer as dynamic, so it gets initialized as so. */
+	void SetDynamicUsage(bool bInDynamicUsage);
+
+	/* Initializes the buffers with the given number of vertices to accommodate. */
+	void CreateBuffers(int32 NumVertices);
+
+	/* Clear all the buffers currently being used. */
+	void ReleaseBuffers();
+
+	/* Moves all the PaperVertex data onto the RHI buffers. */
+	void CommitVertexData();
+
+	// FRenderResource interface
+	virtual void InitRHI() override;
+	virtual void ReleaseRHI() override;
+	virtual void InitResource() override;
+	virtual void ReleaseResource() override;
+	// End of FRenderResource interface
+
+	/* True if generating a commit would require a reallocation of the buffers. */
+	FORCEINLINE bool CommitRequiresBufferRecreation() const { return NumAllocatedVertices != Vertices.Num(); }
+
+	/* Checks if the buffer has been initialized. */
+	FORCEINLINE bool IsInitialized() const { return NumAllocatedVertices > 0; }
+
+	/* Obtain the index buffer initialized for this buffer. */
+	FORCEINLINE const FIndexBuffer* GetIndexPtr() const { return &IndexBuffer; }
+
+private:
+	/* Indicates if this buffer will be configured for dynamic usage. */
+	bool bDynamicUsage;
+
+	/* Amount of vertices allocated on the vertex buffer. */
+	int32 NumAllocatedVertices;
+};
+
+
+//////////////////////////////////////////////////////////////////////////
+// FPaperSpriteVertexFactory
+
+class FPaperSpriteVertexFactory : public FLocalVertexFactory
+{
+public:
+	FPaperSpriteVertexFactory(ERHIFeatureLevel::Type FeatureLevel);
+
+	/* Initializes this factory with a given vertex buffer. */
+	void Init(const FPaperSpriteVertexBuffer* VertexBuffer);
+
+private:
+	/* Vertex buffer used to initialize this factory. */
+	const FPaperSpriteVertexBuffer* VertexBuffer;
+};
