PR #5687: Fix Paper2D editor viewport mouse operations when using scale > 1 (Contributed by [REDACTED_PERSON])

diff --git a/Engine/Plugins/2D/Paper2D/Source/Paper2DEditor/Private/SPaperEditorViewport.cpp b/Engine/Plugins/2D/Paper2D/Source/Paper2DEditor/Private/SPaperEditorViewport.cpp
--- a/Engine/Plugins/2D/Paper2D/Source/Paper2DEditor/Private/SPaperEditorViewport.cpp
+++ b/Engine/Plugins/2D/Paper2D/Source/Paper2DEditor/Private/SPaperEditorViewport.cpp
@@ -179,7 +179,7 @@ FReply SPaperEditorViewport::OnMouseButtonDown(const FGeometry& MyGeometry, cons
 		ReplyState.CaptureMouse( SharedThis(this) );
 		ReplyState.UseHighPrecisionMouseMovement( SharedThis(this) );
 
-		SoftwareCursorPosition = PanelCoordToGraphCoord( MyGeometry.AbsoluteToLocal( MouseEvent.GetScreenSpacePosition() ) );
+		SoftwareCursorPosition = PanelCoordToGraphCoord( MyGeometry.AbsoluteToLocal( MouseEvent.GetScreenSpacePosition() ) * MyGeometry.Scale );
 
 		// clear any interpolation when you manually pan
 		//DeferredMovementTargetObject = nullptr;
@@ -189,7 +189,7 @@ FReply SPaperEditorViewport::OnMouseButtonDown(const FGeometry& MyGeometry, cons
 	else if (MouseEvent.GetEffectingButton() == EKeys::LeftMouseButton)
 	{
 		// START MARQUEE SELECTION.
-		const FVector2D GraphMousePos = PanelCoordToGraphCoord( MyGeometry.AbsoluteToLocal( MouseEvent.GetScreenSpacePosition() ) );
+		const FVector2D GraphMousePos = PanelCoordToGraphCoord( MyGeometry.AbsoluteToLocal( MouseEvent.GetScreenSpacePosition() ) * MyGeometry.Scale );
 		Marquee.Start( GraphMousePos, FMarqueeOperation::OperationTypeFromMouseEvent(MouseEvent) );
 
 		// Trigger a selection update now so that single-clicks without a drag still select something
@@ -218,7 +218,7 @@ FReply SPaperEditorViewport::OnMouseButtonUp(const FGeometry& MyGeometry, const
 		if (HasMouseCapture())
 		{
 			FSlateRect ThisPanelScreenSpaceRect = MyGeometry.GetLayoutBoundingRect();
-			const FVector2D ScreenSpaceCursorPos = MyGeometry.LocalToAbsolute( GraphCoordToPanelCoord( SoftwareCursorPosition ) );
+			const FVector2D ScreenSpaceCursorPos = MyGeometry.LocalToAbsolute( GraphCoordToPanelCoord( SoftwareCursorPosition ) / MyGeometry.Scale );
 
 			FIntPoint BestPositionInViewport(
 				FMath::RoundToInt( FMath::Clamp( ScreenSpaceCursorPos.X, ThisPanelScreenSpaceRect.Left, ThisPanelScreenSpaceRect.Right ) ),
@@ -335,7 +335,7 @@ FReply SPaperEditorViewport::OnMouseMove(const FGeometry& MyGeometry, const FPoi
 
 			{
 				// We are marquee selecting
-				const FVector2D GraphMousePos = PanelCoordToGraphCoord( MyGeometry.AbsoluteToLocal( MouseEvent.GetScreenSpacePosition() ) );
+				const FVector2D GraphMousePos = PanelCoordToGraphCoord( MyGeometry.AbsoluteToLocal( MouseEvent.GetScreenSpacePosition() ) * MyGeometry.Scale );
 				Marquee.Rect.UpdateEndPoint(GraphMousePos);
 
 				return FReply::Handled();
@@ -349,7 +349,7 @@ FReply SPaperEditorViewport::OnMouseMove(const FGeometry& MyGeometry, const FPoi
 FReply SPaperEditorViewport::OnMouseWheel(const FGeometry& MyGeometry, const FPointerEvent& MouseEvent)
 {
 	// We want to zoom into this point; i.e. keep it the same fraction offset into the panel
-	const FVector2D WidgetSpaceCursorPos = MyGeometry.AbsoluteToLocal( MouseEvent.GetScreenSpacePosition() );
+	const FVector2D WidgetSpaceCursorPos = MyGeometry.AbsoluteToLocal( MouseEvent.GetScreenSpacePosition() ) * MyGeometry.Scale;
 	FVector2D PointToMaintainGraphSpace = PanelCoordToGraphCoord( WidgetSpaceCursorPos );
 
 
@@ -515,7 +515,7 @@ void SPaperEditorViewport::PaintSoftwareCursor(const FGeometry& AllottedGeometry
 		FSlateDrawElement::MakeBox(
 			OutDrawElements,
 			DrawLayerId,
-			AllottedGeometry.ToPaintGeometry( GraphCoordToPanelCoord(SoftwareCursorPosition) - ( Brush->ImageSize / 2 ), Brush->ImageSize ),
+			AllottedGeometry.ToPaintGeometry( GraphCoordToPanelCoord(SoftwareCursorPosition) / AllottedGeometry.Scale - ( Brush->ImageSize / 2 ), Brush->ImageSize ),
 			Brush);
 	}
 }
