Missed one other early access rename

diff --git a/Engine/Plugins/2D/Paper2D/Source/Paper2DEditor/Private/TileMapEditing/TileMapEditor.cpp b/Engine/Plugins/2D/Paper2D/Source/Paper2DEditor/Private/TileMapEditing/TileMapEditor.cpp
--- a/Engine/Plugins/2D/Paper2D/Source/Paper2DEditor/Private/TileMapEditing/TileMapEditor.cpp
+++ b/Engine/Plugins/2D/Paper2D/Source/Paper2DEditor/Private/TileMapEditing/TileMapEditor.cpp
@@ -233,7 +233,7 @@ TSharedRef<SDockTab> FTileMapEditor::SpawnTab_Viewport(const FSpawnTabArgs& Args
 					SNew(STextBlock)
 					.Visibility(EVisibility::HitTestInvisible)
 					.TextStyle(FEditorStyle::Get(), "Graph.CornerText")
-					.Text(LOCTEXT("TileMapEditorViewportEarlyAccessPreviewWarning", "Early access preview"))
+					.Text(LOCTEXT("TileMapEditorViewportEarlyAccessPreviewWarning", "Beta preview"))
 				]
 		];
 }
