Fix to 2DSprites not rendering.

diff --git a/Engine/Plugins/2D/Paper2D/Source/Paper2D/Private/PaperSprite.cpp b/Engine/Plugins/2D/Paper2D/Source/Paper2D/Private/PaperSprite.cpp
--- a/Engine/Plugins/2D/Paper2D/Source/Paper2D/Private/PaperSprite.cpp
+++ b/Engine/Plugins/2D/Paper2D/Source/Paper2D/Private/PaperSprite.cpp
@@ -1905,6 +1905,8 @@ void UPaperSprite::PostLoad()
 	if (AtlasGroup == nullptr)
 	{
 		BakedSourceTexture = GetSourceTexture();
+		BakedSourceUV = SourceUV;
+		BakedSourceDimension = SourceDimension;
 	}
 #endif
 }
