Updated Paper2DEditorModule to check validity of GEditor before registering for OnAssetReimport. This prevents a crash when running a commandlet with IsEditor=false.

diff --git a/Engine/Plugins/2D/Paper2D/Source/Paper2DEditor/Private/Paper2DEditorModule.cpp b/Engine/Plugins/2D/Paper2D/Source/Paper2DEditor/Private/Paper2DEditorModule.cpp
--- a/Engine/Plugins/2D/Paper2D/Source/Paper2DEditor/Private/Paper2DEditorModule.cpp
+++ b/Engine/Plugins/2D/Paper2D/Source/Paper2DEditor/Private/Paper2DEditorModule.cpp
@@ -181,7 +181,10 @@ public:
 		OnPropertyChangedDelegateHandle = FCoreUObjectDelegates::OnObjectPropertyChanged.AddRaw(this, &FPaper2DEditor::OnPropertyChanged);
 
 		// Register to be notified when an asset is reimported
-		OnAssetReimportDelegateHandle = GEditor->GetEditorSubsystem<UImportSubsystem>()->OnAssetReimport.AddRaw(this, &FPaper2DEditor::OnObjectReimported);
+		if (GEditor)
+		{
+			OnAssetReimportDelegateHandle = GEditor->GetEditorSubsystem<UImportSubsystem>()->OnAssetReimport.AddRaw(this, &FPaper2DEditor::OnObjectReimported);
+		}
 
 		// Register the thumbnail renderers
 		UThumbnailManager::Get().RegisterCustomRenderer(UPaperSprite::StaticClass(), UPaperSpriteThumbnailRenderer::StaticClass());
@@ -261,7 +264,10 @@ public:
 			FCoreUObjectDelegates::OnObjectPropertyChanged.Remove(OnPropertyChangedDelegateHandle);
 
 			// Unregister the asset reimport handler
-			GEditor->GetEditorSubsystem<UImportSubsystem>()->OnAssetReimport.Remove(OnAssetReimportDelegateHandle);
+			if (GEditor)
+			{
+				GEditor->GetEditorSubsystem<UImportSubsystem>()->OnAssetReimport.Remove(OnAssetReimportDelegateHandle);
+			}
 		}
 
 		// Unregister the details customization
