Fix for Paper2DSceneProxy crash. It is possible for a FlipBookComponent to be created without any valid vertex data. This causes bad state when RecreateCachedRenderData is called, because the vertex buffer / vertex factory are left uninitialized, which causes an assert downstream during rendering.

diff --git a/Engine/Plugins/2D/Paper2D/Source/Paper2D/Private/PaperRenderSceneProxy.cpp b/Engine/Plugins/2D/Paper2D/Source/Paper2D/Private/PaperRenderSceneProxy.cpp
--- a/Engine/Plugins/2D/Paper2D/Source/Paper2D/Private/PaperRenderSceneProxy.cpp
+++ b/Engine/Plugins/2D/Paper2D/Source/Paper2D/Private/PaperRenderSceneProxy.cpp
@@ -260,6 +260,11 @@ void FPaperRenderSceneProxy::RecreateCachedRenderData()
 				VertexFactory.Init(&VertexBuffer);
 			}
 		}
+		else
+		{
+			VertexBuffer.InitResource();
+			VertexFactory.Init(&VertexBuffer);
+		}
 	}
 }
 
@@ -465,6 +470,8 @@ bool FPaperRenderSceneProxy::GetMeshElement(int32 SectionIndex, uint8 DepthPrior
 	const FSpriteRenderSection& Section = BatchedSections[SectionIndex];
 	if (Section.IsValid())
 	{
+		checkSlow(VertexBuffer.IsInitialized() && VertexFactory.IsInitialized());
+
 		OutMeshBatch.bCanApplyViewModeOverrides = true;
 		OutMeshBatch.bUseWireframeSelectionColoring = bIsSelected;
 
diff --git a/Engine/Source/Runtime/Engine/Private/SceneManagement.cpp b/Engine/Source/Runtime/Engine/Private/SceneManagement.cpp
--- a/Engine/Source/Runtime/Engine/Private/SceneManagement.cpp
+++ b/Engine/Source/Runtime/Engine/Private/SceneManagement.cpp
@@ -243,8 +243,9 @@ void FMeshElementCollector::AddMesh(int32 ViewIndex, FMeshBatch& MeshBatch)
 {
 	DEFINE_LOG_CATEGORY_STATIC(FMeshElementCollector_AddMesh, Warning, All);
 
-	//checkSlow(MeshBatch.GetNumPrimitives() > 0);
-	checkSlow(MeshBatch.VertexFactory && MeshBatch.MaterialRenderProxy);
+	checkSlow(MeshBatch.VertexFactory);
+	checkSlow(MeshBatch.VertexFactory->IsInitialized());
+	checkSlow(MeshBatch.MaterialRenderProxy);
 	checkSlow(PrimitiveSceneProxy);
 
 	PrimitiveSceneProxy->VerifyUsedMaterial(MeshBatch.MaterialRenderProxy);
