Fixed Paper2D EngineTest crash due to resource initialization happening in main thread.

diff --git a/Engine/Plugins/2D/Paper2D/Source/Paper2D/Private/PaperFlipbookComponent.cpp b/Engine/Plugins/2D/Paper2D/Source/Paper2D/Private/PaperFlipbookComponent.cpp
--- a/Engine/Plugins/2D/Paper2D/Source/Paper2D/Private/PaperFlipbookComponent.cpp
+++ b/Engine/Plugins/2D/Paper2D/Source/Paper2D/Private/PaperFlipbookComponent.cpp
@@ -98,7 +98,13 @@ FPrimitiveSceneProxy* UPaperFlipbookComponent::CreateSceneProxy()
 	FSpriteDrawCallRecord DrawCall;
 	DrawCall.BuildFromSprite(SpriteToSend);
 	DrawCall.Color = SpriteColor.ToFColor(/*bSRGB=*/ false);
-	NewProxy->SetSprite_RenderThread(DrawCall, SplitIndex);
+
+	FPaperFlipbookSceneProxy* InSceneProxy = (FPaperFlipbookSceneProxy*)NewProxy;
+	ENQUEUE_RENDER_COMMAND(FCreatePaperFlipbookProxy_SetSprite)(
+		[InSceneProxy, DrawCall, SplitIndex](FRHICommandListImmediate& RHICmdList)
+	{
+		InSceneProxy->SetSprite_RenderThread(DrawCall, SplitIndex);
+	});
 	return NewProxy;
 }
 
diff --git a/Engine/Plugins/2D/Paper2D/Source/Paper2D/Private/PaperRenderSceneProxy.cpp b/Engine/Plugins/2D/Paper2D/Source/Paper2D/Private/PaperRenderSceneProxy.cpp
--- a/Engine/Plugins/2D/Paper2D/Source/Paper2D/Private/PaperRenderSceneProxy.cpp
+++ b/Engine/Plugins/2D/Paper2D/Source/Paper2D/Private/PaperRenderSceneProxy.cpp
@@ -645,6 +645,7 @@ FPaperRenderSceneProxy_SpriteBase::FPaperRenderSceneProxy_SpriteBase(const UMesh
 void FPaperRenderSceneProxy_SpriteBase::SetSprite_RenderThread(const FSpriteDrawCallRecord& NewDynamicData, int32 SplitIndex)
 {
 	SCOPE_CYCLE_COUNTER(STAT_PaperRender_SetSpriteRT);
+	check(IsInRenderingThread());
 
 	BatchedSections.Reset();
 	Vertices.Reset();
diff --git a/Engine/Plugins/2D/Paper2D/Source/Paper2D/Private/PaperSpriteComponent.cpp b/Engine/Plugins/2D/Paper2D/Source/Paper2D/Private/PaperSpriteComponent.cpp
--- a/Engine/Plugins/2D/Paper2D/Source/Paper2D/Private/PaperSpriteComponent.cpp
+++ b/Engine/Plugins/2D/Paper2D/Source/Paper2D/Private/PaperSpriteComponent.cpp
@@ -76,7 +76,13 @@ FPrimitiveSceneProxy* UPaperSpriteComponent::CreateSceneProxy()
 		FSpriteDrawCallRecord DrawCall;
 		DrawCall.BuildFromSprite(SourceSprite);
 		DrawCall.Color = SpriteColor.ToFColor(/*bSRGB=*/ false);
-		NewProxy->SetSprite_RenderThread(DrawCall, SourceSprite->AlternateMaterialSplitIndex);
+
+		FPaperSpriteSceneProxy* InSceneProxy = (FPaperSpriteSceneProxy*)NewProxy;
+		ENQUEUE_RENDER_COMMAND(FCreatePaperSpriteProxy_SetSprite)(
+			[InSceneProxy, DrawCall, SplitIndex = SourceSprite->AlternateMaterialSplitIndex](FRHICommandListImmediate& RHICmdList)
+		{
+			InSceneProxy->SetSprite_RenderThread(DrawCall, SplitIndex);
+		});
 	}
 	return NewProxy;
 }
