Make SpriteEditor use SoftObjectPtr.ToString to find SourceTexture in Asset Registry

diff --git a/Engine/Plugins/2D/Paper2D/Source/Paper2DEditor/Private/SpriteEditor/SSpriteList.cpp b/Engine/Plugins/2D/Paper2D/Source/Paper2DEditor/Private/SpriteEditor/SSpriteList.cpp
--- a/Engine/Plugins/2D/Paper2D/Source/Paper2DEditor/Private/SpriteEditor/SSpriteList.cpp
+++ b/Engine/Plugins/2D/Paper2D/Source/Paper2DEditor/Private/SpriteEditor/SSpriteList.cpp
@@ -33,7 +33,7 @@ void SSpriteList::RebuildWidget(UTexture2D* NewTextureFilter)
 
 	if (NewTextureFilter != nullptr)
 	{
-		const FString TextureString = FAssetData(NewTextureFilter).GetExportTextName();
+		const FString TextureString = TSoftObjectPtr<UTexture2D>(NewTextureFilter).ToString();
 		const FName SourceTexturePropName(TEXT("SourceTexture")); //@TODO: Protected: GET_MEMBER_NAME_CHECKED(UPaperSprite, SourceTexture);
 		Config.Filter.TagsAndValues.Add(SourceTexturePropName, TextureString);
 
diff --git a/Engine/Plugins/2D/Paper2D/Source/Paper2DEditor/Private/SpriteEditor/SpriteEditorViewportClient.cpp b/Engine/Plugins/2D/Paper2D/Source/Paper2DEditor/Private/SpriteEditor/SpriteEditorViewportClient.cpp
--- a/Engine/Plugins/2D/Paper2D/Source/Paper2DEditor/Private/SpriteEditor/SpriteEditorViewportClient.cpp
+++ b/Engine/Plugins/2D/Paper2D/Source/Paper2DEditor/Private/SpriteEditor/SpriteEditorViewportClient.cpp
@@ -693,7 +693,7 @@ void FSpriteEditorViewportClient::UpdateRelatedSpritesList()
 	{
 		FARFilter Filter;
 		Filter.ClassNames.Add(UPaperSprite::StaticClass()->GetFName());
-		const FString TextureString = FAssetData(Texture).GetExportTextName();
+		const FString TextureString = TSoftObjectPtr<UTexture2D>(Texture).ToString();
 		const FName SourceTexturePropName(TEXT("SourceTexture"));
 		Filter.TagsAndValues.Add(SourceTexturePropName, TextureString);
 		FAssetRegistryModule& AssetRegistryModule = FModuleManager::LoadModuleChecked<FAssetRegistryModule>("AssetRegistry");
