Fix the default settings on Tile Map editor preferences to actually use the specified defaults per user.

diff --git a/Engine/Plugins/2D/Paper2D/Source/Paper2DEditor/Private/PaperImporterSettings.cpp b/Engine/Plugins/2D/Paper2D/Source/Paper2DEditor/Private/PaperImporterSettings.cpp
--- a/Engine/Plugins/2D/Paper2D/Source/Paper2DEditor/Private/PaperImporterSettings.cpp
+++ b/Engine/Plugins/2D/Paper2D/Source/Paper2DEditor/Private/PaperImporterSettings.cpp
@@ -146,7 +146,22 @@ void UPaperImporterSettings::ApplySettingsForTileMapInit(UPaperTileMap* TileMap,
 	}
 
 	TileMap->PixelsPerUnrealUnit = DefaultPixelsPerUnrealUnit;
-	TileMap->BackgroundColor = GetDefault<UTileMapEditorSettings>()->DefaultBackgroundColor;
+
+	// Set Defaults based on the editor preferences 
+	const UTileMapEditorSettings* Settings = GetMutableDefault<UTileMapEditorSettings>();
+	if (!Settings)
+	{
+		Settings = GetDefault<UTileMapEditorSettings>();
+	}
+
+	TileMap->BackgroundColor = Settings->DefaultBackgroundColor;
+	TileMap->TileGridColor = Settings->DefaultTileGridColor;
+	TileMap->MultiTileGridColor = Settings->DefaultMultiTileGridColor;
+	TileMap->MultiTileGridWidth = Settings->DefaultMultiTileGridWidth;
+	TileMap->MultiTileGridHeight = Settings->DefaultMultiTileGridHeight;
+	TileMap->MultiTileGridOffsetX = Settings->DefaultMultiTileGridOffsetX;
+	TileMap->MultiTileGridOffsetY = Settings->DefaultMultiTileGridOffsetY;
+	TileMap->LayerGridColor = Settings->DefaultLayerGridColor;
 
 	ESpriteInitMaterialType DesiredMaterialType = MaterialTypeMode;
 	if (DesiredMaterialType == ESpriteInitMaterialType::Automatic)
