Paper2D: Fixed sprite normal generation regression so lit sprites work as expected
PR #6493: Paper2d bugfixes: UE-85752 (Contributed by [REDACTED_PERSON])

diff --git a/Engine/Plugins/2D/Paper2D/Source/Paper2D/Private/PaperRenderSceneProxy.h b/Engine/Plugins/2D/Paper2D/Source/Paper2D/Private/PaperRenderSceneProxy.h
--- a/Engine/Plugins/2D/Paper2D/Source/Paper2D/Private/PaperRenderSceneProxy.h
+++ b/Engine/Plugins/2D/Paper2D/Source/Paper2D/Private/PaperRenderSceneProxy.h
@@ -96,7 +96,7 @@ struct PAPER2D_API FSpriteRenderSection
 			const FVector Pos((PaperAxisX * SourceVert.X) + (PaperAxisY * SourceVert.Y) + Record.Destination);
 			const FVector2D UV(SourceVert.Z, SourceVert.W);
 
-			new (Vertices) FDynamicMeshVertex(Pos, UV, VertColor);
+			new (Vertices) FDynamicMeshVertex(Pos, FPaperSpriteTangents::PackedNormalX.ToFVector(), FPaperSpriteTangents::PackedNormalZ.ToFVector(), UV, VertColor);
 		}
 	}
 
@@ -105,7 +105,7 @@ struct PAPER2D_API FSpriteRenderSection
 	{
 		const FVector Pos((PaperAxisX * X) + (PaperAxisY * Y) + Origin);
 
-		new (Vertices) FDynamicMeshVertex(Pos, FVector2D(U, V), Color);
+		new (Vertices) FDynamicMeshVertex(Pos, FPaperSpriteTangents::PackedNormalX.ToFVector(), FPaperSpriteTangents::PackedNormalZ.ToFVector(), FVector2D(U, V), Color);
 		++NumVertices;
 	}
 
