Copy from Dev-Editor
Re-add old tag for SpriteEditor to cache assets that haven't been saved in a few years

diff --git a/Engine/Plugins/2D/Paper2D/Source/Paper2DEditor/Private/SpriteEditor/SSpriteList.cpp b/Engine/Plugins/2D/Paper2D/Source/Paper2DEditor/Private/SpriteEditor/SSpriteList.cpp
--- a/Engine/Plugins/2D/Paper2D/Source/Paper2DEditor/Private/SpriteEditor/SSpriteList.cpp
+++ b/Engine/Plugins/2D/Paper2D/Source/Paper2DEditor/Private/SpriteEditor/SSpriteList.cpp
@@ -33,9 +33,10 @@ void SSpriteList::RebuildWidget(UTexture2D* NewTextureFilter)
 
 	if (NewTextureFilter != nullptr)
 	{
-		const FString TextureString = TSoftObjectPtr<UTexture2D>(NewTextureFilter).ToString();
 		const FName SourceTexturePropName(TEXT("SourceTexture")); //@TODO: Protected: GET_MEMBER_NAME_CHECKED(UPaperSprite, SourceTexture);
-		Config.Filter.TagsAndValues.Add(SourceTexturePropName, TextureString);
+		Config.Filter.TagsAndValues.Add(SourceTexturePropName, TSoftObjectPtr<UTexture2D>(NewTextureFilter).ToString());
+		// Legacy format for files that haven't been resaved
+		Config.Filter.TagsAndValues.Add(SourceTexturePropName, FAssetData(NewTextureFilter).GetExportTextName());
 
 		// Ignore the source texture since we're filtering against it
 		AssetRegistryTagsToIgnore.Add(SourceTexturePropName);
diff --git a/Engine/Plugins/2D/Paper2D/Source/Paper2DEditor/Private/SpriteEditor/SpriteEditorViewportClient.cpp b/Engine/Plugins/2D/Paper2D/Source/Paper2DEditor/Private/SpriteEditor/SpriteEditorViewportClient.cpp
--- a/Engine/Plugins/2D/Paper2D/Source/Paper2DEditor/Private/SpriteEditor/SpriteEditorViewportClient.cpp
+++ b/Engine/Plugins/2D/Paper2D/Source/Paper2DEditor/Private/SpriteEditor/SpriteEditorViewportClient.cpp
@@ -693,9 +693,11 @@ void FSpriteEditorViewportClient::UpdateRelatedSpritesList()
 	{
 		FARFilter Filter;
 		Filter.ClassNames.Add(UPaperSprite::StaticClass()->GetFName());
-		const FString TextureString = TSoftObjectPtr<UTexture2D>(Texture).ToString();
 		const FName SourceTexturePropName(TEXT("SourceTexture"));
-		Filter.TagsAndValues.Add(SourceTexturePropName, TextureString);
+		Filter.TagsAndValues.Add(SourceTexturePropName, TSoftObjectPtr<UTexture2D>(Texture).ToString());
+		// Legacy format for files that haven't been resaved
+		Filter.TagsAndValues.Add(SourceTexturePropName, FAssetData(Texture).GetExportTextName());
+
 		FAssetRegistryModule& AssetRegistryModule = FModuleManager::LoadModuleChecked<FAssetRegistryModule>("AssetRegistry");
 		TArray<FAssetData> SpriteAssetData;
 		AssetRegistryModule.Get().GetAssets(Filter, SpriteAssetData);
