UE4.25 - Fix ensure
#rnx

diff --git a/Engine/Plugins/2D/Paper2D/Source/Paper2D/Private/PaperRenderSceneProxy.cpp b/Engine/Plugins/2D/Paper2D/Source/Paper2D/Private/PaperRenderSceneProxy.cpp
--- a/Engine/Plugins/2D/Paper2D/Source/Paper2D/Private/PaperRenderSceneProxy.cpp
+++ b/Engine/Plugins/2D/Paper2D/Source/Paper2D/Private/PaperRenderSceneProxy.cpp
@@ -178,8 +178,8 @@ FPaperRenderSceneProxy::FPaperRenderSceneProxy(const UPrimitiveComponent* InComp
 {
 	SetWireframeColor(FLinearColor::White);
 
-	bDrawTwoSided = CVarDrawSpritesAsTwoSided.GetValueOnGameThread() != 0;
-	bSpritesUseVertexBufferPath = CVarDrawSpritesUsingPrebuiltVertexBuffers.GetValueOnGameThread() != 0;
+	bDrawTwoSided = CVarDrawSpritesAsTwoSided.GetValueOnAnyThread() != 0;
+	bSpritesUseVertexBufferPath = CVarDrawSpritesUsingPrebuiltVertexBuffers.GetValueOnAnyThread() != 0;
 }
 
 SIZE_T FPaperRenderSceneProxy::GetTypeHash() const
