Paper2D doesn't include json in public dependencies, but uses it in public headers.

diff --git a/Engine/Plugins/2D/Paper2D/Source/Paper2DEditor/Paper2DEditor.Build.cs b/Engine/Plugins/2D/Paper2D/Source/Paper2DEditor/Paper2DEditor.Build.cs
--- a/Engine/Plugins/2D/Paper2D/Source/Paper2DEditor/Paper2DEditor.Build.cs
+++ b/Engine/Plugins/2D/Paper2D/Source/Paper2DEditor/Paper2DEditor.Build.cs
@@ -13,7 +13,6 @@ public class Paper2DEditor : ModuleRules
 				"Core",
 				"CoreUObject",
 				"ApplicationCore",
-				"Json",
 				"Slate",
 				"SlateCore",
 				"Engine",
@@ -35,6 +34,11 @@ public class Paper2DEditor : ModuleRules
                 "IntroTutorials"
             });
 
+		PublicDependencyModuleNames.AddRange(
+			new string[] {
+				"Json",
+			});
+
 		PrivateIncludePathModuleNames.AddRange(
 			new string[] {
 				"Settings",
