Move paper2d icons and style setup to the paper2d style

diff --git a/Engine/Plugins/2D/Paper2D/Source/Paper2DEditor/Private/PaperStyle.cpp b/Engine/Plugins/2D/Paper2D/Source/Paper2DEditor/Private/PaperStyle.cpp
--- a/Engine/Plugins/2D/Paper2D/Source/Paper2DEditor/Private/PaperStyle.cpp
+++ b/Engine/Plugins/2D/Paper2D/Source/Paper2DEditor/Private/PaperStyle.cpp
@@ -38,6 +38,7 @@ void FPaperStyle::Initialize()
 	const FVector2D Icon16x16(16.0f, 16.0f);
 	const FVector2D Icon20x20(20.0f, 20.0f);
 	const FVector2D Icon40x40(40.0f, 40.0f);
+	const FVector2D Icon64x64(64.0f, 64.0f);
 
 	// Only register once
 	if (StyleSet.IsValid())
@@ -259,6 +260,57 @@ void FPaperStyle::Initialize()
 		StyleSet->Set("AssetActions.CreateTileSet", new IMAGE_PLUGIN_BRUSH(TEXT("Icons/icon_Texture_CreateTileSet_16x"), Icon16x16));
 	}
 
+	{
+		// Sprites (asset, component, actor)
+		StyleSet->Set("ClassIcon.PaperSprite", new IMAGE_PLUGIN_BRUSH("Icons/AssetIcons/PaperSprite_16x", Icon16x16));
+		StyleSet->Set("ClassThumbnail.PaperSprite", new IMAGE_PLUGIN_BRUSH("Icons/AssetIcons/PaperSprite_64x", Icon64x64));
+		StyleSet->Set("ClassIcon.PaperSpriteComponent", new IMAGE_PLUGIN_BRUSH("Icons/AssetIcons/PaperSpriteComponent_16x", Icon16x16));
+		StyleSet->Set("ClassThumbnail.PaperSpriteComponent", new IMAGE_PLUGIN_BRUSH("Icons/AssetIcons/PaperSpriteComponent_64x", Icon64x64));
+		StyleSet->Set("ClassIcon.PaperSpriteActor", new IMAGE_PLUGIN_BRUSH("Icons/AssetIcons/PaperSpriteActor_16x", Icon16x16));
+		StyleSet->Set("ClassThumbnail.PaperSpriteActor", new IMAGE_PLUGIN_BRUSH("Icons/AssetIcons/PaperSpriteActor_64x", Icon64x64));
+
+		// Flipbooks (asset, component, actor)
+		StyleSet->Set("ClassIcon.PaperFlipbook", new IMAGE_PLUGIN_BRUSH("Icons/AssetIcons/PaperFlipbook_16x", Icon16x16));
+		StyleSet->Set("ClassThumbnail.PaperFlipbook", new IMAGE_PLUGIN_BRUSH("Icons/AssetIcons/PaperFlipbook_64x", Icon64x64));
+		StyleSet->Set("ClassIcon.PaperFlipbookComponent", new IMAGE_PLUGIN_BRUSH("Icons/AssetIcons/PaperFlipbookComponent_16x", Icon16x16));
+		StyleSet->Set("ClassThumbnail.PaperFlipbookComponent", new IMAGE_PLUGIN_BRUSH("Icons/AssetIcons/PaperFlipbookComponent_64x", Icon64x64));
+		StyleSet->Set("ClassIcon.PaperFlipbookActor", new IMAGE_PLUGIN_BRUSH("Icons/AssetIcons/PaperFlipbookActor_16x", Icon16x16));
+		StyleSet->Set("ClassThumbnail.PaperFlipbookActor", new IMAGE_PLUGIN_BRUSH("Icons/AssetIcons/PaperFlipbookActor_64x", Icon64x64));
+
+		// Tile maps (asset, component, actor)
+		StyleSet->Set("ClassIcon.PaperTileMap", new IMAGE_PLUGIN_BRUSH("Icons/AssetIcons/PaperTileMap_16x", Icon16x16));
+		StyleSet->Set("ClassThumbnail.PaperTileMap", new IMAGE_PLUGIN_BRUSH("Icons/AssetIcons/PaperTileMap_64x", Icon64x64));
+		StyleSet->Set("ClassIcon.PaperTileMapComponent", new IMAGE_PLUGIN_BRUSH("Icons/AssetIcons/PaperTileMapComponent_16x", Icon16x16));
+		StyleSet->Set("ClassThumbnail.PaperTileMapComponent", new IMAGE_PLUGIN_BRUSH("Icons/AssetIcons/PaperTileMapComponent_64x", Icon64x64));
+		StyleSet->Set("ClassIcon.PaperTileMapActor", new IMAGE_PLUGIN_BRUSH("Icons/AssetIcons/PaperTileMapActor_16x", Icon16x16));
+		StyleSet->Set("ClassThumbnail.PaperTileMapActor", new IMAGE_PLUGIN_BRUSH("Icons/AssetIcons/PaperTileMapActor_64x", Icon64x64));
+
+		// UPaperSpriteAtlas assets
+		//@TODO: Paper2D: These icons don't match the naming scheme
+		StyleSet->Set("ClassIcon.PaperSpriteAtlas", new IMAGE_PLUGIN_BRUSH("Icons/AssetIcons/Paper2DSpriteAtlasGroup_16x", Icon16x16));
+		StyleSet->Set("ClassThumbnail.PaperSpriteAtlas", new IMAGE_PLUGIN_BRUSH("Icons/AssetIcons/Paper2DSpriteAtlasGroup_64x", Icon64x64));
+
+		//@TODO: UPaperSpriteSheet icons?
+
+		// APaperCharacter icons
+		StyleSet->Set("ClassIcon.PaperCharacter", new IMAGE_PLUGIN_BRUSH("Icons/AssetIcons/PaperCharacter_16x", Icon16x16));
+		StyleSet->Set("ClassThumbnail.PaperCharacter", new IMAGE_PLUGIN_BRUSH("Icons/AssetIcons/PaperCharacter_64x", Icon64x64));
+
+		// UPaperTileSet icons
+		StyleSet->Set("ClassIcon.PaperTileSet", new IMAGE_PLUGIN_BRUSH("Icons/AssetIcons/PaperTileSet_16x", Icon16x16));
+		StyleSet->Set("ClassThumbnail.PaperTileSet", new IMAGE_PLUGIN_BRUSH("Icons/AssetIcons/PaperTileSet_64x", Icon64x64));
+
+		// UPaperTerrainMaterial icons
+		StyleSet->Set("ClassIcon.PaperTerrainMaterial", new IMAGE_PLUGIN_BRUSH("Icons/AssetIcons/PaperTerrainMaterial_16x", Icon16x16));
+		StyleSet->Set("ClassThumbnail.PaperTerrainMaterial", new IMAGE_PLUGIN_BRUSH("Icons/AssetIcons/PaperTerrainMaterial_64x", Icon64x64));
+
+		// Terrain splines (component, actor)
+		StyleSet->Set("ClassIcon.PaperTerrainComponent", new IMAGE_PLUGIN_BRUSH("Icons/AssetIcons/PaperTerrainComponent_16x", Icon16x16));
+		StyleSet->Set("ClassThumbnail.PaperTerrainComponent", new IMAGE_PLUGIN_BRUSH("Icons/AssetIcons/PaperTerrainComponent_64x", Icon64x64));
+		StyleSet->Set("ClassIcon.PaperTerrainActor", new IMAGE_PLUGIN_BRUSH("Icons/AssetIcons/PaperTerrainActor_16x", Icon16x16));
+		StyleSet->Set("ClassThumbnail.PaperTerrainActor", new IMAGE_PLUGIN_BRUSH("Icons/AssetIcons/PaperTerrainActor_16x", Icon64x64));
+	}
+
 	FSlateStyleRegistry::RegisterSlateStyle(*StyleSet.Get());
 };
 
diff --git a/Engine/Source/Editor/EditorStyle/Private/StarshipStyle.cpp b/Engine/Source/Editor/EditorStyle/Private/StarshipStyle.cpp
--- a/Engine/Source/Editor/EditorStyle/Private/StarshipStyle.cpp
+++ b/Engine/Source/Editor/EditorStyle/Private/StarshipStyle.cpp
@@ -5752,58 +5752,6 @@ void FStarshipEditorStyle::FStyle::SetupClassIconsAndThumbnails()
 		Set("ClassIcon.ComponentMobilityHeaderIcon", new IMAGE_BRUSH("Icons/ActorIcons/ComponentMobilityHeader_7x16", Icon7x16));
 		Set("ClassIcon.AnimInstance", new IMAGE_BRUSH("Icons/AssetIcons/AnimInstance_16x", Icon16x16));
 
-		//@TODO: PAPER2D: Defined here until it is possible to define these in a plugin
-		{
-			// Sprites (asset, component, actor)
-			Set("ClassIcon.PaperSprite", new IMAGE_BRUSH("Icons/AssetIcons/PaperSprite_16x", Icon16x16));
-			Set("ClassThumbnail.PaperSprite", new IMAGE_BRUSH("Icons/AssetIcons/PaperSprite_64x", Icon64x64));
-			Set("ClassIcon.PaperSpriteComponent", new IMAGE_BRUSH("Icons/AssetIcons/PaperSpriteComponent_16x", Icon16x16));
-			Set("ClassThumbnail.PaperSpriteComponent", new IMAGE_BRUSH("Icons/AssetIcons/PaperSpriteComponent_64x", Icon64x64));
-			Set("ClassIcon.PaperSpriteActor", new IMAGE_BRUSH("Icons/AssetIcons/PaperSpriteActor_16x", Icon16x16));
-			Set("ClassThumbnail.PaperSpriteActor", new IMAGE_BRUSH("Icons/AssetIcons/PaperSpriteActor_64x", Icon64x64));
-
-			// Flipbooks (asset, component, actor)
-			Set("ClassIcon.PaperFlipbook", new IMAGE_BRUSH("Icons/AssetIcons/PaperFlipbook_16x", Icon16x16));
-			Set("ClassThumbnail.PaperFlipbook", new IMAGE_BRUSH("Icons/AssetIcons/PaperFlipbook_64x", Icon64x64));
-			Set("ClassIcon.PaperFlipbookComponent", new IMAGE_BRUSH("Icons/AssetIcons/PaperFlipbookComponent_16x", Icon16x16));
-			Set("ClassThumbnail.PaperFlipbookComponent", new IMAGE_BRUSH("Icons/AssetIcons/PaperFlipbookComponent_64x", Icon64x64));
-			Set("ClassIcon.PaperFlipbookActor", new IMAGE_BRUSH("Icons/AssetIcons/PaperFlipbookActor_16x", Icon16x16));
-			Set("ClassThumbnail.PaperFlipbookActor", new IMAGE_BRUSH("Icons/AssetIcons/PaperFlipbookActor_64x", Icon64x64));
-
-			// Tile maps (asset, component, actor)
-			Set("ClassIcon.PaperTileMap", new IMAGE_BRUSH("Icons/AssetIcons/PaperTileMap_16x", Icon16x16));
-			Set("ClassThumbnail.PaperTileMap", new IMAGE_BRUSH("Icons/AssetIcons/PaperTileMap_64x", Icon64x64));
-			Set("ClassIcon.PaperTileMapComponent", new IMAGE_BRUSH("Icons/AssetIcons/PaperTileMapComponent_16x", Icon16x16));
-			Set("ClassThumbnail.PaperTileMapComponent", new IMAGE_BRUSH("Icons/AssetIcons/PaperTileMapComponent_64x", Icon64x64));
-			Set("ClassIcon.PaperTileMapActor", new IMAGE_BRUSH("Icons/AssetIcons/PaperTileMapActor_16x", Icon16x16));
-			Set("ClassThumbnail.PaperTileMapActor", new IMAGE_BRUSH("Icons/AssetIcons/PaperTileMapActor_64x", Icon64x64));
-
-			// UPaperSpriteAtlas assets
-			//@TODO: Paper2D: These icons don't match the naming scheme
-			Set("ClassIcon.PaperSpriteAtlas", new IMAGE_BRUSH("Icons/AssetIcons/Paper2DSpriteAtlasGroup_16x", Icon16x16));
-			Set("ClassThumbnail.PaperSpriteAtlas", new IMAGE_BRUSH("Icons/AssetIcons/Paper2DSpriteAtlasGroup_64x", Icon64x64));
-
-			//@TODO: UPaperSpriteSheet icons?
-
-			// APaperCharacter icons
-			Set("ClassIcon.PaperCharacter", new IMAGE_BRUSH("Icons/AssetIcons/PaperCharacter_16x", Icon16x16));
-			Set("ClassThumbnail.PaperCharacter", new IMAGE_BRUSH("Icons/AssetIcons/PaperCharacter_64x", Icon64x64));
-
-			// UPaperTileSet icons
-			Set("ClassIcon.PaperTileSet", new IMAGE_BRUSH("Icons/AssetIcons/PaperTileSet_16x", Icon16x16));
-			Set("ClassThumbnail.PaperTileSet", new IMAGE_BRUSH("Icons/AssetIcons/PaperTileSet_64x", Icon64x64));
-
-			// UPaperTerrainMaterial icons
-			Set("ClassIcon.PaperTerrainMaterial", new IMAGE_BRUSH("Icons/AssetIcons/PaperTerrainMaterial_16x", Icon16x16));
-			Set("ClassThumbnail.PaperTerrainMaterial", new IMAGE_BRUSH("Icons/AssetIcons/PaperTerrainMaterial_64x", Icon64x64));
-
-			// Terrain splines (component, actor)
-			Set("ClassIcon.PaperTerrainComponent", new IMAGE_BRUSH("Icons/AssetIcons/PaperTerrainComponent_16x", Icon16x16));
-			Set("ClassThumbnail.PaperTerrainComponent", new IMAGE_BRUSH("Icons/AssetIcons/PaperTerrainComponent_64x", Icon64x64));
-			Set("ClassIcon.PaperTerrainActor", new IMAGE_BRUSH("Icons/AssetIcons/PaperTerrainActor_16x", Icon16x16));
-			Set("ClassThumbnail.PaperTerrainActor", new IMAGE_BRUSH("Icons/AssetIcons/PaperTerrainActor_16x", Icon64x64));
-		}
-
 		// Factory classes
 		Set( "ClassIcon.ActorFactoryBoxVolume", new IMAGE_BRUSH( "Icons/icon_volume_Box_16x", Icon16x16 ) );
 		Set( "ClassIcon.ActorFactoryCylinderVolume", new IMAGE_BRUSH( "Icons/icon_volume_cylinder_16x", Icon16x16 ) );
