Fixing crashes in automated tests for Paper2D.

# This is from the toolkit mode manager ownership changes, and is a stop gap while we keep moving the mode manager out of the viewport client's ownership.

#rnx

diff --git a/Engine/Plugins/2D/Paper2D/Source/Paper2DEditor/Private/PaperEditorViewportClient.cpp b/Engine/Plugins/2D/Paper2D/Source/Paper2DEditor/Private/PaperEditorViewportClient.cpp
--- a/Engine/Plugins/2D/Paper2D/Source/Paper2DEditor/Private/PaperEditorViewportClient.cpp
+++ b/Engine/Plugins/2D/Paper2D/Source/Paper2DEditor/Private/PaperEditorViewportClient.cpp
@@ -17,7 +17,6 @@ FPaperEditorViewportClient::FPaperEditorViewportClient(const TWeakPtr<SEditorVie
 	ZoomAmount = 1.0f;
 
 	//ModifyCheckerboardTextureColors();
-	//@TODO: ModeTools->SetToolkitHost
 
 	//@TODO: Pretty lame hardcoding
 	//@TODO: Doesn't handle negatives either (not really)
diff --git a/Engine/Plugins/2D/Paper2D/Source/Paper2DEditor/Private/SpriteEditor/SpriteEditor.cpp b/Engine/Plugins/2D/Paper2D/Source/Paper2DEditor/Private/SpriteEditor/SpriteEditor.cpp
--- a/Engine/Plugins/2D/Paper2D/Source/Paper2DEditor/Private/SpriteEditor/SpriteEditor.cpp
+++ b/Engine/Plugins/2D/Paper2D/Source/Paper2DEditor/Private/SpriteEditor/SpriteEditor.cpp
@@ -24,6 +24,7 @@
 #include "Widgets/Docking/SDockTab.h"
 
 #include "ExtractSprites/SPaperExtractSpritesDialog.h"
+#include "EditorModeManager.h"
 
 #define LOCTEXT_NAMESPACE "SpriteEditor"
 
@@ -571,6 +572,16 @@ void FSpriteEditor::ExtendToolbar()
 	AddToolbarExtender(Paper2DEditorModule->GetSpriteEditorToolBarExtensibilityManager()->GetAllExtenders());
 }
 
+void FSpriteEditor::CreateEditorModeManager()
+{
+	check(ViewportPtr.IsValid());
+	TSharedPtr<FEditorViewportClient> ViewportClient = ViewportPtr->GetViewportClient();
+	check(ViewportClient.IsValid());
+	PRAGMA_DISABLE_DEPRECATION_WARNINGS
+	ViewportClient->TakeOwnershipOfModeManager(EditorModeManager);
+	PRAGMA_ENABLE_DEPRECATION_WARNINGS
+}
+
 void FSpriteEditor::SetSpriteBeingEdited(UPaperSprite* NewSprite)
 {
 	if ((NewSprite != SpriteBeingEdited) && (NewSprite != nullptr))
diff --git a/Engine/Plugins/2D/Paper2D/Source/Paper2DEditor/Private/SpriteEditor/SpriteEditor.h b/Engine/Plugins/2D/Paper2D/Source/Paper2DEditor/Private/SpriteEditor/SpriteEditor.h
--- a/Engine/Plugins/2D/Paper2D/Source/Paper2DEditor/Private/SpriteEditor/SpriteEditor.h
+++ b/Engine/Plugins/2D/Paper2D/Source/Paper2DEditor/Private/SpriteEditor/SpriteEditor.h
@@ -75,6 +75,7 @@ protected:
 	void BindCommands();
 	void ExtendMenu();
 	void ExtendToolbar();
+	virtual void CreateEditorModeManager() override;
 
 	TSharedRef<SDockTab> SpawnTab_Viewport(const FSpawnTabArgs& Args);
 	TSharedRef<SDockTab> SpawnTab_Details(const FSpawnTabArgs& Args);
diff --git a/Engine/Plugins/2D/Paper2D/Source/Paper2DEditor/Private/SpriteEditor/SpriteEditorViewportClient.cpp b/Engine/Plugins/2D/Paper2D/Source/Paper2DEditor/Private/SpriteEditor/SpriteEditorViewportClient.cpp
--- a/Engine/Plugins/2D/Paper2D/Source/Paper2DEditor/Private/SpriteEditor/SpriteEditorViewportClient.cpp
+++ b/Engine/Plugins/2D/Paper2D/Source/Paper2DEditor/Private/SpriteEditor/SpriteEditorViewportClient.cpp
@@ -127,7 +127,6 @@ FSpriteEditorViewportClient::FSpriteEditorViewportClient(TWeakPtr<FSpriteEditor>
 void FSpriteEditorViewportClient::ActivateEditMode()
 {
 	// Activate the sprite geometry edit mode
-	ModeTools->SetToolkitHost(SpriteEditorPtr.Pin()->GetToolkitHost());
 	ModeTools->SetDefaultMode(FSpriteGeometryEditMode::EM_SpriteGeometry);
 	ModeTools->ActivateDefaultMode();
 
diff --git a/Engine/Plugins/2D/Paper2D/Source/Paper2DEditor/Private/TileMapEditing/TileMapEditor.cpp b/Engine/Plugins/2D/Paper2D/Source/Paper2DEditor/Private/TileMapEditing/TileMapEditor.cpp
--- a/Engine/Plugins/2D/Paper2D/Source/Paper2DEditor/Private/TileMapEditing/TileMapEditor.cpp
+++ b/Engine/Plugins/2D/Paper2D/Source/Paper2DEditor/Private/TileMapEditing/TileMapEditor.cpp
@@ -16,6 +16,7 @@
 #include "TileMapEditing/STileMapEditorViewportToolbar.h"
 #include "Widgets/Docking/SDockTab.h"
 #include "Subsystems/AssetEditorSubsystem.h"
+#include "EditorModeManager.h"
 
 #define LOCTEXT_NAMESPACE "TileMapEditor"
 
@@ -484,6 +485,16 @@ void FTileMapEditor::ExtendToolbar()
 	AddToolbarExtender(ToolbarExtender);
 }
 
+void FTileMapEditor::CreateEditorModeManager()
+{
+	check(ViewportPtr.IsValid());
+	TSharedPtr<FEditorViewportClient> ViewportClient = ViewportPtr->GetViewportClient();
+	check(ViewportClient.IsValid());
+	PRAGMA_DISABLE_DEPRECATION_WARNINGS
+	ViewportClient->TakeOwnershipOfModeManager(EditorModeManager);
+	PRAGMA_ENABLE_DEPRECATION_WARNINGS
+}
+
 void FTileMapEditor::SetTileMapBeingEdited(UPaperTileMap* NewTileMap)
 {
 	if ((NewTileMap != TileMapBeingEdited) && (NewTileMap != nullptr))
diff --git a/Engine/Plugins/2D/Paper2D/Source/Paper2DEditor/Private/TileMapEditing/TileMapEditor.h b/Engine/Plugins/2D/Paper2D/Source/Paper2DEditor/Private/TileMapEditing/TileMapEditor.h
--- a/Engine/Plugins/2D/Paper2D/Source/Paper2DEditor/Private/TileMapEditing/TileMapEditor.h
+++ b/Engine/Plugins/2D/Paper2D/Source/Paper2DEditor/Private/TileMapEditing/TileMapEditor.h
@@ -53,6 +53,7 @@ protected:
 	void BindCommands();
 	void ExtendMenu();
 	void ExtendToolbar();
+	virtual void CreateEditorModeManager() override;
 
 	TSharedRef<SDockTab> SpawnTab_Viewport(const FSpawnTabArgs& Args);
 	TSharedRef<SDockTab> SpawnTab_ToolboxHost(const FSpawnTabArgs& Args);
diff --git a/Engine/Plugins/2D/Paper2D/Source/Paper2DEditor/Private/TileMapEditing/TileMapEditorViewportClient.cpp b/Engine/Plugins/2D/Paper2D/Source/Paper2DEditor/Private/TileMapEditing/TileMapEditorViewportClient.cpp
--- a/Engine/Plugins/2D/Paper2D/Source/Paper2DEditor/Private/TileMapEditing/TileMapEditorViewportClient.cpp
+++ b/Engine/Plugins/2D/Paper2D/Source/Paper2DEditor/Private/TileMapEditing/TileMapEditorViewportClient.cpp
@@ -64,7 +64,6 @@ FTileMapEditorViewportClient::FTileMapEditorViewportClient(TWeakPtr<FTileMapEdit
 void FTileMapEditorViewportClient::ActivateEditMode()
 {
 	// Activate the tile map edit mode
-	ModeTools->SetToolkitHost(TileMapEditorPtr.Pin()->GetToolkitHost());
 	ModeTools->SetDefaultMode(FEdModeTileMap::EM_TileMap);
 	ModeTools->ActivateDefaultMode();
 	
diff --git a/Engine/Plugins/2D/Paper2D/Source/Paper2DEditor/Private/TileSetEditor.cpp b/Engine/Plugins/2D/Paper2D/Source/Paper2DEditor/Private/TileSetEditor.cpp
--- a/Engine/Plugins/2D/Paper2D/Source/Paper2DEditor/Private/TileSetEditor.cpp
+++ b/Engine/Plugins/2D/Paper2D/Source/Paper2DEditor/Private/TileSetEditor.cpp
@@ -15,6 +15,7 @@
 #include "PaperEditorShared/SpriteGeometryEditCommands.h"
 #include "IDetailsView.h"
 #include "Subsystems/AssetEditorSubsystem.h"
+#include "EditorModeManager.h"
 
 #define LOCTEXT_NAMESPACE "TileSetEditor"
 
@@ -289,6 +290,14 @@ void FTileSetEditor::ExtendToolbar()
 	AddToolbarExtender(ToolbarExtender);
 }
 
+void FTileSetEditor::CreateEditorModeManager()
+{
+	check(TileEditorViewportClient.IsValid());
+	PRAGMA_DISABLE_DEPRECATION_WARNINGS
+	TileEditorViewportClient->TakeOwnershipOfModeManager(EditorModeManager);
+	PRAGMA_ENABLE_DEPRECATION_WARNINGS
+}
+
 void FTileSetEditor::OnPropertyChanged(UObject* ObjectBeingModified, FPropertyChangedEvent& PropertyChangedEvent)
 {
 	if (ObjectBeingModified == TileSetBeingEdited)
diff --git a/Engine/Plugins/2D/Paper2D/Source/Paper2DEditor/Private/TileSetEditor.h b/Engine/Plugins/2D/Paper2D/Source/Paper2DEditor/Private/TileSetEditor.h
--- a/Engine/Plugins/2D/Paper2D/Source/Paper2DEditor/Private/TileSetEditor.h
+++ b/Engine/Plugins/2D/Paper2D/Source/Paper2DEditor/Private/TileSetEditor.h
@@ -59,6 +59,7 @@ protected:
 	void BindCommands();
 	void ExtendMenu();
 	void ExtendToolbar();
+	virtual void CreateEditorModeManager() override;
 
 	void OnPropertyChanged(UObject* ObjectBeingModified, FPropertyChangedEvent& PropertyChangedEvent);
 
diff --git a/Engine/Plugins/2D/Paper2D/Source/Paper2DEditor/Private/TileSetEditor/SingleTileEditorViewportClient.cpp b/Engine/Plugins/2D/Paper2D/Source/Paper2DEditor/Private/TileSetEditor/SingleTileEditorViewportClient.cpp
--- a/Engine/Plugins/2D/Paper2D/Source/Paper2DEditor/Private/TileSetEditor/SingleTileEditorViewportClient.cpp
+++ b/Engine/Plugins/2D/Paper2D/Source/Paper2DEditor/Private/TileSetEditor/SingleTileEditorViewportClient.cpp
@@ -296,7 +296,6 @@ int32 FSingleTileEditorViewportClient::GetTileIndex() const
 void FSingleTileEditorViewportClient::ActivateEditMode(TSharedPtr<FUICommandList> InCommandList)
 {
 	// Activate the sprite geometry edit mode
-	//@TODO: ModeTools->SetToolkitHost(SpriteEditorPtr.Pin()->GetToolkitHost());
 	ModeTools->SetDefaultMode(FSpriteGeometryEditMode::EM_SpriteGeometry);
 	ModeTools->ActivateDefaultMode();
 	ModeTools->SetWidgetMode(UE::Widget::WM_Translate);
diff --git a/Engine/Source/Editor/UnrealEd/Private/EditorViewportClient.cpp b/Engine/Source/Editor/UnrealEd/Private/EditorViewportClient.cpp
--- a/Engine/Source/Editor/UnrealEd/Private/EditorViewportClient.cpp
+++ b/Engine/Source/Editor/UnrealEd/Private/EditorViewportClient.cpp
@@ -339,6 +339,12 @@ void FEditorViewportClient::SetCameraSpeedScalar(float SpeedScalar)
 	CameraSpeedScalar = FMath::Clamp<float>(SpeedScalar, 1.0f, TNumericLimits <float>::Max());
 }
 
+void FEditorViewportClient::TakeOwnershipOfModeManager(TSharedPtr<FEditorModeTools>& ModeManagerPtr)
+{
+	check(bOwnsModeTools);
+	ModeManagerPtr = MakeShareable(ModeTools);
+	bOwnsModeTools = false;
+}
 
 float const FEditorViewportClient::SafePadding = 0.075f;
 
diff --git a/Engine/Source/Editor/UnrealEd/Private/Toolkits/AssetEditorToolkit.cpp b/Engine/Source/Editor/UnrealEd/Private/Toolkits/AssetEditorToolkit.cpp
--- a/Engine/Source/Editor/UnrealEd/Private/Toolkits/AssetEditorToolkit.cpp
+++ b/Engine/Source/Editor/UnrealEd/Private/Toolkits/AssetEditorToolkit.cpp
@@ -707,6 +707,11 @@ FLinearColor FAssetEditorToolkit::GetDefaultTabColor() const
 	return TabColor;
 }
 
+void FAssetEditorToolkit::CreateEditorModeManager()
+{
+	EditorModeManager = MakeShared<FEditorModeTools>();
+}
+
 FAssetEditorModeManager* FAssetEditorToolkit::GetAssetEditorModeManager() const
 {
 	PRAGMA_DISABLE_DEPRECATION_WARNINGS
diff --git a/Engine/Source/Editor/UnrealEd/Private/Toolkits/BaseToolkit.cpp b/Engine/Source/Editor/UnrealEd/Private/Toolkits/BaseToolkit.cpp
--- a/Engine/Source/Editor/UnrealEd/Private/Toolkits/BaseToolkit.cpp
+++ b/Engine/Source/Editor/UnrealEd/Private/Toolkits/BaseToolkit.cpp
@@ -133,9 +133,10 @@ void FModeToolkit::Init(const TSharedPtr< class IToolkitHost >& InitToolkitHost,
 {
 	check( InitToolkitHost.IsValid() );
 
-	ToolkitMode = EToolkitMode::WorldCentric;
+	ToolkitMode = EToolkitMode::Type::Standalone;
 	ToolkitHost = InitToolkitHost;
 	OwningEditorMode = InOwningMode;
+	EditorModeManager = InitToolkitHost->GetEditorModeManager().AsShared();
 
 	UInteractiveToolManager* ToolManager = GetEditorModeManager().GetInteractiveToolsContext()->ToolManager;
 	ToolManager->OnToolStarted.AddSP(this, &FModeToolkit::OnToolStarted);
diff --git a/Engine/Source/Editor/UnrealEd/Private/Toolkits/SStandaloneAssetEditorToolkitHost.cpp b/Engine/Source/Editor/UnrealEd/Private/Toolkits/SStandaloneAssetEditorToolkitHost.cpp
--- a/Engine/Source/Editor/UnrealEd/Private/Toolkits/SStandaloneAssetEditorToolkitHost.cpp
+++ b/Engine/Source/Editor/UnrealEd/Private/Toolkits/SStandaloneAssetEditorToolkitHost.cpp
@@ -249,6 +249,10 @@ FEditorModeTools& SStandaloneAssetEditorToolkitHost::GetEditorModeManager() cons
 
 SStandaloneAssetEditorToolkitHost::~SStandaloneAssetEditorToolkitHost()
 {
+	// Reset the mode manager so that any modes from the HostedAssetEditorToolkit are exited
+	GetEditorModeManager().SetDefaultMode(FBuiltinEditorModes::EM_Default);
+	GetEditorModeManager().DeactivateAllModes();
+
 	// Let the toolkit manager know that we're going away now
 	FToolkitManager::Get().OnToolkitHostDestroyed( this );
 	HostedToolkits.Reset();
diff --git a/Engine/Source/Editor/UnrealEd/Public/EditorViewportClient.h b/Engine/Source/Editor/UnrealEd/Public/EditorViewportClient.h
--- a/Engine/Source/Editor/UnrealEd/Public/EditorViewportClient.h
+++ b/Engine/Source/Editor/UnrealEd/Public/EditorViewportClient.h
@@ -1092,6 +1092,10 @@ public:
 		return ModeTools;
 	}
 
+	/** Legacy adapter for a toolkit to take ownership of a mode manager that may have been created by this viewport client */
+	UE_DEPRECATED(5.0, "This function is meant for legacy edge cases (pre-UAssetEditor), where the toolkit may not be the original owner of the mode manager.")
+	void TakeOwnershipOfModeManager(TSharedPtr<FEditorModeTools>& ModeManagerPtr);
+
 	/** Get the editor viewport widget */
 	TSharedPtr<SEditorViewport> GetEditorViewportWidget() const { return EditorViewportWidget.Pin(); }
 
diff --git a/Engine/Source/Editor/UnrealEd/Public/Toolkits/AssetEditorToolkit.h b/Engine/Source/Editor/UnrealEd/Public/Toolkits/AssetEditorToolkit.h
--- a/Engine/Source/Editor/UnrealEd/Public/Toolkits/AssetEditorToolkit.h
+++ b/Engine/Source/Editor/UnrealEd/Public/Toolkits/AssetEditorToolkit.h
@@ -341,6 +341,8 @@ protected:
 	/** @return the color to use for the tab color */
 	virtual FLinearColor GetDefaultTabColor() const;
 
+	virtual void CreateEditorModeManager() override;
+
 private:
 	/** Spawns the toolbar tab */
 	TSharedRef<SDockTab> SpawnTab_Toolbar(const FSpawnTabArgs& Args);
