Fix giant paper2d toolbar

diff --git a/Engine/Plugins/2D/Paper2D/Source/Paper2DEditor/Private/SpriteEditor/SpriteEditor.cpp b/Engine/Plugins/2D/Paper2D/Source/Paper2DEditor/Private/SpriteEditor/SpriteEditor.cpp
--- a/Engine/Plugins/2D/Paper2D/Source/Paper2DEditor/Private/SpriteEditor/SpriteEditor.cpp
+++ b/Engine/Plugins/2D/Paper2D/Source/Paper2DEditor/Private/SpriteEditor/SpriteEditor.cpp
@@ -601,7 +601,7 @@ ESpriteEditorMode::Type FSpriteEditor::GetCurrentMode() const
 
 void FSpriteEditor::CreateModeToolbarWidgets(FToolBarBuilder& IgnoredBuilder)
 {
-	FToolBarBuilder ToolbarBuilder(ViewportPtr->GetCommandList(), FMultiBoxCustomization::None);
+	FSlimHorizontalToolBarBuilder ToolbarBuilder(ViewportPtr->GetCommandList(), FMultiBoxCustomization::None);
 	ToolbarBuilder.AddToolBarButton(FSpriteEditorCommands::Get().EnterViewMode);
 	ToolbarBuilder.AddToolBarButton(FSpriteEditorCommands::Get().EnterSourceRegionEditMode);
 	ToolbarBuilder.AddToolBarButton(FSpriteEditorCommands::Get().EnterCollisionEditMode);
