LWC: Fix FPaperSpriteVertexFactory position component stream stride

diff --git a/Engine/Plugins/2D/Paper2D/Source/Paper2D/Private/PaperSpriteVertexBuffer.cpp b/Engine/Plugins/2D/Paper2D/Source/Paper2D/Private/PaperSpriteVertexBuffer.cpp
--- a/Engine/Plugins/2D/Paper2D/Source/Paper2D/Private/PaperSpriteVertexBuffer.cpp
+++ b/Engine/Plugins/2D/Paper2D/Source/Paper2D/Private/PaperSpriteVertexBuffer.cpp
@@ -228,7 +228,7 @@ void FPaperSpriteVertexFactory::Init(const FPaperSpriteVertexBuffer* InVertexBuf
 		VertexData.PositionComponentSRV = InVertexBuffer->PositionBufferSRV;
 
 		// Vertex Streams
-		VertexData.PositionComponent = FVertexStreamComponent(&InVertexBuffer->PositionBuffer, 0, sizeof(FVector), VET_Float3, EVertexStreamUsage::Default);
+		VertexData.PositionComponent = FVertexStreamComponent(&InVertexBuffer->PositionBuffer, 0, sizeof(FVector3f), VET_Float3, EVertexStreamUsage::Default);
 		VertexData.TangentBasisComponents[0] = FVertexStreamComponent(&InVertexBuffer->TangentBuffer, 0, 2 * sizeof(FPackedNormal), VET_PackedNormal, EVertexStreamUsage::ManualFetch);
 		VertexData.TangentBasisComponents[1] = FVertexStreamComponent(&InVertexBuffer->TangentBuffer, sizeof(FPackedNormal), 2 * sizeof(FPackedNormal), VET_PackedNormal, EVertexStreamUsage::ManualFetch);
 		VertexData.ColorComponent = FVertexStreamComponent(&InVertexBuffer->ColorBuffer, 0, sizeof(FColor), VET_Color, EVertexStreamUsage::ManualFetch);
