Fix: Scaling and rotation modes of the gizmo should be disabled in Paper2D Sprite Editor for vertexes

diff --git a/Engine/Plugins/2D/Paper2D/Source/Paper2DEditor/Private/PaperEditorShared/SpriteGeometryEditMode.cpp b/Engine/Plugins/2D/Paper2D/Source/Paper2DEditor/Private/PaperEditorShared/SpriteGeometryEditMode.cpp
--- a/Engine/Plugins/2D/Paper2D/Source/Paper2DEditor/Private/PaperEditorShared/SpriteGeometryEditMode.cpp
+++ b/Engine/Plugins/2D/Paper2D/Source/Paper2DEditor/Private/PaperEditorShared/SpriteGeometryEditMode.cpp
@@ -633,6 +633,18 @@ void FSpriteGeometryEditMode::DrawCollisionStats(FViewport& InViewport, FSceneVi
 	YPos = (int32)TextItem.Position.Y;
 }
 
+bool FSpriteGeometryEditMode::UsesTransformWidget(UE::Widget::EWidgetMode CheckMode) const
+{
+	switch(CheckMode)
+	{
+	case UE::Widget::WM_None:
+	case UE::Widget::WM_Translate:
+		return true;
+	default:
+		return false;
+	}
+}
+
 //////////////////////////////////////////////////////////////////////////
 
 #undef LOCTEXT_NAMESPACE
diff --git a/Engine/Plugins/2D/Paper2D/Source/Paper2DEditor/Private/PaperEditorShared/SpriteGeometryEditMode.h b/Engine/Plugins/2D/Paper2D/Source/Paper2DEditor/Private/PaperEditorShared/SpriteGeometryEditMode.h
--- a/Engine/Plugins/2D/Paper2D/Source/Paper2DEditor/Private/PaperEditorShared/SpriteGeometryEditMode.h
+++ b/Engine/Plugins/2D/Paper2D/Source/Paper2DEditor/Private/PaperEditorShared/SpriteGeometryEditMode.h
@@ -64,6 +64,7 @@ public:
 	static void DrawGeometryStats(FViewport& InViewport, FSceneView& View, FCanvas& Canvas, const FSpriteGeometryCollection& Geometry, bool bIsRenderGeometry, int32& YPos);
 	static void DrawCollisionStats(FViewport& InViewport, FSceneView& View, FCanvas& Canvas, class UBodySetup* BodySetup, int32& YPos);
 
+	virtual bool UsesTransformWidget(UE::Widget::EWidgetMode CheckMode) const override;
 protected:
 	FBox2D BoundsForNewShapes;
 	FLinearColor GeometryVertexColor;
