Make FVector2f <-> FVector2d conversion casts explicit
Submitted on behalf of fred.kimberley

diff --git a/Engine/Plugins/2D/Paper2D/Source/Paper2D/Private/PaperRenderSceneProxy.h b/Engine/Plugins/2D/Paper2D/Source/Paper2D/Private/PaperRenderSceneProxy.h
--- a/Engine/Plugins/2D/Paper2D/Source/Paper2D/Private/PaperRenderSceneProxy.h
+++ b/Engine/Plugins/2D/Paper2D/Source/Paper2D/Private/PaperRenderSceneProxy.h
@@ -96,7 +96,7 @@ struct PAPER2D_API FSpriteRenderSection
 			const FVector Pos((PaperAxisX * SourceVert.X) + (PaperAxisY * SourceVert.Y) + Record.Destination);
 			const FVector2f UV(SourceVert.Z, SourceVert.W);	// LWC_TODO: Precision loss
 
-			new (Vertices) FDynamicMeshVertex(Pos, FPaperSpriteTangents::PackedNormalX.ToFVector(), FPaperSpriteTangents::PackedNormalZ.ToFVector(), UV, VertColor);
+			new (Vertices) FDynamicMeshVertex(FVector3f(Pos), FPaperSpriteTangents::PackedNormalX.ToFVector3f(), FPaperSpriteTangents::PackedNormalZ.ToFVector3f(), UV, VertColor);
 		}
 	}
 
@@ -105,7 +105,7 @@ struct PAPER2D_API FSpriteRenderSection
 	{
 		const FVector Pos((PaperAxisX * X) + (PaperAxisY * Y) + Origin);
 
-		new (Vertices) FDynamicMeshVertex(Pos, FPaperSpriteTangents::PackedNormalX.ToFVector(), FPaperSpriteTangents::PackedNormalZ.ToFVector(), FVector2f(U, V), Color);
+		new (Vertices) FDynamicMeshVertex(FVector3f(Pos), FPaperSpriteTangents::PackedNormalX.ToFVector3f(), FPaperSpriteTangents::PackedNormalZ.ToFVector3f(), FVector2f(U, V), Color);
 		++NumVertices;
 	}
 
@@ -114,7 +114,7 @@ struct PAPER2D_API FSpriteRenderSection
 	{
 		const FVector Pos((PaperAxisX * X) + (PaperAxisY * Y) + Origin);
 
-		new (Vertices) FDynamicMeshVertex(Pos, TangentX.ToFVector(), TangentZ.ToFVector(), FVector2D(U, V), Color);
+		new (Vertices) FDynamicMeshVertex(FVector3f(Pos), TangentX.ToFVector3f(), TangentZ.ToFVector3f(), FVector2f(U, V), Color);
 		++NumVertices;
 	}
 };
diff --git a/Engine/Plugins/Enterprise/DatasmithC4DImporter/Source/DatasmithC4DDynamicImporter/Private/DatasmithC4DImporter.cpp b/Engine/Plugins/Enterprise/DatasmithC4DImporter/Source/DatasmithC4DDynamicImporter/Private/DatasmithC4DImporter.cpp
--- a/Engine/Plugins/Enterprise/DatasmithC4DImporter/Source/DatasmithC4DDynamicImporter/Private/DatasmithC4DImporter.cpp
+++ b/Engine/Plugins/Enterprise/DatasmithC4DImporter/Source/DatasmithC4DDynamicImporter/Private/DatasmithC4DImporter.cpp
@@ -3807,7 +3807,7 @@ TSharedPtr<IDatasmithMeshElement> FDatasmithC4DDynamicImporter::ImportMesh(cinew
 				FVertexInstanceID& VertInstanceID = VertexInstances[VertexCount];
 				int32 VertexIDInQuad = (*IndexOffsets)[VertexCount];
 
-				VertexInstanceUVs.Set(VertInstanceID, ChannelIndex, QuadUVs[VertexIDInQuad]);
+				VertexInstanceUVs.Set(VertInstanceID, ChannelIndex, FVector2f(QuadUVs[VertexIDInQuad]));
 			}
 		}
 
diff --git a/Engine/Plugins/Runtime/Oculus/OculusVR/Source/OculusInput/Private/OculusHandTracking.cpp b/Engine/Plugins/Runtime/Oculus/OculusVR/Source/OculusInput/Private/OculusHandTracking.cpp
--- a/Engine/Plugins/Runtime/Oculus/OculusVR/Source/OculusInput/Private/OculusHandTracking.cpp
+++ b/Engine/Plugins/Runtime/Oculus/OculusVR/Source/OculusInput/Private/OculusHandTracking.cpp
@@ -419,7 +419,7 @@ void FOculusHandTracking::InitializeHandMesh(USkeletalMesh* SkeletalMesh, const
 		ModelVertex.Position = FVector(VertexPosition.x, VertexPosition.z, VertexPosition.y) * WorldToMeters;
 		ModelVertex.TangentZ = FVector3f(Normal.x, Normal.z, Normal.y);
 		ModelVertex.TangentX = FVector3f(1.0f, 0.0f, 0.0f);
-		ModelVertex.TexCoord = FVector2D(OvrMesh->VertexUV0[VertexIndex].x, OvrMesh->VertexUV0[VertexIndex].y);
+		ModelVertex.TexCoord = FVector2f(OvrMesh->VertexUV0[VertexIndex].x, OvrMesh->VertexUV0[VertexIndex].y);
 
 		// Add Model Vertex data to vertex buffer
 		LodRenderData->StaticVertexBuffers.PositionVertexBuffer.VertexPosition(VertexIndex) = ModelVertex.Position;
diff --git a/Engine/Source/Runtime/Core/Public/Math/Vector2D.h b/Engine/Source/Runtime/Core/Public/Math/Vector2D.h
--- a/Engine/Source/Runtime/Core/Public/Math/Vector2D.h
+++ b/Engine/Source/Runtime/Core/Public/Math/Vector2D.h
@@ -638,9 +638,9 @@ public:
 	
 	
 	
-	// Conversion from other type. TODO: explicit! 
+	// Conversion from other type.
 	template<typename FArg, TEMPLATE_REQUIRES(!TIsSame<T, FArg>::Value)>
-	TVector2(const TVector2<FArg>& From) : TVector2<T>((T)From.X, (T)From.Y) {}
+	explicit TVector2(const TVector2<FArg>& From) : TVector2<T>((T)From.X, (T)From.Y) {}
 };
 
 /**
diff --git a/Engine/Source/Runtime/Experimental/Chaos/Public/GeometryCollection/GeometryCollectionUtility.h b/Engine/Source/Runtime/Experimental/Chaos/Public/GeometryCollection/GeometryCollectionUtility.h
--- a/Engine/Source/Runtime/Experimental/Chaos/Public/GeometryCollection/GeometryCollectionUtility.h
+++ b/Engine/Source/Runtime/Experimental/Chaos/Public/GeometryCollection/GeometryCollectionUtility.h
@@ -60,7 +60,7 @@ namespace GeometryCollection
 			Normals[Idx] = NormalsIn.Num() > Idx ? FVector(NormalsIn[Idx][0], NormalsIn[Idx][1], NormalsIn[Idx][2]) : FVector(0);
 
 			FVector2D UV = UVsIn.Num() > Idx ? FVector2D(UVsIn[Idx][0], UVsIn[Idx][1]) : FVector2D(0);
-			UVs[Idx].Add(UV);
+			UVs[Idx].Add(FVector2f(UV));
 			
 			Colors[Idx] = FLinearColor::White;
 		}
