Move IntroTutorials to a plugin. Part 1: move code module and clean up references.

diff --git a/Engine/Plugins/2D/Paper2D/Source/Paper2DEditor/Private/Paper2DEditorModule.cpp b/Engine/Plugins/2D/Paper2D/Source/Paper2DEditor/Private/Paper2DEditorModule.cpp
--- a/Engine/Plugins/2D/Paper2D/Source/Paper2DEditor/Private/Paper2DEditorModule.cpp
+++ b/Engine/Plugins/2D/Paper2D/Source/Paper2DEditor/Private/Paper2DEditorModule.cpp
@@ -70,10 +70,6 @@
 #include "PaperRuntimeSettings.h"
 #include "ISettingsModule.h"
 
-// Intro tutorials
-#include "EditorTutorial.h"
-#include "IIntroTutorials.h"
-
 // Mesh paint adapters
 #include "MeshPaintModule.h"
 #include "MeshPainting/MeshPaintSpriteAdapter.h"
@@ -220,22 +216,6 @@ public:
 			MeshPaintModule->RegisterGeometryAdapterFactory(SpriteMeshPaintAdapterFactory.ToSharedRef());
 		}
 
-		if (!IsRunningCommandlet())
-		{
-			// Register Paper2D tutorial category with the tutorial browser
-			FTutorialCategory Paper2DCategory = FTutorialCategory();
-			Paper2DCategory.Identifier = "Paper2D";
-			Paper2DCategory.Title = NSLOCTEXT("TutorialCategories", "Paper2DTitle", "Paper2D");
-			Paper2DCategory.Description = NSLOCTEXT("TutorialCategories", "Paper2DDescription", "Tutorials covering the usage of the Unreal Engine 4 2D game tool: Paper2D.");
-			Paper2DCategory.Icon = "ClassThumbnail.PaperSprite";
-			Paper2DCategory.Texture = FSoftObjectPath("/Paper2D/Tutorial/Paper2D/TutorialAssets/Paper2DSprite_TutorialIcon.Paper2DSprite_TutorialIcon");
-			Paper2DCategory.SortOrder = 500;
-
-			IIntroTutorials& IntroTutorials = FModuleManager::GetModuleChecked<IIntroTutorials>(TEXT("IntroTutorials"));
-			IntroTutorials.RegisterCategory(Paper2DCategory);
-		}
-
-		//
 		RegisterSettings();
 	}
 
diff --git a/Engine/Plugins/Runtime/Analytics/AnalyticsMulticast/Source/AnaltyicsMulticastEditor/AnalyticsMulticastEditor.Build.cs b/Engine/Plugins/Runtime/Analytics/AnalyticsMulticast/Source/AnaltyicsMulticastEditor/AnalyticsMulticastEditor.Build.cs
--- a/Engine/Plugins/Runtime/Analytics/AnalyticsMulticast/Source/AnaltyicsMulticastEditor/AnalyticsMulticastEditor.Build.cs
+++ b/Engine/Plugins/Runtime/Analytics/AnalyticsMulticast/Source/AnaltyicsMulticastEditor/AnalyticsMulticastEditor.Build.cs
@@ -29,7 +29,6 @@ public class AnalyticsMulticastEditor : ModuleRules
 		PrivateIncludePathModuleNames.AddRange(
 			new string[] {
 				"Settings",
-				"IntroTutorials",
                 "AssetTools"
 			}
 		);
diff --git a/Engine/Source/Developer/MessageLog/MessageLog.Build.cs b/Engine/Source/Developer/MessageLog/MessageLog.Build.cs
--- a/Engine/Source/Developer/MessageLog/MessageLog.Build.cs
+++ b/Engine/Source/Developer/MessageLog/MessageLog.Build.cs
@@ -31,7 +31,6 @@ public class MessageLog : ModuleRules
 			PrivateIncludePathModuleNames.AddRange(
 				new string[] {
 					"Documentation",
-					"IntroTutorials",
 				}
 			);		
 		}
diff --git a/Engine/Source/Developer/MessageLog/Private/UserInterface/SMessageLogMessageListRow.h b/Engine/Source/Developer/MessageLog/Private/UserInterface/SMessageLogMessageListRow.h
--- a/Engine/Source/Developer/MessageLog/Private/UserInterface/SMessageLogMessageListRow.h
+++ b/Engine/Source/Developer/MessageLog/Private/UserInterface/SMessageLogMessageListRow.h
@@ -15,7 +15,6 @@
 
 #if WITH_EDITOR
 	#include "IDocumentation.h"
-	#include "IIntroTutorials.h"
 #endif
 
 #define LOCTEXT_NAMESPACE "SMessageLogMessageListRow"
@@ -103,7 +102,7 @@ private:
 
 	void HandleTutorialHyperlinkNavigate( FString TutorialAssetName )
 	{
-		IIntroTutorials::Get().LaunchTutorial(TutorialAssetName);
+		//IIntroTutorials::Get().LaunchTutorial(TutorialAssetName);
 	}
 #endif
 
diff --git a/Engine/Source/Editor/ContentBrowser/ContentBrowser.Build.cs b/Engine/Source/Editor/ContentBrowser/ContentBrowser.Build.cs
--- a/Engine/Source/Editor/ContentBrowser/ContentBrowser.Build.cs
+++ b/Engine/Source/Editor/ContentBrowser/ContentBrowser.Build.cs
@@ -62,7 +62,6 @@ public class ContentBrowser : ModuleRules
 		PublicIncludePathModuleNames.AddRange(
             new string[] {
 				"ContentBrowserData",
-				"IntroTutorials"
             }
         );
 	}
diff --git a/Engine/Source/Editor/MaterialEditor/Private/MaterialEditorModule.cpp b/Engine/Source/Editor/MaterialEditor/Private/MaterialEditorModule.cpp
--- a/Engine/Source/Editor/MaterialEditor/Private/MaterialEditorModule.cpp
+++ b/Engine/Source/Editor/MaterialEditor/Private/MaterialEditorModule.cpp
@@ -13,9 +13,7 @@
 #include "EditorStyleSet.h"
 #include "ContentBrowserDelegates.h"
 #include "Interfaces/IPluginManager.h"
-#include "IIntroTutorials.h"
 #include "Widgets/Docking/SDockTab.h"
-#include "EditorTutorial.h"
 #include "Framework/Application/SlateApplication.h"
 #include "Widgets/Images/SImage.h"
 #include "Misc/EngineBuildSettings.h"
diff --git a/Engine/Source/Editor/UnrealEd/Private/EditorEngine.cpp b/Engine/Source/Editor/UnrealEd/Private/EditorEngine.cpp
--- a/Engine/Source/Editor/UnrealEd/Private/EditorEngine.cpp
+++ b/Engine/Source/Editor/UnrealEd/Private/EditorEngine.cpp
@@ -1102,11 +1102,6 @@ void UEditorEngine::Init(IEngineLoop* InEngineLoop)
 			}
 		}
 
-		if (!IsRunningCommandlet())
-		{
-			FModuleManager::Get().LoadModule(TEXT("IntroTutorials"));
-		}
-
 		if( FParse::Param( FCommandLine::Get(),TEXT( "PListEditor" ) ) )
 		{
 			FModuleManager::Get().LoadModule(TEXT("PListEditor"));
diff --git a/Engine/Source/Editor/UnrealEd/UnrealEd.Build.cs b/Engine/Source/Editor/UnrealEd/UnrealEd.Build.cs
--- a/Engine/Source/Editor/UnrealEd/UnrealEd.Build.cs
+++ b/Engine/Source/Editor/UnrealEd/UnrealEd.Build.cs
@@ -52,7 +52,6 @@ public class UnrealEd : ModuleRules
 				"AudioEditor",
 				"ViewportSnapping",
 				"SourceCodeAccess",
-				"IntroTutorials",
 				"OutputLog",
 				"Landscape",
 				"LocalizationService",
@@ -245,7 +244,6 @@ public class UnrealEd : ModuleRules
 				"SettingsEditor",
 				"StringTableEditor",
 				"Blutility",
-				"IntroTutorials",
 				"WorkspaceMenuStructure",
 				"PlacementMode",
 				"MeshUtilities",
