Fix FVector2D variant casts - Plugins
Submitted on behalf of fred.kimberley

diff --git a/Engine/Plugins/2D/Paper2D/Source/Paper2D/Private/PaperRenderSceneProxy.h b/Engine/Plugins/2D/Paper2D/Source/Paper2D/Private/PaperRenderSceneProxy.h
--- a/Engine/Plugins/2D/Paper2D/Source/Paper2D/Private/PaperRenderSceneProxy.h
+++ b/Engine/Plugins/2D/Paper2D/Source/Paper2D/Private/PaperRenderSceneProxy.h
@@ -96,7 +96,7 @@ struct PAPER2D_API FSpriteRenderSection
 			const FVector Pos((PaperAxisX * SourceVert.X) + (PaperAxisY * SourceVert.Y) + Record.Destination);
 			const FVector2f UV(SourceVert.Z, SourceVert.W);	// LWC_TODO: Precision loss
 
-			new (Vertices) FDynamicMeshVertex((FVector3f)Pos, FPaperSpriteTangents::PackedNormalX.ToFVector3f(), FPaperSpriteTangents::PackedNormalZ.ToFVector3f(), UV, VertColor);
+			new (Vertices) FDynamicMeshVertex(FVector3f(Pos), FPaperSpriteTangents::PackedNormalX.ToFVector3f(), FPaperSpriteTangents::PackedNormalZ.ToFVector3f(), UV, VertColor);
 		}
 	}
 
