Workaround PVS false warning caused by change to ensure macro
#rnx

diff --git a/Engine/Plugins/2D/Paper2D/Source/Paper2DEditor/Private/TileMapEditing/EdModeTileMap.cpp b/Engine/Plugins/2D/Paper2D/Source/Paper2DEditor/Private/TileMapEditing/EdModeTileMap.cpp
--- a/Engine/Plugins/2D/Paper2D/Source/Paper2DEditor/Private/TileMapEditing/EdModeTileMap.cpp
+++ b/Engine/Plugins/2D/Paper2D/Source/Paper2DEditor/Private/TileMapEditing/EdModeTileMap.cpp
@@ -1386,7 +1386,8 @@ void FEdModeTileMap::UpdatePreviewCursor(const FViewportCursorLocation& Ray)
 		{
 			UPaperTileMap* TileMap = TileLayer->GetTileMap();
 			int32 LayerIndex;
-			ensure(TileMap->TileLayers.Find(TileLayer, LayerIndex));
+			const bool bFoundLayer = TileMap->TileLayers.Find(TileLayer, LayerIndex)
+			ensure(bFoundLayer);
 
 			LastCursorTileX = LocalTileX0;
 			LastCursorTileY = LocalTileY0;
