[Backout] - CL20113928
No longer needed as the ensure macro change necessitating it has been undone (and this was missing a ;)
Original CL Desc
-----------------------------------------------------------------
Workaround PVS false warning caused by change to ensure macro
#rnx

diff --git a/Engine/Plugins/2D/Paper2D/Source/Paper2DEditor/Private/TileMapEditing/EdModeTileMap.cpp b/Engine/Plugins/2D/Paper2D/Source/Paper2DEditor/Private/TileMapEditing/EdModeTileMap.cpp
--- a/Engine/Plugins/2D/Paper2D/Source/Paper2DEditor/Private/TileMapEditing/EdModeTileMap.cpp
+++ b/Engine/Plugins/2D/Paper2D/Source/Paper2DEditor/Private/TileMapEditing/EdModeTileMap.cpp
@@ -1386,8 +1386,7 @@ void FEdModeTileMap::UpdatePreviewCursor(const FViewportCursorLocation& Ray)
 		{
 			UPaperTileMap* TileMap = TileLayer->GetTileMap();
 			int32 LayerIndex;
-			const bool bFoundLayer = TileMap->TileLayers.Find(TileLayer, LayerIndex)
-			ensure(bFoundLayer);
+			ensure(TileMap->TileLayers.Find(TileLayer, LayerIndex));
 
 			LastCursorTileX = LocalTileX0;
 			LastCursorTileY = LocalTileY0;
