Converting hardcoded short class/enum names to pathnames ahead of ANY_PACKAGE removal

diff --git a/Engine/Plugins/2D/Paper2D/Source/Paper2D/Classes/PaperTileLayer.h b/Engine/Plugins/2D/Paper2D/Source/Paper2D/Classes/PaperTileLayer.h
--- a/Engine/Plugins/2D/Paper2D/Source/Paper2D/Classes/PaperTileLayer.h
+++ b/Engine/Plugins/2D/Paper2D/Source/Paper2D/Classes/PaperTileLayer.h
@@ -23,7 +23,7 @@ enum class EPaperTileFlags : uint32
 };
 
 // This is the contents of a tile map cell
-USTRUCT(BlueprintType, meta=(HasNativeBreak="Paper2D.TileMapBlueprintLibrary.BreakTile", HasNativeMake="Paper2D.TileMapBlueprintLibrary.MakeTile"))
+USTRUCT(BlueprintType, meta=(HasNativeBreak="/Script/Paper2D.TileMapBlueprintLibrary.BreakTile", HasNativeMake="/Script/Paper2D.TileMapBlueprintLibrary.MakeTile"))
 struct FPaperTileInfo
 {
 	GENERATED_USTRUCT_BODY()
diff --git a/Engine/Plugins/2D/Paper2D/Source/Paper2DEditor/Classes/PaperImporterSettings.h b/Engine/Plugins/2D/Paper2D/Source/Paper2DEditor/Classes/PaperImporterSettings.h
--- a/Engine/Plugins/2D/Paper2D/Source/Paper2DEditor/Classes/PaperImporterSettings.h
+++ b/Engine/Plugins/2D/Paper2D/Source/Paper2DEditor/Classes/PaperImporterSettings.h
@@ -92,27 +92,27 @@ protected:
 	TEnumAsByte<TextureCompressionSettings> DefaultSpriteTextureCompression;
 
 	// The unlit default masked material for newly created sprites (masked means binary opacity: things are either opaque or see-thru, with nothing in between)
-	UPROPERTY(config, EditAnywhere, Category=MaterialSettings, meta=(AllowedClasses="MaterialInterface", DisplayName="Unlit Default Masked Material"))
+	UPROPERTY(config, EditAnywhere, Category=MaterialSettings, meta=(AllowedClasses="/Script/Engine.MaterialInterface", DisplayName="Unlit Default Masked Material"))
 	FSoftObjectPath UnlitDefaultMaskedMaterialName;
 
 	// The unlit default translucent material for newly created sprites (translucent means smooth opacity which can vary continuously from 0..1, but translucent rendering is more expensive that opaque or masked rendering and has different sorting rules)
-	UPROPERTY(config, EditAnywhere, Category=MaterialSettings, meta=(AllowedClasses="MaterialInterface", DisplayName="Unlit Default Translucent Material"))
+	UPROPERTY(config, EditAnywhere, Category=MaterialSettings, meta=(AllowedClasses="/Script/Engine.MaterialInterface", DisplayName="Unlit Default Translucent Material"))
 	FSoftObjectPath UnlitDefaultTranslucentMaterialName;
 
 	// The unlit default opaque material for newly created sprites
-	UPROPERTY(config, EditAnywhere, Category=MaterialSettings, meta=(AllowedClasses="MaterialInterface", DisplayName="Unlit Default Opaque Sprite Material"))
+	UPROPERTY(config, EditAnywhere, Category=MaterialSettings, meta=(AllowedClasses="/Script/Engine.MaterialInterface", DisplayName="Unlit Default Opaque Sprite Material"))
 	FSoftObjectPath UnlitDefaultOpaqueMaterialName;
 
 	// The lit default masked material for newly created sprites (masked means binary opacity: things are either opaque or see-thru, with nothing in between)
-	UPROPERTY(config, EditAnywhere, Category=MaterialSettings, meta=(AllowedClasses="MaterialInterface", DisplayName="Lit Default Masked Material"))
+	UPROPERTY(config, EditAnywhere, Category=MaterialSettings, meta=(AllowedClasses="/Script/Engine.MaterialInterface", DisplayName="Lit Default Masked Material"))
 	FSoftObjectPath LitDefaultMaskedMaterialName;
 
 	// The lit default translucent material for newly created sprites (translucent means smooth opacity which can vary continuously from 0..1, but translucent rendering is more expensive that opaque or masked rendering and has different sorting rules)
-	UPROPERTY(config, EditAnywhere, Category=MaterialSettings, meta=(AllowedClasses="MaterialInterface", DisplayName="Lit Default Translucent Material"))
+	UPROPERTY(config, EditAnywhere, Category=MaterialSettings, meta=(AllowedClasses="/Script/Engine.MaterialInterface", DisplayName="Lit Default Translucent Material"))
 	FSoftObjectPath LitDefaultTranslucentMaterialName;
 
 	// The lit default opaque material for newly created sprites
-	UPROPERTY(config, EditAnywhere, Category=MaterialSettings, meta=(AllowedClasses="MaterialInterface", DisplayName="Lit Default Opaque Material"))
+	UPROPERTY(config, EditAnywhere, Category=MaterialSettings, meta=(AllowedClasses="/Script/Engine.MaterialInterface", DisplayName="Lit Default Opaque Material"))
 	FSoftObjectPath LitDefaultOpaqueMaterialName;
 
 public:
diff --git a/Engine/Plugins/AI/MLAdapter/Source/MLAdapter/Public/MLAdapterSettings.h b/Engine/Plugins/AI/MLAdapter/Source/MLAdapter/Public/MLAdapterSettings.h
--- a/Engine/Plugins/AI/MLAdapter/Source/MLAdapter/Public/MLAdapterSettings.h
+++ b/Engine/Plugins/AI/MLAdapter/Source/MLAdapter/Public/MLAdapterSettings.h
@@ -35,13 +35,13 @@ public:
 	static uint16 GetDefaultRPCServerPort() { return GET_CONFIG_VALUE(DefaultRPCServerPort); }
 
 protected:
-	UPROPERTY(EditAnywhere, config, Category = MLAdapter, meta = (MetaClass = "MLAdapterManager"))
+	UPROPERTY(EditAnywhere, config, Category = MLAdapter, meta = (MetaClass = "/Script/MLAdapter.MLAdapterManager"))
 	FSoftClassPath ManagerClass;
 
-	UPROPERTY(EditAnywhere, config, Category = MLAdapter, meta = (MetaClass = "MLAdapterSession"))
+	UPROPERTY(EditAnywhere, config, Category = MLAdapter, meta = (MetaClass = "/Script/MLAdapter.MLAdapterSession"))
 	FSoftClassPath SessionClass;
 
-	UPROPERTY(EditAnywhere, config, Category = MLAdapter, meta = (MetaClass = "MLAdapterAgent"))
+	UPROPERTY(EditAnywhere, config, Category = MLAdapter, meta = (MetaClass = "/Script/MLAdapter.MLAdapterAgent"))
 	FSoftClassPath DefaultAgentClass;
 
 	UPROPERTY(EditAnywhere, config, Category = MLAdapter)
diff --git a/Engine/Plugins/Animation/AnimationWarping/Source/Runtime/Private/BoneControllers/AnimNode_OrientationWarping.cpp b/Engine/Plugins/Animation/AnimationWarping/Source/Runtime/Private/BoneControllers/AnimNode_OrientationWarping.cpp
--- a/Engine/Plugins/Animation/AnimationWarping/Source/Runtime/Private/BoneControllers/AnimNode_OrientationWarping.cpp
+++ b/Engine/Plugins/Animation/AnimationWarping/Source/Runtime/Private/BoneControllers/AnimNode_OrientationWarping.cpp
@@ -56,7 +56,7 @@ void FAnimNode_OrientationWarping::GatherDebugData(FNodeDebugData& DebugData)
 #endif
 		}
 		DebugLine += FString::Printf(TEXT("\n - Distributed Bone Orientation Alpha: (%.3fd)"), DistributedBoneOrientationAlpha);
-		if (const UEnum* TypeEnum = FindObject<UEnum>(ANY_PACKAGE, TEXT("EAxis")))
+		if (const UEnum* TypeEnum = FindObject<UEnum>(nullptr, TEXT("/Script/CoreUObject.EAxis")))
 		{
 			DebugLine += FString::Printf(TEXT("\n - Rotation Axis: (%s)"), *(TypeEnum->GetNameStringByIndex(static_cast<int32>(RotationAxis))));
 		}
diff --git a/Engine/Plugins/Animation/LiveLink/Source/LiveLinkGraphNode/Private/K2Node_EvaluateLiveLinkCustom.cpp b/Engine/Plugins/Animation/LiveLink/Source/LiveLinkGraphNode/Private/K2Node_EvaluateLiveLinkCustom.cpp
--- a/Engine/Plugins/Animation/LiveLink/Source/LiveLinkGraphNode/Private/K2Node_EvaluateLiveLinkCustom.cpp
+++ b/Engine/Plugins/Animation/LiveLink/Source/LiveLinkGraphNode/Private/K2Node_EvaluateLiveLinkCustom.cpp
@@ -49,7 +49,7 @@ void UK2Node_EvaluateLiveLinkFrameAtSceneTime::AllocateDefaultPins()
 	const UEdGraphSchema_K2* K2Schema = GetDefault<UEdGraphSchema_K2>();
 
 	// Timecode pin
-	UScriptStruct* TimecodeScriptStruct = FindObject<UScriptStruct>(ANY_PACKAGE, TEXT("Timecode"), true);
+	UScriptStruct* TimecodeScriptStruct = FindObject<UScriptStruct>(nullptr, TEXT("/Script/CoreUObject.Timecode"), true);
 	check(TimecodeScriptStruct);
 	UEdGraphPin* LiveLinkSceneTimePin = CreatePin(EGPD_Input, UEdGraphSchema_K2::PC_Struct, TimecodeScriptStruct, UK2Node_EvaluateLiveLinkFrameHelper::LiveLinkSceneTimePinName);
 	LiveLinkSceneTimePin->PinFriendlyName = LOCTEXT("LiveLinkSceneTimePin", "Scene Time");
diff --git a/Engine/Plugins/Compositing/Composure/Source/Composure/Classes/CompositingElement.h b/Engine/Plugins/Compositing/Composure/Source/Composure/Classes/CompositingElement.h
--- a/Engine/Plugins/Compositing/Composure/Source/Composure/Classes/CompositingElement.h
+++ b/Engine/Plugins/Compositing/Composure/Source/Composure/Classes/CompositingElement.h
@@ -138,7 +138,7 @@ public:
 	// Editor Only
 
 private:
-	UPROPERTY(/*BlueprintReadWrite, Category = "Composure",*/ meta = (Bitmask, BitmaskEnum = ETargetUsageFlags))
+	UPROPERTY(/*BlueprintReadWrite, Category = "Composure",*/ meta = (Bitmask, BitmaskEnum = "/Script/Composure.ETargetUsageFlags"))
 	int32 FreezeFrameMask = 0x00;
 
 public:
diff --git a/Engine/Plugins/Developer/AnimationSharing/Source/AnimationSharing/Public/AnimationSharingTypes.h b/Engine/Plugins/Developer/AnimationSharing/Source/AnimationSharing/Public/AnimationSharingTypes.h
--- a/Engine/Plugins/Developer/AnimationSharing/Source/AnimationSharing/Public/AnimationSharingTypes.h
+++ b/Engine/Plugins/Developer/AnimationSharing/Source/AnimationSharing/Public/AnimationSharingTypes.h
@@ -135,7 +135,7 @@ public:
 	TObjectPtr<USkeleton> Skeleton;
 
 	/** Skeletal mesh used to setup skeletal mesh components */
-	UPROPERTY(EditAnywhere, Category = AnimationSharing, meta = (DisallowedClasses = "DestructibleMesh"))
+	UPROPERTY(EditAnywhere, Category = AnimationSharing, meta = (DisallowedClasses = "/Script/ApexDestruction.DestructibleMesh"))
 	TObjectPtr<USkeletalMesh> SkeletalMesh;
 
 	/** Animation blueprint used to perform the blending between states */
diff --git a/Engine/Plugins/Developer/Concert/ConcertMain/Source/Concert/Public/ConcertSettings.h b/Engine/Plugins/Developer/Concert/ConcertMain/Source/Concert/Public/ConcertSettings.h
--- a/Engine/Plugins/Developer/Concert/ConcertMain/Source/Concert/Public/ConcertSettings.h
+++ b/Engine/Plugins/Developer/Concert/ConcertMain/Source/Concert/Public/ConcertSettings.h
@@ -220,11 +220,11 @@ struct FConcertClientSettings
 	FLinearColor AvatarColor;
 
 	/** The desktop representation of this editor's user to other connected users */
-	UPROPERTY(config, EditAnywhere, NoClear, Category = "Client Settings", meta = (MetaClass = "ConcertClientDesktopPresenceActor"))
+	UPROPERTY(config, EditAnywhere, NoClear, Category = "Client Settings", meta = (MetaClass = "/Script/ConcertSyncClient.ConcertClientDesktopPresenceActor"))
 	FSoftClassPath DesktopAvatarActorClass;
 
 	/** The VR representation of this editor's user to other connected users */
-	UPROPERTY(config, EditAnywhere, NoClear, Category = "Client Settings", meta = (MetaClass = "ConcertClientVRPresenceActor", DisplayName = "VR Avatar Actor Class"))
+	UPROPERTY(config, EditAnywhere, NoClear, Category = "Client Settings", meta = (MetaClass = "/Script/ConcertSyncClient.ConcertClientVRPresenceActor", DisplayName = "VR Avatar Actor Class"))
 	FSoftClassPath VRAvatarActorClass;
 
 	/** The port to use to reach the server with static endpoints when launched through the editor. This port will be used over the unicast endpoint port in the UDP Messagging settings if non 0 when transferring the editor settings to the launched server. */
diff --git a/Engine/Plugins/Developer/Concert/ConcertSync/ConcertSyncCore/Source/ConcertSyncCore/Public/ConcertSyncSettings.h b/Engine/Plugins/Developer/Concert/ConcertSync/ConcertSyncCore/Source/ConcertSyncCore/Public/ConcertSyncSettings.h
--- a/Engine/Plugins/Developer/Concert/ConcertSync/ConcertSyncCore/Source/ConcertSyncCore/Public/ConcertSyncSettings.h
+++ b/Engine/Plugins/Developer/Concert/ConcertSync/ConcertSyncCore/Source/ConcertSyncCore/Public/ConcertSyncSettings.h
@@ -93,7 +93,7 @@ public:
 	/**
 	 * Array of transient class properties that we should send transaction updates for even if usually filtered out.
 	 */
-	UPROPERTY(config, EditAnywhere, Category="Transaction Settings", meta=(AllowedClasses="Property"))
+	UPROPERTY(config, EditAnywhere, Category="Transaction Settings", meta=(AllowedClasses="/Script/CoreUObject.Property"))
 	TArray<TFieldPath<FProperty>> AllowedTransientProperties;
 
 	/**
diff --git a/Engine/Plugins/Developer/TraceSourceFiltering/Source/SourceFilteringEditor/Private/FilterSetObject.cpp b/Engine/Plugins/Developer/TraceSourceFiltering/Source/SourceFilteringEditor/Private/FilterSetObject.cpp
--- a/Engine/Plugins/Developer/TraceSourceFiltering/Source/SourceFilteringEditor/Private/FilterSetObject.cpp
+++ b/Engine/Plugins/Developer/TraceSourceFiltering/Source/SourceFilteringEditor/Private/FilterSetObject.cpp
@@ -20,7 +20,7 @@
 
 #define LOCTEXT_NAMESPACE "FilterSetObject"
 
-FFilterSetObject::FFilterSetObject(IDataSourceFilterSetInterface& InFilterSet, IDataSourceFilterInterface& InFilter, const TArray<TSharedPtr<IFilterObject>>& InChildFilters, TSharedRef<ISessionSourceFilterService> InSessionFilterService) : ChildFilters(InChildFilters), SessionFilterService(InSessionFilterService), FilterSetModeEnumPtr(FindObject<UEnum>(ANY_PACKAGE, TEXT("EFilterSetMode"), true))
+FFilterSetObject::FFilterSetObject(IDataSourceFilterSetInterface& InFilterSet, IDataSourceFilterInterface& InFilter, const TArray<TSharedPtr<IFilterObject>>& InChildFilters, TSharedRef<ISessionSourceFilterService> InSessionFilterService) : ChildFilters(InChildFilters), SessionFilterService(InSessionFilterService), FilterSetModeEnumPtr(FindObject<UEnum>(nullptr, TEXT("/Script/SourceFilteringCore.EFilterSetMode"), true))
 {
 	WeakFilter = &InFilter;
 	WeakFilterSet = &InFilterSet;
diff --git a/Engine/Plugins/Developer/TraceSourceFiltering/Source/SourceFilteringEditor/Private/SUserTraceFilteringWidget.cpp b/Engine/Plugins/Developer/TraceSourceFiltering/Source/SourceFilteringEditor/Private/SUserTraceFilteringWidget.cpp
--- a/Engine/Plugins/Developer/TraceSourceFiltering/Source/SourceFilteringEditor/Private/SUserTraceFilteringWidget.cpp
+++ b/Engine/Plugins/Developer/TraceSourceFiltering/Source/SourceFilteringEditor/Private/SUserTraceFilteringWidget.cpp
@@ -376,7 +376,7 @@ TSharedPtr<SWidget> SUserTraceFilteringWidget::OnContextMenuOpening()
 				FText LabelTextFormat = LOCTEXT("MakeFilterSetLabel", "{0}");
 				FText ToolTipTextFormat = LOCTEXT("MakeFilterSetTooltip", "Creates a new filter set, containing this filter, with the {0} operator");
 
-				const UEnum* EnumPtr = FindObject<UEnum>(ANY_PACKAGE, TEXT("EFilterSetMode"), true);
+				const UEnum* EnumPtr = FindObject<UEnum>(nullptr, TEXT("/Script/SourceFilteringCore.EFilterSetMode"), true);
 				for (EFilterSetMode Mode : TEnumRange<EFilterSetMode>())
 				{
 					FText ModeText = EnumPtr->GetDisplayNameTextByValue((int64)Mode);
diff --git a/Engine/Plugins/Developer/TraceSourceFiltering/Source/SourceFilteringTrace/Private/SourceFilterCollection.cpp b/Engine/Plugins/Developer/TraceSourceFiltering/Source/SourceFilteringTrace/Private/SourceFilterCollection.cpp
--- a/Engine/Plugins/Developer/TraceSourceFiltering/Source/SourceFilteringTrace/Private/SourceFilterCollection.cpp
+++ b/Engine/Plugins/Developer/TraceSourceFiltering/Source/SourceFilteringTrace/Private/SourceFilterCollection.cpp
@@ -196,7 +196,7 @@ void USourceFilterCollection::AddFiltersFromPreset(const TArray<FString>& ClassN
 	AsyncTask(ENamedThreads::GameThread,
 		[this, ClassNames, ChildToParentIndices]()
 	{
-		const UEnum* EnumPtr = FindObject<UEnum>(ANY_PACKAGE, TEXT("EFilterSetMode"), true);
+		const UEnum* EnumPtr = FindObject<UEnum>(nullptr, TEXT("/Script/SourceFilteringCore.EFilterSetMode"), true);
 
 		TArray<UDataSourceFilter*> CreatedFilters;
 		TArray<bool> MovedFlags;
diff --git a/Engine/Plugins/Developer/TraceSourceFiltering/Source/SourceFilteringTrace/Private/SourceFilterSetup.cpp b/Engine/Plugins/Developer/TraceSourceFiltering/Source/SourceFilteringTrace/Private/SourceFilterSetup.cpp
--- a/Engine/Plugins/Developer/TraceSourceFiltering/Source/SourceFilteringTrace/Private/SourceFilterSetup.cpp
+++ b/Engine/Plugins/Developer/TraceSourceFiltering/Source/SourceFilteringTrace/Private/SourceFilterSetup.cpp
@@ -247,7 +247,7 @@ bool FSourceFilterSetup::DetermineDefaultPassingValueForSet(const UDataSourceFil
 
 void FSourceFilterSetup::OutputFilterSetState(const FFilterSet& FilterSet, FString& OutputString, int32 Depth)
 {
-	const UEnum* FilterModeEnum = FindObject<UEnum>(ANY_PACKAGE, TEXT("EFilterSetMode"));
+	const UEnum* FilterModeEnum = FindObject<UEnum>(nullptr, TEXT("/Script/SourceFilteringCore.EFilterSetMode"));
 
 	FString NewLineValue = TEXT("\n");
 	FString TabValue = TEXT("\t");
diff --git a/Engine/Plugins/Editor/GuidedTutorials/Source/IntroTutorials/Private/EditorTutorialSettings.h b/Engine/Plugins/Editor/GuidedTutorials/Source/IntroTutorials/Private/EditorTutorialSettings.h
--- a/Engine/Plugins/Editor/GuidedTutorials/Source/IntroTutorials/Private/EditorTutorialSettings.h
+++ b/Engine/Plugins/Editor/GuidedTutorials/Source/IntroTutorials/Private/EditorTutorialSettings.h
@@ -24,11 +24,11 @@ struct FTutorialContext
 	FString BrowserFilter;
 
 	/** The tutorial to use in this context to let the user know there is a tutorial available */
-	UPROPERTY(EditAnywhere, Category = "Tutorials", meta = (MetaClass = "EditorTutorial"))
+	UPROPERTY(EditAnywhere, Category = "Tutorials", meta = (MetaClass = "/Script/IntroTutorials.EditorTutorial"))
 	FSoftClassPath AttractTutorial;
 
 	/** The tutorial to use in this context when the user chooses to launch */
-	UPROPERTY(EditAnywhere, Category = "Tutorials", meta = (MetaClass = "EditorTutorial"))
+	UPROPERTY(EditAnywhere, Category = "Tutorials", meta = (MetaClass = "/Script/IntroTutorials.EditorTutorial"))
 	FSoftClassPath LaunchTutorial;
 };
 
@@ -47,7 +47,7 @@ class UEditorTutorialSettings : public UObject
 	TArray<FTutorialCategory> Categories;
 
 	/** Tutorial to start on Editor startup */
-	UPROPERTY(Config, EditAnywhere, Category="Tutorials", meta=(MetaClass="EditorTutorial"))
+	UPROPERTY(Config, EditAnywhere, Category="Tutorials", meta=(MetaClass="/Script/IntroTutorials.EditorTutorial"))
 	FSoftClassPath StartupTutorial;
 
 	/** Tutorials used in various contexts - e.g. the various asset editors */
diff --git a/Engine/Plugins/Editor/GuidedTutorials/Source/IntroTutorials/Private/TutorialSettings.h b/Engine/Plugins/Editor/GuidedTutorials/Source/IntroTutorials/Private/TutorialSettings.h
--- a/Engine/Plugins/Editor/GuidedTutorials/Source/IntroTutorials/Private/TutorialSettings.h
+++ b/Engine/Plugins/Editor/GuidedTutorials/Source/IntroTutorials/Private/TutorialSettings.h
@@ -20,6 +20,6 @@ class UTutorialSettings : public UObject
 	TArray<FTutorialCategory> Categories;
 
 	/** Tutorial to start on project startup */
-	UPROPERTY(Config, EditAnywhere, Category="Tutorials", meta=(MetaClass="EditorTutorial"))
+	UPROPERTY(Config, EditAnywhere, Category="Tutorials", meta=(MetaClass="/Script/IntroTutorials.EditorTutorial"))
 	FSoftClassPath StartupTutorial;
 };
diff --git a/Engine/Plugins/Editor/GuidedTutorials/Source/IntroTutorials/Public/EditorTutorial.h b/Engine/Plugins/Editor/GuidedTutorials/Source/IntroTutorials/Public/EditorTutorial.h
--- a/Engine/Plugins/Editor/GuidedTutorials/Source/IntroTutorials/Public/EditorTutorial.h
+++ b/Engine/Plugins/Editor/GuidedTutorials/Source/IntroTutorials/Public/EditorTutorial.h
@@ -77,7 +77,7 @@ struct INTROTUTORIALS_API FTutorialCategory
 	FString Icon;
 
 	/** Texture for this tutorial, used when presented to the user in the tutorial browser. */
-	UPROPERTY(EditAnywhere, Category="Tutorial", meta=(AllowedClasses="Texture2D"))
+	UPROPERTY(EditAnywhere, Category="Tutorial", meta=(AllowedClasses="/Script/Engine.Texture2D"))
 	FSoftObjectPath Texture;
 };
 
@@ -266,11 +266,11 @@ class INTROTUTORIALS_API UEditorTutorial : public UObject
 	TArray<FTutorialStage> Stages;
 
 	/** Tutorial to optionally chain back to if the "back" button is clicked on the first stage */
-	UPROPERTY(EditAnywhere, Category = "Tutorial", meta = (MetaClass = "EditorTutorial"))
+	UPROPERTY(EditAnywhere, Category = "Tutorial", meta = (MetaClass = "/Script/IntroTutorials.EditorTutorial"))
 	FSoftClassPath PreviousTutorial;
 
 	/** Tutorial to optionally chain onto after this tutorial completes */
-	UPROPERTY(EditAnywhere, Category="Tutorial", meta=(MetaClass="EditorTutorial"))
+	UPROPERTY(EditAnywhere, Category="Tutorial", meta=(MetaClass="/Script/IntroTutorials.EditorTutorial"))
 	FSoftClassPath NextTutorial;
 
 	/**
diff --git a/Engine/Plugins/EnhancedInput/Source/EnhancedInput/Public/InputActionValue.h b/Engine/Plugins/EnhancedInput/Source/EnhancedInput/Public/InputActionValue.h
--- a/Engine/Plugins/EnhancedInput/Source/EnhancedInput/Public/InputActionValue.h
+++ b/Engine/Plugins/EnhancedInput/Source/EnhancedInput/Public/InputActionValue.h
@@ -20,7 +20,7 @@ enum class EInputActionValueType : uint8
 };
 
 
-USTRUCT(BlueprintType, meta = (HasNativeMake = "EnhancedInput.EnhancedInputLibrary.MakeInputActionValue", HasNativeBreak = "EnhancedInput.EnhancedInputLibrary.BreakInputActionValue"))
+USTRUCT(BlueprintType, meta = (HasNativeMake = "/Script/EnhancedInput.EnhancedInputLibrary.MakeInputActionValue", HasNativeBreak = "/Script/EnhancedInput.EnhancedInputLibrary.BreakInputActionValue"))
 struct FInputActionValue
 {
 	GENERATED_BODY()
diff --git a/Engine/Plugins/Enterprise/VariantManagerContent/Source/VariantManagerContent/Public/LevelVariantSetsActor.h b/Engine/Plugins/Enterprise/VariantManagerContent/Source/VariantManagerContent/Public/LevelVariantSetsActor.h
--- a/Engine/Plugins/Enterprise/VariantManagerContent/Source/VariantManagerContent/Public/LevelVariantSetsActor.h
+++ b/Engine/Plugins/Enterprise/VariantManagerContent/Source/VariantManagerContent/Public/LevelVariantSetsActor.h
@@ -40,7 +40,7 @@ public:
 
 public:
 
-	UPROPERTY(EditAnywhere, BlueprintReadOnly, Category="LevelVariantSets", meta=(AllowedClasses="LevelVariantSets"))
+	UPROPERTY(EditAnywhere, BlueprintReadOnly, Category="LevelVariantSets", meta=(AllowedClasses="/Script/VariantManagerContent.LevelVariantSets"))
 	FSoftObjectPath LevelVariantSets;
 
 private:
diff --git a/Engine/Plugins/Experimental/Animation/MotionTrajectory/Source/MotionTrajectory/Public/MotionTrajectory.h b/Engine/Plugins/Experimental/Animation/MotionTrajectory/Source/MotionTrajectory/Public/MotionTrajectory.h
--- a/Engine/Plugins/Experimental/Animation/MotionTrajectory/Source/MotionTrajectory/Public/MotionTrajectory.h
+++ b/Engine/Plugins/Experimental/Animation/MotionTrajectory/Source/MotionTrajectory/Public/MotionTrajectory.h
@@ -30,7 +30,7 @@ struct MOTIONTRAJECTORY_API FMotionTrajectorySettings
 
 	// Trajectory mask for specifying a time and/or distance domain
 	// Default: Time Domain of 2.0 seconds
-	UPROPERTY(EditAnywhere, BlueprintReadWrite, Category="Settings", meta=(Bitmask, BitmaskEnum=ETrajectorySampleDomain))
+	UPROPERTY(EditAnywhere, BlueprintReadWrite, Category="Settings", meta=(Bitmask, BitmaskEnum="/Script/Engine.ETrajectorySampleDomain"))
 	int32 Domain = 1; // ETrajectorySampleDomain::Time
 
 	// Sample time horizon
diff --git a/Engine/Plugins/Experimental/CommonUI/Source/CommonInput/Public/CommonInputBaseTypes.h b/Engine/Plugins/Experimental/CommonUI/Source/CommonInput/Public/CommonInputBaseTypes.h
--- a/Engine/Plugins/Experimental/CommonUI/Source/CommonInput/Public/CommonInputBaseTypes.h
+++ b/Engine/Plugins/Experimental/CommonUI/Source/CommonInput/Public/CommonInputBaseTypes.h
@@ -96,10 +96,10 @@ public:
 	virtual bool NeedsLoadForServer() const override;
 
 public:
-	UPROPERTY(EditDefaultsOnly, Category = "Properties", meta = (RowType = CommonInputActionDataBase))
+	UPROPERTY(EditDefaultsOnly, Category = "Properties", meta = (RowType = "/Script/CommonUI.CommonInputActionDataBase"))
 	FDataTableRowHandle DefaultClickAction;
 
-	UPROPERTY(EditDefaultsOnly, Category = "Properties", meta = (RowType = CommonInputActionDataBase))
+	UPROPERTY(EditDefaultsOnly, Category = "Properties", meta = (RowType = "/Script/CommonUI.CommonInputActionDataBase"))
 	FDataTableRowHandle DefaultBackAction;
 };
 
diff --git a/Engine/Plugins/Experimental/CommonUI/Source/CommonUI/Public/CommonActionWidget.h b/Engine/Plugins/Experimental/CommonUI/Source/CommonUI/Public/CommonActionWidget.h
--- a/Engine/Plugins/Experimental/CommonUI/Source/CommonUI/Public/CommonActionWidget.h
+++ b/Engine/Plugins/Experimental/CommonUI/Source/CommonUI/Public/CommonActionWidget.h
@@ -79,7 +79,7 @@ protected:
 	 * that you need to represent as a single entry in the UI.  For example - zoom, might be mouse wheel up or down, but you just need to
 	 * show a single icon for Up & Down on the mouse, this solves that problem.
 	 */
-	UPROPERTY(EditAnywhere, BlueprintReadOnly, Category = CommonActionWidget, meta = (RowType = CommonInputActionDataBase, TitleProperty = "RowName"))
+	UPROPERTY(EditAnywhere, BlueprintReadOnly, Category = CommonActionWidget, meta = (RowType = "/Script/CommonUI.CommonInputActionDataBase", TitleProperty = "RowName"))
 	TArray<FDataTableRowHandle> InputActions;
 
 	//@todo DanH: Create clearer split between support for the new & legacy system in here
diff --git a/Engine/Plugins/Experimental/CommonUI/Source/CommonUI/Public/CommonButtonBase.h b/Engine/Plugins/Experimental/CommonUI/Source/CommonUI/Public/CommonButtonBase.h
--- a/Engine/Plugins/Experimental/CommonUI/Source/CommonUI/Public/CommonButtonBase.h
+++ b/Engine/Plugins/Experimental/CommonUI/Source/CommonUI/Public/CommonButtonBase.h
@@ -642,7 +642,7 @@ public:
 	 *	The input action that is bound to this button. The common input manager will trigger this button to 
 	 *	click if the action was pressed 
 	 */
-	UPROPERTY(EditAnywhere, BlueprintReadOnly, Category = Input, meta = (ExposeOnSpawn = true, RowType = CommonInputActionDataBase))
+	UPROPERTY(EditAnywhere, BlueprintReadOnly, Category = Input, meta = (ExposeOnSpawn = true, RowType = "/Script/CommonUI.CommonInputActionDataBase"))
 	FDataTableRowHandle TriggeringInputAction;
 
 	/**
diff --git a/Engine/Plugins/Experimental/CommonUI/Source/CommonUI/Public/CommonTabListWidgetBase.h b/Engine/Plugins/Experimental/CommonUI/Source/CommonUI/Public/CommonTabListWidgetBase.h
--- a/Engine/Plugins/Experimental/CommonUI/Source/CommonUI/Public/CommonTabListWidgetBase.h
+++ b/Engine/Plugins/Experimental/CommonUI/Source/CommonUI/Public/CommonTabListWidgetBase.h
@@ -170,11 +170,11 @@ protected:
 	void HandleTabRemoval(FName TabNameID, UCommonButtonBase* TabButton);
 
 	/** The input action to listen for causing the next tab to be selected */
-	UPROPERTY(EditAnywhere, BlueprintReadOnly, Category = TabList, meta = (RowType = CommonInputActionDataBase))
+	UPROPERTY(EditAnywhere, BlueprintReadOnly, Category = TabList, meta = (RowType = "/Script/CommonUI.CommonInputActionDataBase"))
 	FDataTableRowHandle NextTabInputActionData;
 
 	/** The input action to listen for causing the previous tab to be selected */
-	UPROPERTY(EditAnywhere, BlueprintReadOnly, Category = TabList, meta = (RowType = CommonInputActionDataBase))
+	UPROPERTY(EditAnywhere, BlueprintReadOnly, Category = TabList, meta = (RowType = "/Script/CommonUI.CommonInputActionDataBase"))
 	FDataTableRowHandle PreviousTabInputActionData;
 
 	/** Whether to register to handle tab list input immediately upon construction */
diff --git a/Engine/Plugins/Experimental/CommonUI/Source/CommonUI/Public/CommonUIRichTextData.h b/Engine/Plugins/Experimental/CommonUI/Source/CommonUI/Public/CommonUIRichTextData.h
--- a/Engine/Plugins/Experimental/CommonUI/Source/CommonUI/Public/CommonUIRichTextData.h
+++ b/Engine/Plugins/Experimental/CommonUI/Source/CommonUI/Public/CommonUIRichTextData.h
@@ -16,7 +16,7 @@ struct COMMONUI_API FRichTextIconData : public FTableRowBase
 	UPROPERTY(EditDefaultsOnly, Category = "RichText Icon")
 	FText DisplayName;
 
-	UPROPERTY(EditDefaultsOnly, Category = "RichText Icon", meta = (DisplayThumbnail = "true", DisplayName = "Image", AllowedClasses = "Texture2D,MaterialInterface,SlateTextureAtlasInterface", DisallowedClasses = "MediaTexture"))
+	UPROPERTY(EditDefaultsOnly, Category = "RichText Icon", meta = (DisplayThumbnail = "true", DisplayName = "Image", AllowedClasses = "/Script/Engine.Texture2D,/Script/Engine.MaterialInterface,/Script/Engine.SlateTextureAtlasInterface", DisallowedClasses = "/Script/MediaAssets.MediaTexture"))
 	TSoftObjectPtr<UObject> ResourceObject;
 
 	UPROPERTY(EditDefaultsOnly, Category = "RichText Icon")
@@ -39,6 +39,6 @@ public:
 	const TMap<FName, uint8*>& GetIconMap() const { return InlineIconSet->GetRowMap(); }
 
 private:
-	UPROPERTY(EditDefaultsOnly, Category = "Inline Icons", meta = (RowType = RichTextIconData))
+	UPROPERTY(EditDefaultsOnly, Category = "Inline Icons", meta = (RowType = "/Script/CommonUI.RichTextIconData"))
 	UDataTable* InlineIconSet;
 };
\ No newline at end of file
diff --git a/Engine/Plugins/Experimental/CommonUI/Source/CommonUI/Public/CommonUISettings.h b/Engine/Plugins/Experimental/CommonUI/Source/CommonUI/Public/CommonUISettings.h
--- a/Engine/Plugins/Experimental/CommonUI/Source/CommonUI/Public/CommonUISettings.h
+++ b/Engine/Plugins/Experimental/CommonUI/Source/CommonUI/Public/CommonUISettings.h
@@ -55,7 +55,7 @@ private:
 	bool bAutoLoadData;
 
 	/** The Default Image Resource, newly created CommonImage Widgets will use this style. */
-	UPROPERTY(config, EditAnywhere, Category = "Image", meta = (AllowedClasses = "Texture2D,MaterialInterface"))
+	UPROPERTY(config, EditAnywhere, Category = "Image", meta = (AllowedClasses = "/Script/Engine.Texture2D,/Script/Engine.MaterialInterface"))
 	TSoftObjectPtr<UObject> DefaultImageResourceObject;
 
 	/** The Default Throbber Material, newly created CommonLoadGuard Widget will use this style. */
diff --git a/Engine/Plugins/Experimental/ComputeFramework/Source/ComputeFramework/Public/ComputeFramework/ComputeKernel.h b/Engine/Plugins/Experimental/ComputeFramework/Source/ComputeFramework/Public/ComputeFramework/ComputeKernel.h
--- a/Engine/Plugins/Experimental/ComputeFramework/Source/ComputeFramework/Public/ComputeFramework/ComputeKernel.h
+++ b/Engine/Plugins/Experimental/ComputeFramework/Source/ComputeFramework/Public/ComputeFramework/ComputeKernel.h
@@ -43,7 +43,7 @@ public:
 	TObjectPtr<UComputeKernelSource> KernelSource = nullptr;
 
 	/** Specifying certain memory access flags allows for optimizations such as kernel fusing. */
-	UPROPERTY(VisibleAnywhere, BlueprintReadOnly, meta = (Bitmask, BitmaskEnum = "EComputeKernelFlags"), Category = "Kernel")
+	UPROPERTY(VisibleAnywhere, BlueprintReadOnly, meta = (Bitmask, BitmaskEnum = "/Script/ComputeFramework.EComputeKernelFlags"), Category = "Kernel")
 	int32 KernelFlags = 0;
 
 	/** Permutation overrides on the kernel. */
diff --git a/Engine/Plugins/Experimental/GameFeatures/Source/GameFeatures/Public/GameFeaturesSubsystemSettings.h b/Engine/Plugins/Experimental/GameFeatures/Source/GameFeatures/Public/GameFeaturesSubsystemSettings.h
--- a/Engine/Plugins/Experimental/GameFeatures/Source/GameFeatures/Public/GameFeaturesSubsystemSettings.h
+++ b/Engine/Plugins/Experimental/GameFeatures/Source/GameFeatures/Public/GameFeaturesSubsystemSettings.h
@@ -23,7 +23,7 @@ public:
 	static const FName LoadStateServer;
 
 	/** Name of a singleton class to spawn as the game feature project policy. If empty, it will spawn the default one (UDefaultGameFeaturesProjectPolicies) */
-	UPROPERTY(config, EditAnywhere, Category=DefaultClasses, meta=(MetaClass="GameFeaturesProjectPolicies", DisplayName="Game Feature Project Policy Class", ConfigRestartRequired=true))
+	UPROPERTY(config, EditAnywhere, Category=DefaultClasses, meta=(MetaClass="/Script/GameFeatures.GameFeaturesProjectPolicies", DisplayName="Game Feature Project Policy Class", ConfigRestartRequired=true))
 	FSoftClassPath GameFeaturesManagerClassName;
 
 	/** List of plugins that are forcibly disabled (e.g., via a hotfix) */
diff --git a/Engine/Plugins/Experimental/GeometryCollectionPlugin/Source/GeometryCollectionTracks/Public/MovieSceneGeometryCollectionSection.h b/Engine/Plugins/Experimental/GeometryCollectionPlugin/Source/GeometryCollectionTracks/Public/MovieSceneGeometryCollectionSection.h
--- a/Engine/Plugins/Experimental/GeometryCollectionPlugin/Source/GeometryCollectionTracks/Public/MovieSceneGeometryCollectionSection.h
+++ b/Engine/Plugins/Experimental/GeometryCollectionPlugin/Source/GeometryCollectionTracks/Public/MovieSceneGeometryCollectionSection.h
@@ -48,7 +48,7 @@ struct FMovieSceneGeometryCollectionParams
 	}
 
 	/** The animation this section plays */
-	UPROPERTY(EditAnywhere, Category="GeometryCollection", meta=(AllowedClasses = "GeometryCollectionCache"))
+	UPROPERTY(EditAnywhere, Category="GeometryCollection", meta=(AllowedClasses = "/Script/GeometryCollectionEngine.GeometryCollectionCache"))
 	FSoftObjectPath GeometryCollectionCache;
 
 	/** The offset into the beginning of the animation clip */
diff --git a/Engine/Plugins/Experimental/MeshModelingToolsetExp/Source/MeshModelingToolsExp/Public/BakeMeshAttributeMapsTool.h b/Engine/Plugins/Experimental/MeshModelingToolsetExp/Source/MeshModelingToolsExp/Public/BakeMeshAttributeMapsTool.h
--- a/Engine/Plugins/Experimental/MeshModelingToolsetExp/Source/MeshModelingToolsExp/Public/BakeMeshAttributeMapsTool.h
+++ b/Engine/Plugins/Experimental/MeshModelingToolsetExp/Source/MeshModelingToolsExp/Public/BakeMeshAttributeMapsTool.h
@@ -52,7 +52,7 @@ class MESHMODELINGTOOLSEXP_API UBakeMeshAttributeMapsToolProperties : public UIn
 
 public:
 	/** The bake output types to generate */
-	UPROPERTY(EditAnywhere, Category = BakeOutput, meta = (DisplayName = "Output Types", Bitmask, BitmaskEnum = EBakeMapType,
+	UPROPERTY(EditAnywhere, Category = BakeOutput, meta = (DisplayName = "Output Types", Bitmask, BitmaskEnum = "/Script/MeshModelingToolsExp.EBakeMapType",
 		ValidEnumValues="TangentSpaceNormal, AmbientOcclusion, BentNormal, Curvature, Texture, ObjectSpaceNormal, FaceNormal, Position, MaterialID, MultiTexture, VertexColor"))
 	int32 MapTypes = static_cast<int32>(EBakeMapType::None);
 
diff --git a/Engine/Plugins/Experimental/MeshModelingToolsetExp/Source/MeshModelingToolsExp/Public/BakeMeshAttributeVertexTool.h b/Engine/Plugins/Experimental/MeshModelingToolsetExp/Source/MeshModelingToolsExp/Public/BakeMeshAttributeVertexTool.h
--- a/Engine/Plugins/Experimental/MeshModelingToolsetExp/Source/MeshModelingToolsExp/Public/BakeMeshAttributeVertexTool.h
+++ b/Engine/Plugins/Experimental/MeshModelingToolsetExp/Source/MeshModelingToolsExp/Public/BakeMeshAttributeVertexTool.h
@@ -64,31 +64,31 @@ public:
 	EBakeVertexOutput OutputMode = EBakeVertexOutput::RGBA;
 
 	/** The bake output type to generate */
-	UPROPERTY(EditAnywhere, Category = BakeOutput, meta=(Bitmask, BitmaskEnum = EBakeMapType,
+	UPROPERTY(EditAnywhere, Category = BakeOutput, meta=(Bitmask, BitmaskEnum = "/Script/MeshModelingToolsExp.EBakeMapType",
 		ValidEnumValues="TangentSpaceNormal, AmbientOcclusion, BentNormal, Curvature, Texture, ObjectSpaceNormal, FaceNormal, Position, MaterialID, MultiTexture",
 		EditCondition="OutputMode == EBakeVertexOutput::RGBA", EditConditionHides))
 	int32 OutputType = static_cast<int32>(EBakeMapType::TangentSpaceNormal);
 
 	/** The bake output type to generate in the Red channel */
-	UPROPERTY(EditAnywhere, Category = BakeOutput, meta=(Bitmask, BitmaskEnum = EBakeMapType,
+	UPROPERTY(EditAnywhere, Category = BakeOutput, meta=(Bitmask, BitmaskEnum = "/Script/MeshModelingToolsExp.EBakeMapType",
 		ValidEnumValues="None, AmbientOcclusion, Curvature",
 		EditCondition="OutputMode == EBakeVertexOutput::PerChannel", EditConditionHides))
 	int32 OutputTypeR = static_cast<int32>(EBakeMapType::None);
 
 	/** The bake output type to generate in the Green channel */
-	UPROPERTY(EditAnywhere, Category = BakeOutput, meta=(Bitmask, BitmaskEnum = EBakeMapType,
+	UPROPERTY(EditAnywhere, Category = BakeOutput, meta=(Bitmask, BitmaskEnum = "/Script/MeshModelingToolsExp.EBakeMapType",
 		ValidEnumValues="None, AmbientOcclusion, Curvature",
 		EditCondition="OutputMode == EBakeVertexOutput::PerChannel", EditConditionHides))
 	int32 OutputTypeG = static_cast<int32>(EBakeMapType::None);
 
 	/** The bake output type to generate in the Blue channel */
-	UPROPERTY(EditAnywhere, Category = BakeOutput, meta=(Bitmask, BitmaskEnum = EBakeMapType,
+	UPROPERTY(EditAnywhere, Category = BakeOutput, meta=(Bitmask, BitmaskEnum = "/Script/MeshModelingToolsExp.EBakeMapType",
 		ValidEnumValues="None, AmbientOcclusion, Curvature",
 		EditCondition="OutputMode == EBakeVertexOutput::PerChannel", EditConditionHides))
 	int32 OutputTypeB = static_cast<int32>(EBakeMapType::None);
 
 	/** The bake output type to generate in the Alpha channel */
-	UPROPERTY(EditAnywhere, Category = BakeOutput, meta=(Bitmask, BitmaskEnum = EBakeMapType,
+	UPROPERTY(EditAnywhere, Category = BakeOutput, meta=(Bitmask, BitmaskEnum = "/Script/MeshModelingToolsExp.EBakeMapType",
 		ValidEnumValues="None, AmbientOcclusion, Curvature",
 		EditCondition="OutputMode == EBakeVertexOutput::PerChannel", EditConditionHides))
 	int32 OutputTypeA = static_cast<int32>(EBakeMapType::None);
diff --git a/Engine/Plugins/Experimental/MeshModelingToolsetExp/Source/MeshModelingToolsExp/Public/BakeMultiMeshAttributeMapsTool.h b/Engine/Plugins/Experimental/MeshModelingToolsetExp/Source/MeshModelingToolsExp/Public/BakeMultiMeshAttributeMapsTool.h
--- a/Engine/Plugins/Experimental/MeshModelingToolsetExp/Source/MeshModelingToolsExp/Public/BakeMultiMeshAttributeMapsTool.h
+++ b/Engine/Plugins/Experimental/MeshModelingToolsetExp/Source/MeshModelingToolsExp/Public/BakeMultiMeshAttributeMapsTool.h
@@ -39,7 +39,7 @@ class MESHMODELINGTOOLSEXP_API UBakeMultiMeshAttributeMapsToolProperties : publi
 
 public:
 	/** The map types to generate */
-	UPROPERTY(EditAnywhere, Category = BakeOutput, meta=(DisplayName="Output Types", Bitmask, BitmaskEnum=EBakeMapType,
+	UPROPERTY(EditAnywhere, Category = BakeOutput, meta=(DisplayName="Output Types", Bitmask, BitmaskEnum= "/Script/MeshModelingToolsExp.EBakeMapType",
 		ValidEnumValues="TangentSpaceNormal, ObjectSpaceNormal, Position, Texture"))
 	int32 MapTypes = (int32) EBakeMapType::None;
 
diff --git a/Engine/Plugins/Experimental/PythonScriptPlugin/Source/PythonScriptPlugin/Private/PyTest.h b/Engine/Plugins/Experimental/PythonScriptPlugin/Source/PythonScriptPlugin/Private/PyTest.h
--- a/Engine/Plugins/Experimental/PythonScriptPlugin/Source/PythonScriptPlugin/Private/PyTest.h
+++ b/Engine/Plugins/Experimental/PythonScriptPlugin/Source/PythonScriptPlugin/Private/PyTest.h
@@ -107,7 +107,7 @@ class UPyTestStructLibrary : public UBlueprintFunctionLibrary
 	UFUNCTION(BlueprintPure, Category = "Python|Internal", meta=(ScriptMethod, DeprecatedFunction, DeprecationMessage="LegacyIsBoolSet is deprecated. Please use IsBoolSet instead."))
 	static bool LegacyIsBoolSet(const FPyTestStruct& InStruct);
 
-	UFUNCTION(BlueprintPure, Category = "Python|Internal", meta=(ScriptConstant="ConstantValue", ScriptConstantHost="PyTestStruct"))
+	UFUNCTION(BlueprintPure, Category = "Python|Internal", meta=(ScriptConstant="ConstantValue", ScriptConstantHost="/Script/PythonScriptPlugin.PyTestStruct"))
 	static int32 GetConstantValue();
 
 	UFUNCTION(BlueprintPure, Category = "Python|Internal", meta=(ScriptMethod, ScriptMethodSelfReturn, ScriptOperator="+;+="))
@@ -285,7 +285,7 @@ class UPyTestObjectLibrary : public UBlueprintFunctionLibrary
 	UFUNCTION(BlueprintPure, Category = "Python|Internal", meta=(ScriptMethod="IsBoolSet"))
 	static bool IsBoolSet(const UPyTestObject* InObj);
 
-	UFUNCTION(BlueprintPure, Category = "Python|Internal", meta=(ScriptConstant="OtherConstantValue", ScriptConstantHost="PyTestObject"))
+	UFUNCTION(BlueprintPure, Category = "Python|Internal", meta=(ScriptConstant="OtherConstantValue", ScriptConstantHost="/Script/PythonScriptPlugin.PyTestObject"))
 	static int32 GetOtherConstantValue();
 };
 
diff --git a/Engine/Plugins/Experimental/PythonScriptPlugin/Source/PythonScriptPlugin/Private/PyWrapperObject.cpp b/Engine/Plugins/Experimental/PythonScriptPlugin/Source/PythonScriptPlugin/Private/PyWrapperObject.cpp
--- a/Engine/Plugins/Experimental/PythonScriptPlugin/Source/PythonScriptPlugin/Private/PyWrapperObject.cpp
+++ b/Engine/Plugins/Experimental/PythonScriptPlugin/Source/PythonScriptPlugin/Private/PyWrapperObject.cpp
@@ -824,7 +824,7 @@ PyTypeObject InitializePyWrapperObjectType()
 				return nullptr;
 			}
 
-			static const UEnum* PropertyAccessChangeNotifyModeEnum = FindObjectChecked<UEnum>(ANY_PACKAGE, TEXT("EPropertyAccessChangeNotifyMode"));
+			static const UEnum* PropertyAccessChangeNotifyModeEnum = FindObjectChecked<UEnum>(nullptr, TEXT("/Script/CoreUObject.EPropertyAccessChangeNotifyMode"));
 
 			EPropertyAccessChangeNotifyMode NotifyMode = EPropertyAccessChangeNotifyMode::Default;
 			if (PyNotifyModeObj && !PyConversion::NativizeEnumEntry(PyNotifyModeObj, PropertyAccessChangeNotifyModeEnum, NotifyMode))
diff --git a/Engine/Plugins/Experimental/PythonScriptPlugin/Source/PythonScriptPlugin/Private/PyWrapperStruct.cpp b/Engine/Plugins/Experimental/PythonScriptPlugin/Source/PythonScriptPlugin/Private/PyWrapperStruct.cpp
--- a/Engine/Plugins/Experimental/PythonScriptPlugin/Source/PythonScriptPlugin/Private/PyWrapperStruct.cpp
+++ b/Engine/Plugins/Experimental/PythonScriptPlugin/Source/PythonScriptPlugin/Private/PyWrapperStruct.cpp
@@ -1079,7 +1079,7 @@ PyTypeObject InitializePyWrapperStructType()
 				return nullptr;
 			}
 
-			static const UEnum* PropertyAccessChangeNotifyModeEnum = FindObjectChecked<UEnum>(ANY_PACKAGE, TEXT("EPropertyAccessChangeNotifyMode"));
+			static const UEnum* PropertyAccessChangeNotifyModeEnum = FindObjectChecked<UEnum>(nullptr, TEXT("/Script/CoreUObject.EPropertyAccessChangeNotifyMode"));
 
 			EPropertyAccessChangeNotifyMode NotifyMode = EPropertyAccessChangeNotifyMode::Default;
 			if (PyNotifyModeObj && !PyConversion::NativizeEnumEntry(PyNotifyModeObj, PropertyAccessChangeNotifyModeEnum, NotifyMode))
diff --git a/Engine/Plugins/Experimental/StructUtils/Source/StructUtils/Public/InstancedStruct.h b/Engine/Plugins/Experimental/StructUtils/Source/StructUtils/Public/InstancedStruct.h
--- a/Engine/Plugins/Experimental/StructUtils/Source/StructUtils/Public/InstancedStruct.h
+++ b/Engine/Plugins/Experimental/StructUtils/Source/StructUtils/Public/InstancedStruct.h
@@ -12,10 +12,10 @@ struct FConstSharedStruct;
  * FInstancedStruct works similarly as instanced UObject* property but is USTRUCTs.
  * Example:
  *
- *	UPROPERTY(EditAnywhere, Category = Foo, meta = (BaseStruct = "TestStructBase"))
+ *	UPROPERTY(EditAnywhere, Category = Foo, meta = (BaseStruct = "/Script/ModuleName.TestStructBase"))
  *	FInstancedStruct Test;
  *
- *	UPROPERTY(EditAnywhere, Category = Foo, meta = (BaseStruct = "TestStructBase"))
+ *	UPROPERTY(EditAnywhere, Category = Foo, meta = (BaseStruct = "/Script/ModuleName.TestStructBase"))
  *	TArray<FInstancedStruct> TestArray;
  */
 USTRUCT()
diff --git a/Engine/Plugins/Experimental/VirtualProduction/LedWallCalibration/Source/LedWallCalibration/Public/LedWallArucoGenerationOptions.h b/Engine/Plugins/Experimental/VirtualProduction/LedWallCalibration/Source/LedWallCalibration/Public/LedWallArucoGenerationOptions.h
--- a/Engine/Plugins/Experimental/VirtualProduction/LedWallCalibration/Source/LedWallCalibration/Public/LedWallArucoGenerationOptions.h
+++ b/Engine/Plugins/Experimental/VirtualProduction/LedWallCalibration/Source/LedWallCalibration/Public/LedWallArucoGenerationOptions.h
@@ -71,7 +71,7 @@ public:
 	 */
 	FString ArucoDictionaryAsString() const
 	{
-		const UEnum* EnumPtr = FindObject<UEnum>((UObject*)ANY_PACKAGE, TEXT("EArucoDictionary"), true);
+		const UEnum* EnumPtr = FindObject<UEnum>(nullptr, TEXT("/Script/LedWallCalibration.EArucoDictionary"), true);
 		check(EnumPtr);
 
 #if WITH_EDITOR
diff --git a/Engine/Plugins/Experimental/VirtualProductionUtilities/Source/VPUtilities/Public/VPViewportTickableActorBase.h b/Engine/Plugins/Experimental/VirtualProductionUtilities/Source/VPUtilities/Public/VPViewportTickableActorBase.h
--- a/Engine/Plugins/Experimental/VirtualProductionUtilities/Source/VPUtilities/Public/VPViewportTickableActorBase.h
+++ b/Engine/Plugins/Experimental/VirtualProductionUtilities/Source/VPUtilities/Public/VPViewportTickableActorBase.h
@@ -45,7 +45,7 @@ public:
 	 * Game = Tick in game even if we are only ticking the viewport. Use the event Tick.
 	 * Preview = Tick if the actor is present in any editing tool like Blueprint or Material graph. Use EditorTick.
 	 */
-	UPROPERTY(EditDefaultsOnly, Category = "Actor Tick", meta = (Bitmask, BitmaskEnum = EVPViewportTickableType))
+	UPROPERTY(EditDefaultsOnly, Category = "Actor Tick", meta = (Bitmask, BitmaskEnum = "/Script/VPUtilities.EVPViewportTickableType"))
 	EVPViewportTickableFlags ViewportTickType;
 
 	/** If true, actor is ticked even if TickType==LEVELTICK_ViewportsOnly */
diff --git a/Engine/Plugins/Experimental/VirtualProductionUtilities/Source/VPUtilitiesEditor/Public/VPUtilitiesEditorSettings.h b/Engine/Plugins/Experimental/VirtualProductionUtilities/Source/VPUtilitiesEditor/Public/VPUtilitiesEditorSettings.h
--- a/Engine/Plugins/Experimental/VirtualProductionUtilities/Source/VPUtilitiesEditor/Public/VPUtilitiesEditorSettings.h
+++ b/Engine/Plugins/Experimental/VirtualProductionUtilities/Source/VPUtilitiesEditor/Public/VPUtilitiesEditorSettings.h
@@ -66,15 +66,15 @@ public:
 	uint16 OSCServerPort;
 
 	/** What EditorUtilityObject should be ran on editor launch. */
-	UPROPERTY(config, EditAnywhere, Category = "OSC", meta = (AllowedClasses = "EditorUtilityBlueprint", DisplayName = "OSC Listeners"))
+	UPROPERTY(config, EditAnywhere, Category = "OSC", meta = (AllowedClasses = "/Script/Blutility.EditorUtilityBlueprint", DisplayName = "OSC Listeners"))
 	TArray<FSoftObjectPath> StartupOSCListeners;
 
 	/** ScoutingSubsystem class to use for Blueprint helpers */
-	UPROPERTY(config, meta = (MetaClass = "VPScoutingSubsystemHelpersBase"))
+	UPROPERTY(config, meta = (MetaClass = "/Script/VPUtilitiesEditor.VPScoutingSubsystemHelpersBase"))
 	FSoftClassPath ScoutingSubsystemEditorUtilityClassPath;
 
 	/** GestureManager class to use by the ScoutingSubsystem */
-	UPROPERTY(config, meta = (MetaClass = "VPScoutingSubsystemGestureManagerBase"))
+	UPROPERTY(config, meta = (MetaClass = "/Script/VPUtilitiesEditor.VPScoutingSubsystemGestureManagerBase"))
 	FSoftClassPath GestureManagerEditorUtilityClassPath;
 
 	/** GestureManager class to use by the ScoutingSubsystem */
diff --git a/Engine/Plugins/Experimental/Water/Source/Editor/Public/WaterEditorSettings.h b/Engine/Plugins/Experimental/Water/Source/Editor/Public/WaterEditorSettings.h
--- a/Engine/Plugins/Experimental/Water/Source/Editor/Public/WaterEditorSettings.h
+++ b/Engine/Plugins/Experimental/Water/Source/Editor/Public/WaterEditorSettings.h
@@ -256,7 +256,7 @@ public:
 
 private:
 	/** Class of the water brush to be used in landscape */
-	UPROPERTY(EditAnywhere, config, Category = Brush, meta = (MetaClass = "WaterLandscapeBrush"))
+	UPROPERTY(EditAnywhere, config, Category = Brush, meta = (MetaClass = "/Script/WaterEditor.WaterLandscapeBrush"))
 	FSoftClassPath WaterManagerClassPath;
 
 	UPROPERTY(EditAnywhere, config, Category = Brush)
diff --git a/Engine/Plugins/Experimental/Water/Source/Runtime/Public/WaterRuntimeSettings.h b/Engine/Plugins/Experimental/Water/Source/Runtime/Public/WaterRuntimeSettings.h
--- a/Engine/Plugins/Experimental/Water/Source/Runtime/Public/WaterRuntimeSettings.h
+++ b/Engine/Plugins/Experimental/Water/Source/Runtime/Public/WaterRuntimeSettings.h
@@ -73,16 +73,16 @@ private:
 	UPROPERTY(EditAnywhere, config, Category = Water)
 	TSoftObjectPtr<UMaterialInterface> DefaultWaterInfoMaterial;
 	
-	UPROPERTY(EditAnywhere, Config, Category = Water,  meta = (MetaClass = "WaterBodyRiverComponent"))
+	UPROPERTY(EditAnywhere, Config, Category = Water,  meta = (MetaClass = "/Script/Water.WaterBodyRiverComponent"))
 	TSubclassOf<UWaterBodyRiverComponent> WaterBodyRiverComponentClass;
 
-	UPROPERTY(EditAnywhere, Config, Category = Water,  meta = (MetaClass = "WaterBodyLakeComponent"))
+	UPROPERTY(EditAnywhere, Config, Category = Water,  meta = (MetaClass = "/Script/Water.WaterBodyLakeComponent"))
 	TSubclassOf<UWaterBodyLakeComponent> WaterBodyLakeComponentClass;
 
-	UPROPERTY(EditAnywhere, Config, Category = Water,  meta = (MetaClass = "WaterBodyOceanComponent"))
+	UPROPERTY(EditAnywhere, Config, Category = Water,  meta = (MetaClass = "/Script/Water.WaterBodyOceanComponent"))
 	TSubclassOf<UWaterBodyOceanComponent> WaterBodyOceanComponentClass;
 
-	UPROPERTY(EditAnywhere, Config, Category = Water,  meta = (MetaClass = "WaterBodyCustomComponent"))
+	UPROPERTY(EditAnywhere, Config, Category = Water,  meta = (MetaClass = "/Script/Water.WaterBodyCustomComponent"))
 	TSubclassOf<UWaterBodyCustomComponent> WaterBodyCustomComponentClass;
 
 #if WITH_EDITORONLY_DATA
diff --git a/Engine/Plugins/FX/Niagara/Source/Niagara/Classes/NiagaraDataInterfaceSkeletalMesh.h b/Engine/Plugins/FX/Niagara/Source/Niagara/Classes/NiagaraDataInterfaceSkeletalMesh.h
--- a/Engine/Plugins/FX/Niagara/Source/Niagara/Classes/NiagaraDataInterfaceSkeletalMesh.h
+++ b/Engine/Plugins/FX/Niagara/Source/Niagara/Classes/NiagaraDataInterfaceSkeletalMesh.h
@@ -725,7 +725,7 @@ public:
 
 #if WITH_EDITORONLY_DATA
 	/** Mesh used to sample from when not overridden by a source actor from the scene. Only available in editor for previewing. This is removed in cooked builds. */
-	UPROPERTY(EditAnywhere, Category = "Mesh", meta = (DisallowedClasses = "DestructibleMesh"))
+	UPROPERTY(EditAnywhere, Category = "Mesh", meta = (DisallowedClasses = "/Script/ApexDestruction.DestructibleMesh"))
 	TSoftObjectPtr<USkeletalMesh> PreviewMesh;
 #endif
 
diff --git a/Engine/Plugins/FX/Niagara/Source/Niagara/Classes/NiagaraScript.h b/Engine/Plugins/FX/Niagara/Source/Niagara/Classes/NiagaraScript.h
--- a/Engine/Plugins/FX/Niagara/Source/Niagara/Classes/NiagaraScript.h
+++ b/Engine/Plugins/FX/Niagara/Source/Niagara/Classes/NiagaraScript.h
@@ -543,7 +543,7 @@ public:
 	FText VersionChangeDescription;
 
 	/** When used as a module, what are the appropriate script types for referencing this module?*/
-	UPROPERTY(EditAnywhere, Category = Script, meta = (Bitmask, BitmaskEnum = ENiagaraScriptUsage))
+	UPROPERTY(EditAnywhere, Category = Script, meta = (Bitmask, BitmaskEnum = "/Script/Niagara.ENiagaraScriptUsage"))
 	int32 ModuleUsageBitmask;
 
 	/** Used to break up scripts of the same Usage type in UI display.*/
diff --git a/Engine/Plugins/FX/Niagara/Source/Niagara/Private/NiagaraDebugHud.cpp b/Engine/Plugins/FX/Niagara/Source/Niagara/Private/NiagaraDebugHud.cpp
--- a/Engine/Plugins/FX/Niagara/Source/Niagara/Private/NiagaraDebugHud.cpp
+++ b/Engine/Plugins/FX/Niagara/Source/Niagara/Private/NiagaraDebugHud.cpp
@@ -1652,7 +1652,7 @@ void FNiagaraDebugHud::DrawOverview(class FNiagaraWorldManager* WorldManager, FC
 			DrawCanvas->DrawShadowedString(TextLocation.X, TextLocation.Y, TEXT("Niagara DebugHud"), Font, HeadingColor);
 			TextLocation.Y += fAdvanceHeight;
 
-			const UEnum* EnumPtr = FindObjectChecked<UEnum>(ANY_PACKAGE, TEXT("ENiagaraDebugHUDOverviewMode"), true);
+			const UEnum* EnumPtr = FindObjectChecked<UEnum>(nullptr, TEXT("/Script/Niagara.ENiagaraDebugHUDOverviewMode"), true);
 			DrawCanvas->DrawShadowedString(TextLocation.X, TextLocation.Y, *EnumPtr->GetDisplayNameTextByValue((int32)Settings.OverviewMode).ToString(), Font, HeadingColor);
 			TextLocation.Y += fAdvanceHeight;
 
diff --git a/Engine/Plugins/FX/Niagara/Source/Niagara/Private/NiagaraSystemInstance.cpp b/Engine/Plugins/FX/Niagara/Source/Niagara/Private/NiagaraSystemInstance.cpp
--- a/Engine/Plugins/FX/Niagara/Source/Niagara/Private/NiagaraSystemInstance.cpp
+++ b/Engine/Plugins/FX/Niagara/Source/Niagara/Private/NiagaraSystemInstance.cpp
@@ -335,7 +335,7 @@ void FNiagaraSystemInstance::Dump()const
 
 void FNiagaraSystemInstance::DumpTickInfo(FOutputDevice& Ar)
 {
-	static const UEnum* TickingGroupEnum = FindObjectChecked<UEnum>(ANY_PACKAGE, TEXT("ETickingGroup"));
+	static const UEnum* TickingGroupEnum = FindObjectChecked<UEnum>(nullptr, TEXT("/Script/Engine.ETickingGroup"));
 
 	FString PrereqInfo;
 	UActorComponent* PrereqComponent = GetPrereqComponent();
diff --git a/Engine/Plugins/FX/Niagara/Source/Niagara/Private/NiagaraWorldManager.cpp b/Engine/Plugins/FX/Niagara/Source/Niagara/Private/NiagaraWorldManager.cpp
--- a/Engine/Plugins/FX/Niagara/Source/Niagara/Private/NiagaraWorldManager.cpp
+++ b/Engine/Plugins/FX/Niagara/Source/Niagara/Private/NiagaraWorldManager.cpp
@@ -263,7 +263,7 @@ void FNiagaraWorldManagerTickFunction::ExecuteTick(float DeltaTime, enum ELevelT
 
 FString FNiagaraWorldManagerTickFunction::DiagnosticMessage()
 {
-	static const UEnum* EnumType = FindObjectChecked<UEnum>(ANY_PACKAGE, TEXT("ETickingGroup"));
+	static const UEnum* EnumType = FindObjectChecked<UEnum>(nullptr, TEXT("/Script/Engine.ETickingGroup"));
 
 	return TEXT("FParticleSystemManager::Tick(") + EnumType->GetNameStringByIndex(static_cast<uint32>(TickGroup)) + TEXT(")");
 }
@@ -1189,7 +1189,7 @@ void FNiagaraWorldManager::DumpDetails(FOutputDevice& Ar)
 {
 	Ar.Logf(TEXT("=== FNiagaraWorldManager Dumping Detailed Information"));
 
-	static const UEnum* TickingGroupEnum = FindObjectChecked<UEnum>(ANY_PACKAGE, TEXT("ETickingGroup"));
+	static const UEnum* TickingGroupEnum = FindObjectChecked<UEnum>(nullptr, TEXT("/Script/Engine.ETickingGroup"));
 
 	for ( int TG=0; TG < NiagaraNumTickGroups; ++TG )
 	{
diff --git a/Engine/Plugins/FX/Niagara/Source/Niagara/Public/NiagaraCommon.h b/Engine/Plugins/FX/Niagara/Source/Niagara/Public/NiagaraCommon.h
--- a/Engine/Plugins/FX/Niagara/Source/Niagara/Public/NiagaraCommon.h
+++ b/Engine/Plugins/FX/Niagara/Source/Niagara/Public/NiagaraCommon.h
@@ -404,7 +404,7 @@ struct NIAGARA_API FNiagaraFunctionSignature
 	uint32 bHidden : 1;
 
 	/** Bitmask for which scripts are supported for this function. Use UNiagaraScript::MakeSupportedUsageContextBitmask to make the bitmask. */
-	UPROPERTY(meta = (Bitmask, BitmaskEnum = ENiagaraScriptUsage))
+	UPROPERTY(meta = (Bitmask, BitmaskEnum = "/Script/Niagara.ENiagaraScriptUsage"))
 	int32 ModuleUsageBitmask;
 
 	/** When using simulation stages and bRequiresContext is true this will be the index of the stage that is associated with the function. */
diff --git a/Engine/Plugins/FX/Niagara/Source/Niagara/Public/NiagaraSettings.h b/Engine/Plugins/FX/Niagara/Source/Niagara/Public/NiagaraSettings.h
--- a/Engine/Plugins/FX/Niagara/Source/Niagara/Public/NiagaraSettings.h
+++ b/Engine/Plugins/FX/Niagara/Source/Niagara/Public/NiagaraSettings.h
@@ -89,13 +89,13 @@ class NIAGARA_API UNiagaraSettings : public UDeveloperSettings
 	GENERATED_UCLASS_BODY()
 
 #if WITH_EDITORONLY_DATA
-	UPROPERTY(config, EditAnywhere, Category = Niagara, meta = (AllowedClasses = "ScriptStruct"))
+	UPROPERTY(config, EditAnywhere, Category = Niagara, meta = (AllowedClasses = "/Script/CoreUObject.ScriptStruct"))
 	TArray<FSoftObjectPath> AdditionalParameterTypes;
 
-	UPROPERTY(config, EditAnywhere, Category = Niagara, meta = (AllowedClasses = "ScriptStruct"))
+	UPROPERTY(config, EditAnywhere, Category = Niagara, meta = (AllowedClasses = "/Script/CoreUObject.ScriptStruct"))
 	TArray<FSoftObjectPath> AdditionalPayloadTypes;
 
-	UPROPERTY(config, EditAnywhere, Category = Niagara, meta = (AllowedClasses = "Enum"))
+	UPROPERTY(config, EditAnywhere, Category = Niagara, meta = (AllowedClasses = "/Script/CoreUObject.Enum"))
 	TArray<FSoftObjectPath> AdditionalParameterEnums;
 
 	/** Sets the default navigation behavior for the system preview viewport. */
@@ -123,7 +123,7 @@ class NIAGARA_API UNiagaraSettings : public UDeveloperSettings
 	bool bEnforceStrictStackTypes = true;
 
 	/** Default effect type to use for effects that don't define their own. Can be null. */
-	UPROPERTY(config, EditAnywhere, Category = Niagara, meta = (AllowedClasses = "NiagaraEffectType"))
+	UPROPERTY(config, EditAnywhere, Category = Niagara, meta = (AllowedClasses = "/Script/Niagara.NiagaraEffectType"))
 	FSoftObjectPath DefaultEffectType;
 
 	/** Position pin type color. The other pin colors are defined in the general editor settings. */
diff --git a/Engine/Plugins/FX/Niagara/Source/NiagaraEditor/Private/NiagaraParameterMapHistory.cpp b/Engine/Plugins/FX/Niagara/Source/NiagaraEditor/Private/NiagaraParameterMapHistory.cpp
--- a/Engine/Plugins/FX/Niagara/Source/NiagaraEditor/Private/NiagaraParameterMapHistory.cpp
+++ b/Engine/Plugins/FX/Niagara/Source/NiagaraEditor/Private/NiagaraParameterMapHistory.cpp
@@ -1181,7 +1181,7 @@ int32 FNiagaraParameterMapHistoryBuilder::FindMatchingParameterMapFromContextInp
 					const UNiagaraNodeOutput* ContextOutputNode = FNiagaraStackGraphUtilities::GetEmitterOutputNodeForStackNode(*Node);
 					if (ContextOutputNode != nullptr)
 					{
-						UEnum* NiagaraScriptUsageEnum = FindObjectChecked<UEnum>(ANY_PACKAGE, TEXT("ENiagaraScriptUsage"), true);
+						UEnum* NiagaraScriptUsageEnum = FindObjectChecked<UEnum>(nullptr, TEXT("/Script/Niagara.ENiagaraScriptUsage"), true);
 						ScriptUsageDisplayName = NiagaraScriptUsageEnum->GetDisplayNameTextByValue((uint64)ContextOutputNode->GetUsage()).ToString();
 					}
 					else
diff --git a/Engine/Plugins/FX/Niagara/Source/NiagaraEditor/Private/ViewModels/Stack/NiagaraStackScriptItemGroup.cpp b/Engine/Plugins/FX/Niagara/Source/NiagaraEditor/Private/ViewModels/Stack/NiagaraStackScriptItemGroup.cpp
--- a/Engine/Plugins/FX/Niagara/Source/NiagaraEditor/Private/ViewModels/Stack/NiagaraStackScriptItemGroup.cpp
+++ b/Engine/Plugins/FX/Niagara/Source/NiagaraEditor/Private/ViewModels/Stack/NiagaraStackScriptItemGroup.cpp
@@ -914,7 +914,7 @@ TOptional<UNiagaraStackEntry::FDropRequestResponse> UNiagaraStackScriptItemGroup
 TOptional<UNiagaraStackEntry::FDropRequestResponse> UNiagaraStackScriptItemGroup::CanDropAssetsOnTarget(const UNiagaraStackEntry& TargetEntry, const FDropRequest& DropRequest)
 {
 	TSharedRef<const FAssetDragDropOp> AssetDragDropOp = StaticCastSharedRef<const FAssetDragDropOp>(DropRequest.DragDropOperation);
-	const UEnum* NiagaraScriptUsageEnum = FindObjectChecked<UEnum>(ANY_PACKAGE, TEXT("ENiagaraScriptUsage"), true);
+	const UEnum* NiagaraScriptUsageEnum = FindObjectChecked<UEnum>(nullptr, TEXT("/Script/Niagara.ENiagaraScriptUsage"), true);
 
 	if (&TargetEntry != this && TargetEntry.IsA<UNiagaraStackModuleItem>() == false)
 	{
diff --git a/Engine/Plugins/FX/Niagara/Source/NiagaraEditor/Public/NiagaraEditorSettings.h b/Engine/Plugins/FX/Niagara/Source/NiagaraEditor/Public/NiagaraEditorSettings.h
--- a/Engine/Plugins/FX/Niagara/Source/NiagaraEditor/Public/NiagaraEditorSettings.h
+++ b/Engine/Plugins/FX/Niagara/Source/NiagaraEditor/Public/NiagaraEditorSettings.h
@@ -188,7 +188,7 @@ struct NIAGARAEDITOR_API FNiagaraCurveTemplate
 	UPROPERTY(EditAnywhere, Category = Curve)
 	FString DisplayNameOverride;
 
-	UPROPERTY(EditAnywhere, Category = Curve, meta = (AllowedClasses = "CurveFloat"))
+	UPROPERTY(EditAnywhere, Category = Curve, meta = (AllowedClasses = "/Script/Engine.CurveFloat"))
 	FSoftObjectPath CurveAsset;
 };
 
diff --git a/Engine/Plugins/Importers/USDImporter/Source/USDStage/Public/USDStageActor.h b/Engine/Plugins/Importers/USDImporter/Source/USDStage/Public/USDStageActor.h
--- a/Engine/Plugins/Importers/USDImporter/Source/USDStage/Public/USDStageActor.h
+++ b/Engine/Plugins/Importers/USDImporter/Source/USDStage/Public/USDStageActor.h
@@ -52,7 +52,7 @@ public:
 	 * Whether to try to combine individual assets and components of the same type on a kind-per-kind basis,
 	 * like multiple Mesh prims into a single Static Mesh
 	 */
-	UPROPERTY(EditAnywhere, BlueprintReadOnly, Category = "USD", meta = (Bitmask, BitmaskEnum=EUsdDefaultKind))
+	UPROPERTY(EditAnywhere, BlueprintReadOnly, Category = "USD", meta = (Bitmask, BitmaskEnum="/Script/UnrealUSDWrapper.EUsdDefaultKind"))
 	int32 KindsToCollapse;
 
 	/**
@@ -71,7 +71,7 @@ public:
 	bool bCollapseTopLevelPointInstancers;
 
 	/* Only load prims with these specific purposes from the USD file */
-	UPROPERTY(EditAnywhere, BlueprintReadOnly, Category = "USD", meta = (Bitmask, BitmaskEnum=EUsdPurpose))
+	UPROPERTY(EditAnywhere, BlueprintReadOnly, Category = "USD", meta = (Bitmask, BitmaskEnum="/Script/UnrealUSDWrapper.EUsdPurpose"))
 	int32 PurposesToLoad;
 
 	/** Try enabling Nanite for static meshes that are generated with at least this many triangles */
diff --git a/Engine/Plugins/Importers/USDImporter/Source/USDStageImporter/Public/USDStageImportOptions.h b/Engine/Plugins/Importers/USDImporter/Source/USDStageImporter/Public/USDStageImportOptions.h
--- a/Engine/Plugins/Importers/USDImporter/Source/USDStageImporter/Public/USDStageImportOptions.h
+++ b/Engine/Plugins/Importers/USDImporter/Source/USDStageImporter/Public/USDStageImportOptions.h
@@ -66,7 +66,7 @@ public:
 
 
 	/** Only import prims with these specific purposes from the USD file */
-	UPROPERTY(BlueprintReadWrite, config, EditAnywhere, Category= "USD options", meta = (Bitmask, BitmaskEnum=EUsdPurpose))
+	UPROPERTY(BlueprintReadWrite, config, EditAnywhere, Category= "USD options", meta = (Bitmask, BitmaskEnum="/Script/UnrealUSDWrapper.EUsdPurpose"))
 	int32 PurposesToImport;
 
 	/** Try enabling Nanite for static meshes that are generated with at least this many triangles */
@@ -114,7 +114,7 @@ public:
 	 * Whether to try to combine individual assets and components of the same type on a kind-per-kind basis,
 	 * like multiple Mesh prims into a single Static Mesh
 	 */
-	UPROPERTY( BlueprintReadWrite, config, EditAnywhere, Category = "Processing", meta = ( Bitmask, BitmaskEnum = EUsdDefaultKind ) )
+	UPROPERTY( BlueprintReadWrite, config, EditAnywhere, Category = "Processing", meta = ( Bitmask, BitmaskEnum = "/Script/UnrealUSDWrapper.EUsdDefaultKind" ) )
 	int32 KindsToCollapse;
 
 	/**
diff --git a/Engine/Plugins/Media/AjaMedia/Source/AjaMediaOutput/Public/AjaMediaFrameGrabberProtocol.h b/Engine/Plugins/Media/AjaMedia/Source/AjaMediaOutput/Public/AjaMediaFrameGrabberProtocol.h
--- a/Engine/Plugins/Media/AjaMedia/Source/AjaMediaOutput/Public/AjaMediaFrameGrabberProtocol.h
+++ b/Engine/Plugins/Media/AjaMedia/Source/AjaMediaOutput/Public/AjaMediaFrameGrabberProtocol.h
@@ -29,7 +29,7 @@ public:
 public:
 
 	/** AJA Setting to use for the FrameGrabberProtocol */
-	UPROPERTY(config, BlueprintReadWrite, EditAnywhere, Category=AJA, meta=(AllowedClasses=AjaMediaOutput))
+	UPROPERTY(config, BlueprintReadWrite, EditAnywhere, Category=AJA, meta=(AllowedClasses="/Script/AjaMediaOutput.AjaMediaOutput"))
 	FSoftObjectPath MediaOutput;
 
 	/** States unused options for AJAFrameGrabberProtocolSettings */
diff --git a/Engine/Plugins/MovieScene/LevelSequenceEditor/Source/LevelSequenceEditor/Private/Misc/LevelSequenceEditorSettings.h b/Engine/Plugins/MovieScene/LevelSequenceEditor/Source/LevelSequenceEditor/Private/Misc/LevelSequenceEditorSettings.h
--- a/Engine/Plugins/MovieScene/LevelSequenceEditor/Source/LevelSequenceEditor/Private/Misc/LevelSequenceEditorSettings.h
+++ b/Engine/Plugins/MovieScene/LevelSequenceEditor/Source/LevelSequenceEditor/Private/Misc/LevelSequenceEditorSettings.h
@@ -31,15 +31,15 @@ struct FLevelSequenceTrackSettings
 	GENERATED_BODY()
 
 	/** The Actor class to create movie scene tracks for. */
-	UPROPERTY(config, noclear, EditAnywhere, Category=TrackSettings, meta=(MetaClass="Actor"))
+	UPROPERTY(config, noclear, EditAnywhere, Category=TrackSettings, meta=(MetaClass="/Script/Engine.Actor"))
 	FSoftClassPath MatchingActorClass;
 
 	/** List of movie scene track classes to be added automatically. */
-	UPROPERTY(config, noclear, EditAnywhere, Category=TrackSettings, meta=(MetaClass="MovieSceneTrack"))
+	UPROPERTY(config, noclear, EditAnywhere, Category=TrackSettings, meta=(MetaClass="/Script/MovieScene.MovieSceneTrack"))
 	TArray<FSoftClassPath> DefaultTracks;
 
 	/** List of movie scene track classes not to be added automatically. */
-	UPROPERTY(config, noclear, EditAnywhere, Category=TrackSettings, meta=(MetaClass="MovieSceneTrack"))
+	UPROPERTY(config, noclear, EditAnywhere, Category=TrackSettings, meta=(MetaClass="/Script/MovieScene.MovieSceneTrack"))
 	TArray<FSoftClassPath> ExcludeDefaultTracks;
 
 	/** List of property names for which movie scene tracks will be created automatically. */
diff --git a/Engine/Plugins/MovieScene/MovieRenderPipeline/Source/MovieRenderPipelineCore/Public/MoviePipelineQueue.h b/Engine/Plugins/MovieScene/MovieRenderPipeline/Source/MovieRenderPipelineCore/Public/MoviePipelineQueue.h
--- a/Engine/Plugins/MovieScene/MovieRenderPipeline/Source/MovieRenderPipelineCore/Public/MoviePipelineQueue.h
+++ b/Engine/Plugins/MovieScene/MovieRenderPipeline/Source/MovieRenderPipelineCore/Public/MoviePipelineQueue.h
@@ -348,11 +348,11 @@ public:
 	FString JobName;
 
 	/** Which sequence should this job render? */
-	UPROPERTY(BlueprintReadWrite, EditAnywhere, BlueprintSetter = "SetSequence", Category = "Movie Render Pipeline", meta = (AllowedClasses = "LevelSequence"))
+	UPROPERTY(BlueprintReadWrite, EditAnywhere, BlueprintSetter = "SetSequence", Category = "Movie Render Pipeline", meta = (AllowedClasses = "/Script/LevelSequence.LevelSequence"))
 	FSoftObjectPath Sequence;
 
 	/** Which map should this job render on */
-	UPROPERTY(BlueprintReadWrite, EditAnywhere, Category = "Movie Render Pipeline", meta = (AllowedClasses = "World"))
+	UPROPERTY(BlueprintReadWrite, EditAnywhere, Category = "Movie Render Pipeline", meta = (AllowedClasses = "/Script/Engine.World"))
 	FSoftObjectPath Map;
 
 	/** (Optional) If left blank, will default to system username. Can be shown in burn in as a first point of contact about the content. */
diff --git a/Engine/Plugins/MovieScene/MovieRenderPipeline/Source/MovieRenderPipelineEditor/Public/MovieRenderPipelineSettings.h b/Engine/Plugins/MovieScene/MovieRenderPipeline/Source/MovieRenderPipelineEditor/Public/MovieRenderPipelineSettings.h
--- a/Engine/Plugins/MovieScene/MovieRenderPipeline/Source/MovieRenderPipelineEditor/Public/MovieRenderPipelineSettings.h
+++ b/Engine/Plugins/MovieScene/MovieRenderPipeline/Source/MovieRenderPipelineEditor/Public/MovieRenderPipelineSettings.h
@@ -42,7 +42,7 @@ public:
 	* which does different logic. See UMoviePipelineExecutorBase for more information. This is used for
 	* the Render button on the UI.
 	*/
-	UPROPERTY(config, EditAnywhere, BlueprintReadWrite, meta = (MetaClass = "MoviePipelineExecutorBase"), Category="Movie Render Pipeline")
+	UPROPERTY(config, EditAnywhere, BlueprintReadWrite, meta = (MetaClass = "/Script/MovieRenderPipelineCore.MoviePipelineExecutorBase"), Category="Movie Render Pipeline")
 	FSoftClassPath DefaultLocalExecutor;
 
 	/**
@@ -51,7 +51,7 @@ public:
 	* which does different logic. See UMoviePipelineExecutorBase for more information. This is used for
 	* the Render Remotely button on the UI.
 	*/
-	UPROPERTY(config, EditAnywhere, BlueprintReadWrite, meta = (MetaClass = "MoviePipelineExecutorBase"), Category = "Movie Render Pipeline")
+	UPROPERTY(config, EditAnywhere, BlueprintReadWrite, meta = (MetaClass = "/Script/MovieRenderPipelineCore.MoviePipelineExecutorBase"), Category = "Movie Render Pipeline")
 	FSoftClassPath DefaultRemoteExecutor;
 	
 	/**
@@ -59,7 +59,7 @@ public:
 	* that will have editable properties in the UI for special handling with your executor. This can be
 	* made dynamic if you add jobs to the queue programatically instead of through the UI.
 	*/
-	UPROPERTY(config, EditAnywhere, BlueprintReadWrite, meta = (MetaClass = "MoviePipelineExecutorJob"), Category="Movie Render Pipeline")
+	UPROPERTY(config, EditAnywhere, BlueprintReadWrite, meta = (MetaClass = "/Script/MovieRenderPipelineCore.MoviePipelineExecutorJob"), Category="Movie Render Pipeline")
 	FSoftClassPath DefaultExecutorJob;
 	
 	/**
@@ -67,7 +67,7 @@ public:
 	* this will allow you to re-use the existing UI/Executors while providing your own logic for producing
 	* a single render.
 	*/
-	UPROPERTY(config, EditAnywhere, BlueprintReadWrite, meta = (MetaClass = "MoviePipeline"), Category="Movie Render Pipeline")
+	UPROPERTY(config, EditAnywhere, BlueprintReadWrite, meta = (MetaClass = "/Script/MovieRenderPipelineCore.MoviePipeline"), Category="Movie Render Pipeline")
 	FSoftClassPath DefaultPipeline;
 
 	/**
@@ -78,6 +78,6 @@ public:
 	*
 	* This only applies to jobs created via the UI. If you do not use the UI (ie: Scripting/Python) you will need to
 	* add settings by hand for each job you create. */
-	UPROPERTY(config, EditAnywhere, BlueprintReadWrite, Category = "Movie Render Pipeline", meta = (MetaClass = "MoviePipelineSetting"), DisplayName = "Default Job Settings Classes")
+	UPROPERTY(config, EditAnywhere, BlueprintReadWrite, Category = "Movie Render Pipeline", meta = (MetaClass = "/Script/MovieRenderPipelineCore.MoviePipelineSetting"), DisplayName = "Default Job Settings Classes")
 	TArray<FSoftClassPath> DefaultClasses;
 };
\ No newline at end of file
diff --git a/Engine/Plugins/MovieScene/MovieRenderPipeline/Source/MovieRenderPipelineSettings/Public/MoviePipelineBurnInSetting.h b/Engine/Plugins/MovieScene/MovieRenderPipeline/Source/MovieRenderPipelineSettings/Public/MoviePipelineBurnInSetting.h
--- a/Engine/Plugins/MovieScene/MovieRenderPipeline/Source/MovieRenderPipelineSettings/Public/MoviePipelineBurnInSetting.h
+++ b/Engine/Plugins/MovieScene/MovieRenderPipeline/Source/MovieRenderPipelineSettings/Public/MoviePipelineBurnInSetting.h
@@ -41,11 +41,11 @@ public:
 	virtual bool IsValidOnShots() const override { return false; }
 	virtual bool IsValidOnMaster() const override { return true; }
 public:
-	UPROPERTY(EditAnywhere, BlueprintReadWrite, meta=(MetaClass="MoviePipelineBurnInWidget"), Category = "Widget Settings")
+	UPROPERTY(EditAnywhere, BlueprintReadWrite, meta=(MetaClass="/Script/MovieRenderPipelineSettings.MoviePipelineBurnInWidget"), Category = "Widget Settings")
 	FSoftClassPath BurnInClass;
 
 	/** If true, the Burn In image will be composited into the Final Image pass. Doesn't apply to multi-layer EXR files. */
-	UPROPERTY(EditAnywhere, BlueprintReadWrite, meta = (MetaClass = "MoviePipelineBurnInWidget"), Category = "Widget Settings")
+	UPROPERTY(EditAnywhere, BlueprintReadWrite, meta = (MetaClass = "/Script/MovieRenderPipelineSettings.MoviePipelineBurnInWidget"), Category = "Widget Settings")
 	bool bCompositeOntoFinalImage;
 
 private:
diff --git a/Engine/Plugins/MovieScene/MovieRenderPipeline/Source/MovieRenderPipelineSettings/Public/MoviePipelineWidgetRenderSetting.h b/Engine/Plugins/MovieScene/MovieRenderPipeline/Source/MovieRenderPipelineSettings/Public/MoviePipelineWidgetRenderSetting.h
--- a/Engine/Plugins/MovieScene/MovieRenderPipeline/Source/MovieRenderPipelineSettings/Public/MoviePipelineWidgetRenderSetting.h
+++ b/Engine/Plugins/MovieScene/MovieRenderPipeline/Source/MovieRenderPipelineSettings/Public/MoviePipelineWidgetRenderSetting.h
@@ -32,7 +32,7 @@ protected:
 
 public:
 	/** If true, the widget renderer image will be composited into the Final Image pass. Doesn't apply to multi-layer EXR files. */
-	UPROPERTY(EditAnywhere, BlueprintReadWrite, meta = (MetaClass = "MoviePipelineBurnInWidget"), Category = "Widget Settings")
+	UPROPERTY(EditAnywhere, BlueprintReadWrite, meta = (MetaClass = "/Script/MovieRenderPipelineSettings.MoviePipelineBurnInWidget"), Category = "Widget Settings")
 	bool bCompositeOntoFinalImage;
 
 public:
diff --git a/Engine/Plugins/MovieScene/TemplateSequence/Source/TemplateSequence/Public/TemplateSequenceActor.h b/Engine/Plugins/MovieScene/TemplateSequence/Source/TemplateSequence/Public/TemplateSequenceActor.h
--- a/Engine/Plugins/MovieScene/TemplateSequence/Source/TemplateSequence/Public/TemplateSequenceActor.h
+++ b/Engine/Plugins/MovieScene/TemplateSequence/Source/TemplateSequence/Public/TemplateSequenceActor.h
@@ -63,7 +63,7 @@ public:
 	UPROPERTY(Instanced, Transient, Replicated, BlueprintReadOnly, BlueprintGetter = GetSequencePlayer, Category = "Playback", meta = (ExposeFunctionCategories = "Sequencer|Player"))
 	TObjectPtr<UTemplateSequencePlayer> SequencePlayer;
 
-	UPROPERTY(EditAnywhere, BlueprintReadOnly, Category = "General", meta = (AllowedClasses = "TemplateSequence"))
+	UPROPERTY(EditAnywhere, BlueprintReadOnly, Category = "General", meta = (AllowedClasses = "/Script/TemplateSequence.TemplateSequence"))
 	FSoftObjectPath TemplateSequence;
 
 	/** The override for the template sequence's root object binding. See SetBinding. */
diff --git a/Engine/Plugins/NetcodeUnitTest/NUTUnrealEngine/Source/NUTUnrealEngine/Private/UnitTests/Obsolete/UTT61_DebugReplicateData.cpp b/Engine/Plugins/NetcodeUnitTest/NUTUnrealEngine/Source/NUTUnrealEngine/Private/UnitTests/Obsolete/UTT61_DebugReplicateData.cpp
--- a/Engine/Plugins/NetcodeUnitTest/NUTUnrealEngine/Source/NUTUnrealEngine/Private/UnitTests/Obsolete/UTT61_DebugReplicateData.cpp
+++ b/Engine/Plugins/NetcodeUnitTest/NUTUnrealEngine/Source/NUTUnrealEngine/Private/UnitTests/Obsolete/UTT61_DebugReplicateData.cpp
@@ -16,7 +16,7 @@
 
 UUTT61_DebugReplicateData::UUTT61_DebugReplicateData(const FObjectInitializer& ObjectInitializer)
 	: Super(ObjectInitializer)
-	, RepClass(FindObject<UClass>(ANY_PACKAGE, TEXT("GameplayDebuggingReplicator")))
+	, RepClass(FindObject<UClass>(nullptr, TEXT("/Script/GameplayDebugger.GameplayDebuggingReplicator")))
 	, Replicator(nullptr)
 {
 	UnitTestName = TEXT("ReplicateDataCheck");
diff --git a/Engine/Plugins/NetcodeUnitTest/NetcodeUnitTest/Source/NetcodeUnitTest/Private/ClientUnitTest.cpp b/Engine/Plugins/NetcodeUnitTest/NetcodeUnitTest/Source/NetcodeUnitTest/Private/ClientUnitTest.cpp
--- a/Engine/Plugins/NetcodeUnitTest/NetcodeUnitTest/Source/NetcodeUnitTest/Private/ClientUnitTest.cpp
+++ b/Engine/Plugins/NetcodeUnitTest/NetcodeUnitTest/Source/NetcodeUnitTest/Private/ClientUnitTest.cpp
@@ -67,7 +67,7 @@ UClientUnitTest::UClientUnitTest(const FObjectInitializer& ObjectInitializer)
 
 void UClientUnitTest::InitOnlineBeaconClass()
 {
-	OnlineBeaconClass_Private = FindObject<UClass>(ANY_PACKAGE, TEXT("OnlineBeaconClient"));
+	OnlineBeaconClass_Private = FindObject<UClass>(nullptr, TEXT("/Script/OnlineSubsystemUtils.OnlineBeaconClient"));
 }
 
 void UClientUnitTest::NotifyAlterMinClient(FMinClientParms& Parms)
diff --git a/Engine/Plugins/Runtime/AR/AppleAR/AppleARKit/Source/AppleARKit/Private/AppleARKitConversion.cpp b/Engine/Plugins/Runtime/AR/AppleAR/AppleARKit/Source/AppleARKit/Private/AppleARKitConversion.cpp
--- a/Engine/Plugins/Runtime/AR/AppleAR/AppleARKit/Source/AppleARKit/Private/AppleARKitConversion.cpp
+++ b/Engine/Plugins/Runtime/AR/AppleAR/AppleARKit/Source/AppleARKit/Private/AppleARKitConversion.cpp
@@ -11,9 +11,9 @@
 #endif
 
 template<typename TEnum>
-static FString GetEnumValueAsString(const FString& Name, TEnum Value)
+static FString GetEnumValueAsString(const TCHAR* PathName, TEnum Value)
 {
-	if (const UEnum* EnumClass = FindObject<UEnum>(ANY_PACKAGE, *Name, true))
+	if (const UEnum* EnumClass = FindObject<UEnum>(nullptr, PathName, true))
 	{
 		return EnumClass->GetNameByValue((int64)Value).ToString();
 	}
@@ -629,8 +629,8 @@ void FAppleARKitConversion::ConfigureSessionTrackingFeatures(UARSessionConfig* S
 				else
 				{
 					UE_LOG(LogAppleARKit, Error, TEXT("Session type [%s] doesn't support the required session feature: [%s]!"),
-						   *GetEnumValueAsString<>(TEXT("EARSessionType"), SessionType),
-						   *GetEnumValueAsString<>(TEXT("EARSessionTrackingFeature"), SessionTrackingFeature));
+						   *GetEnumValueAsString<>(TEXT("/Script/AugmentedReality.EARSessionType"), SessionType),
+						   *GetEnumValueAsString<>(TEXT("/Script/AugmentedReality.EARSessionTrackingFeature"), SessionTrackingFeature));
 				}
 			}
 		}
diff --git a/Engine/Plugins/Runtime/AR/Google/GoogleARCore/Source/GoogleARCoreBase/Public/GoogleARCoreTypes.h b/Engine/Plugins/Runtime/AR/Google/GoogleARCore/Source/GoogleARCoreBase/Public/GoogleARCoreTypes.h
--- a/Engine/Plugins/Runtime/AR/Google/GoogleARCore/Source/GoogleARCoreBase/Public/GoogleARCoreTypes.h
+++ b/Engine/Plugins/Runtime/AR/Google/GoogleARCore/Source/GoogleARCoreBase/Public/GoogleARCoreTypes.h
@@ -292,11 +292,11 @@ struct GOOGLEARCOREBASE_API FGoogleARCoreCameraConfig
 	FString CameraID;
 
 	/** The target FPS of the camera. */
-	UPROPERTY(BlueprintReadOnly, Category = "GoogleARCore|CameraConfig", meta = (Bitmask, BitmaskEnum = EGoogleARCoreCameraFPS))
+	UPROPERTY(BlueprintReadOnly, Category = "GoogleARCore|CameraConfig", meta = (Bitmask, BitmaskEnum = "/Script/GoogleARCoreBase.EGoogleARCoreCameraFPS"))
 	int32 TargetFPS = 0; // EGoogleARCoreCameraFPS::FPS_Any
 
 	/** The desired depth sensor usage of the camera. */
-	UPROPERTY(BlueprintReadOnly, Category = "GoogleARCore|CameraConfig", meta = (Bitmask, BitmaskEnum = EGoogleARCoreCameraDepthSensorUsage))
+	UPROPERTY(BlueprintReadOnly, Category = "GoogleARCore|CameraConfig", meta = (Bitmask, BitmaskEnum = "/Script/GoogleARCoreBase.EGoogleARCoreCameraDepthSensorUsage"))
 	int32 DepthSensorUsage = 0; // EGoogleARCoreCameraDepthSensorUsage::DepthSensor_Any
 
 	/** Below are for internal use only */
diff --git a/Engine/Plugins/Runtime/AudioModulation/Source/AudioModulation/Public/AudioModulationSettings.h b/Engine/Plugins/Runtime/AudioModulation/Source/AudioModulation/Public/AudioModulationSettings.h
--- a/Engine/Plugins/Runtime/AudioModulation/Source/AudioModulation/Public/AudioModulationSettings.h
+++ b/Engine/Plugins/Runtime/AudioModulation/Source/AudioModulation/Public/AudioModulationSettings.h
@@ -17,7 +17,7 @@ class AUDIOMODULATION_API UAudioModulationSettings : public UDeveloperSettings
 
 public:
 	// Array of Modulation Parameters that are loaded on plugin startup
-	UPROPERTY(config, EditAnywhere, Category = "Parameters", meta = (AllowedClasses = "SoundModulationParameter"))
+	UPROPERTY(config, EditAnywhere, Category = "Parameters", meta = (AllowedClasses = "/Script/AudioModulation.SoundModulationParameter"))
 	TArray<FSoftObjectPath> Parameters;
 
 	void RegisterParameters() const;
diff --git a/Engine/Plugins/Runtime/GameplayAbilities/Source/GameplayAbilities/Private/AbilitySystemGlobals.cpp b/Engine/Plugins/Runtime/GameplayAbilities/Source/GameplayAbilities/Private/AbilitySystemGlobals.cpp
--- a/Engine/Plugins/Runtime/GameplayAbilities/Source/GameplayAbilities/Private/AbilitySystemGlobals.cpp
+++ b/Engine/Plugins/Runtime/GameplayAbilities/Source/GameplayAbilities/Private/AbilitySystemGlobals.cpp
@@ -496,9 +496,9 @@ void UAbilitySystemGlobals::ListPlayerAbilities()
 	UAbilitySystemComponent* AbilityComponent = PC ? GetAbilitySystemComponentFromActor(PC->GetPawn()) : nullptr;
 	if(AbilityComponent)
 	{
-		const UEnum* ExecutionEnumPtr = FindObject<UEnum>(ANY_PACKAGE, TEXT("EGameplayAbilityNetExecutionPolicy"), true);
+		const UEnum* ExecutionEnumPtr = FindObject<UEnum>(nullptr, TEXT("/Script/GameplayAbilities.EGameplayAbilityNetExecutionPolicy"), true);
 		check(ExecutionEnumPtr && TEXT("Couldn't locate EGameplayAbilityNetExecutionPolicy enum!"));
-		const UEnum* SecurityEnumPtr = FindObject<UEnum>(ANY_PACKAGE, TEXT("EGameplayAbilityNetSecurityPolicy"), true);
+		const UEnum* SecurityEnumPtr = FindObject<UEnum>(nullptr, TEXT("/Script/GameplayAbilities.EGameplayAbilityNetSecurityPolicy"), true);
 		check(SecurityEnumPtr && TEXT("Couldn't locate EGameplayAbilityNetSecurityPolicy enum!"));
 
 		PC->ClientMessage(TEXT("Available abilities:"));
diff --git a/Engine/Plugins/Runtime/GameplayAbilities/Source/GameplayAbilities/Public/GameplayEffectTypes.h b/Engine/Plugins/Runtime/GameplayAbilities/Source/GameplayAbilities/Public/GameplayEffectTypes.h
--- a/Engine/Plugins/Runtime/GameplayAbilities/Source/GameplayAbilities/Public/GameplayEffectTypes.h
+++ b/Engine/Plugins/Runtime/GameplayAbilities/Source/GameplayAbilities/Public/GameplayEffectTypes.h
@@ -917,7 +917,7 @@ struct FGameplayEffectRemovalInfo
 
 
 /** Metadata about a gameplay cue execution */
-USTRUCT(BlueprintType, meta = (HasNativeBreak = "GameplayAbilities.AbilitySystemBlueprintLibrary.BreakGameplayCueParameters", HasNativeMake = "GameplayAbilities.AbilitySystemBlueprintLibrary.MakeGameplayCueParameters"))
+USTRUCT(BlueprintType, meta = (HasNativeBreak = "/Script/GameplayAbilities.AbilitySystemBlueprintLibrary.BreakGameplayCueParameters", HasNativeMake = "/Script/GameplayAbilities.AbilitySystemBlueprintLibrary.MakeGameplayCueParameters"))
 struct GAMEPLAYABILITIES_API FGameplayCueParameters
 {
 	GENERATED_USTRUCT_BODY()
diff --git a/Engine/Plugins/Runtime/HairStrands/Source/HairStrandsCore/Public/GroomCacheImportOptions.h b/Engine/Plugins/Runtime/HairStrands/Source/HairStrandsCore/Public/GroomCacheImportOptions.h
--- a/Engine/Plugins/Runtime/HairStrands/Source/HairStrandsCore/Public/GroomCacheImportOptions.h
+++ b/Engine/Plugins/Runtime/HairStrands/Source/HairStrandsCore/Public/GroomCacheImportOptions.h
@@ -32,7 +32,7 @@ struct HAIRSTRANDSCORE_API FGroomCacheImportSettings
 	bool bImportGroomAsset = true;
 
 	/** The groom asset the groom cache will be built from (must be compatible) */
-	UPROPERTY(EditAnywhere, BlueprintReadWrite, Category = GroomCache, meta = (MetaClass = "GroomAsset"))
+	UPROPERTY(EditAnywhere, BlueprintReadWrite, Category = GroomCache, meta = (MetaClass = "/Script/HairStrandsCore.GroomAsset"))
 	FSoftObjectPath GroomAsset;
 };
 
diff --git a/Engine/Plugins/Runtime/MassEntity/Source/MassEntity/Public/MassProcessor.h b/Engine/Plugins/Runtime/MassEntity/Source/MassEntity/Public/MassProcessor.h
--- a/Engine/Plugins/Runtime/MassEntity/Source/MassEntity/Public/MassProcessor.h
+++ b/Engine/Plugins/Runtime/MassEntity/Source/MassEntity/Public/MassProcessor.h
@@ -103,7 +103,7 @@ protected:
 
 protected:
 	/** Whether this processor should be executed on StandAlone or Server or Client */
-	UPROPERTY(EditAnywhere, Category = "Pipeline", meta = (Bitmask, BitmaskEnum = EProcessorExecutionFlags), config)
+	UPROPERTY(EditAnywhere, Category = "Pipeline", meta = (Bitmask, BitmaskEnum = "/Script/MassEntity.EProcessorExecutionFlags"), config)
 	int32 ExecutionFlags;
 
 	/** Processing phase this processor will be automatically run as part of. */
diff --git a/Engine/Plugins/Runtime/MassGameplay/Source/MassSpawner/Public/MassAssortedFragmentsTrait.h b/Engine/Plugins/Runtime/MassGameplay/Source/MassSpawner/Public/MassAssortedFragmentsTrait.h
--- a/Engine/Plugins/Runtime/MassGameplay/Source/MassSpawner/Public/MassAssortedFragmentsTrait.h
+++ b/Engine/Plugins/Runtime/MassGameplay/Source/MassSpawner/Public/MassAssortedFragmentsTrait.h
@@ -19,6 +19,6 @@ protected:
 
 	virtual void BuildTemplate(FMassEntityTemplateBuildContext& BuildContext, const UWorld& World) const override;
 
-	UPROPERTY(Category="Fragments", EditAnywhere, meta = (BaseStruct = "MassFragment", ExcludeBaseStruct))
+	UPROPERTY(Category="Fragments", EditAnywhere, meta = (BaseStruct = "/Script/MassEntity.MassFragment", ExcludeBaseStruct))
 	TArray<FInstancedStruct> Fragments;
 };
diff --git a/Engine/Plugins/Runtime/Metasound/Source/MetasoundEditor/Public/MetasoundEditorModule.h b/Engine/Plugins/Runtime/Metasound/Source/MetasoundEditor/Public/MetasoundEditorModule.h
--- a/Engine/Plugins/Runtime/Metasound/Source/MetasoundEditor/Public/MetasoundEditorModule.h
+++ b/Engine/Plugins/Runtime/Metasound/Source/MetasoundEditor/Public/MetasoundEditorModule.h
@@ -82,7 +82,7 @@ namespace Metasound
 			// By default, proxy classes support child classes & inheritance.
 			virtual void RegisterExplicitProxyClass(const UClass& InClass) = 0;
 
-			virtual bool IsMetaSoundAssetClass(const FName InClassName) const = 0;
+			virtual bool IsMetaSoundAssetClass(const FTopLevelAssetPath& InClassName) const = 0;
 
 			// Primes MetaSound assets, effectively loading the asset asynchronously (if not already
 			// loaded) & registers them if not already registered with the MetaSound Class Registry.
diff --git a/Engine/Plugins/Runtime/Metasound/Source/MetasoundEngine/Public/MetasoundSettings.h b/Engine/Plugins/Runtime/Metasound/Source/MetasoundEngine/Public/MetasoundSettings.h
--- a/Engine/Plugins/Runtime/Metasound/Source/MetasoundEngine/Public/MetasoundSettings.h
+++ b/Engine/Plugins/Runtime/Metasound/Source/MetasoundEngine/Public/MetasoundSettings.h
@@ -29,7 +29,7 @@ struct METASOUNDENGINE_API FDefaultMetaSoundAssetAutoUpdateSettings
 	GENERATED_BODY()
 
 	/** MetaSound to prevent from AutoUpdate. */
-	UPROPERTY(EditAnywhere, Category = "AutoUpdate", meta = (AllowedClasses = "MetaSound, MetaSoundSource"))
+	UPROPERTY(EditAnywhere, Category = "AutoUpdate", meta = (AllowedClasses = "/Script/MetasoundEngine.MetaSound, /Script/MetasoundEngine.MetaSoundSource"))
 	FSoftObjectPath MetaSound;
 };
 
diff --git a/Engine/Plugins/Runtime/Oculus/OculusAudio/Source/OculusAudio/Private/OculusAudioSettings.h b/Engine/Plugins/Runtime/Oculus/OculusAudio/Source/OculusAudio/Private/OculusAudioSettings.h
--- a/Engine/Plugins/Runtime/Oculus/OculusAudio/Source/OculusAudio/Private/OculusAudioSettings.h
+++ b/Engine/Plugins/Runtime/Oculus/OculusAudio/Source/OculusAudio/Private/OculusAudioSettings.h
@@ -19,7 +19,7 @@ public:
 	virtual bool CanEditChange(const FProperty* InProperty) const override;
 #endif
 
-	UPROPERTY(GlobalConfig, EditAnywhere, Category = "Reverb", meta = (AllowedClasses = "SoundSubmix"))
+	UPROPERTY(GlobalConfig, EditAnywhere, Category = "Reverb", meta = (AllowedClasses = "/Script/Engine.SoundSubmix"))
 	FSoftObjectPath OutputSubmix;
 
 	UPROPERTY(GlobalConfig, EditAnywhere, Category = "Reverb", meta = (ClampMin = "-60.0", ClampMax = "20.0", UIMin = "-60.0", UIMax = "20.0"))
diff --git a/Engine/Plugins/Runtime/Oculus/OculusVR/Source/OculusInput/Private/OculusHandTracking.cpp b/Engine/Plugins/Runtime/Oculus/OculusVR/Source/OculusInput/Private/OculusHandTracking.cpp
--- a/Engine/Plugins/Runtime/Oculus/OculusVR/Source/OculusInput/Private/OculusHandTracking.cpp
+++ b/Engine/Plugins/Runtime/Oculus/OculusVR/Source/OculusInput/Private/OculusHandTracking.cpp
@@ -594,7 +594,7 @@ FString FOculusHandTracking::GetBoneName(uint8 Bone)
 {
 	uint8 HandBone = Bone == ovrpBoneId_Invalid ? (uint8)EBone::Invalid : Bone;
 
-	UEnum* BoneEnum = FindObject<UEnum>(ANY_PACKAGE, TEXT("EBone"), true);
+	UEnum* BoneEnum = FindObject<UEnum>(nullptr, TEXT("/Script/OculusInput.EBone"), true);
 	if (BoneEnum)
 	{
 		return BoneEnum->GetDisplayNameTextByValue(HandBone).ToString();
diff --git a/Engine/Plugins/Runtime/OpenXRHandTracking/Source/OpenXRHandTracking/Private/OpenXRHandTracking.cpp b/Engine/Plugins/Runtime/OpenXRHandTracking/Source/OpenXRHandTracking/Private/OpenXRHandTracking.cpp
--- a/Engine/Plugins/Runtime/OpenXRHandTracking/Source/OpenXRHandTracking/Private/OpenXRHandTracking.cpp
+++ b/Engine/Plugins/Runtime/OpenXRHandTracking/Source/OpenXRHandTracking/Private/OpenXRHandTracking.cpp
@@ -335,7 +335,7 @@ bool FOpenXRHandTracking::GetControllerOrientationAndPosition(const int32 Contro
 		// This can only be done in the game thread since it uses the UEnum directly
 		if (IsInGameThread())
 		{
-			const UEnum* EnumPtr = FindObject<UEnum>(ANY_PACKAGE, TEXT("EHandKeypoint"), true);
+			const UEnum* EnumPtr = FindObject<UEnum>(nullptr, TEXT("/Script/HeadMountedDisplay.EHandKeypoint"), true);
 			check(EnumPtr != nullptr);
 			bool bUseRightHand = false;
 			FString SourceString = MotionSource.ToString();
@@ -437,7 +437,7 @@ void FOpenXRHandTracking::EnumerateSources(TArray<FMotionControllerSource>& Sour
 
 	SourcesOut.Empty(EHandKeypointCount * 2);
 
-	const UEnum* EnumPtr = FindObject<UEnum>(ANY_PACKAGE, TEXT("EHandKeypoint"), true);
+	const UEnum* EnumPtr = FindObject<UEnum>(nullptr, TEXT("/Script/HeadMountedDisplay.EHandKeypoint"), true);
 	check(EnumPtr != nullptr);
 	for (int32 Keypoint = 0; Keypoint < EHandKeypointCount; Keypoint++)
 	{
diff --git a/Engine/Plugins/Runtime/OpenXRHandTracking/Source/OpenXRHandTracking/Private/OpenXRHandTrackingLiveLink.cpp b/Engine/Plugins/Runtime/OpenXRHandTracking/Source/OpenXRHandTracking/Private/OpenXRHandTrackingLiveLink.cpp
--- a/Engine/Plugins/Runtime/OpenXRHandTracking/Source/OpenXRHandTracking/Private/OpenXRHandTrackingLiveLink.cpp
+++ b/Engine/Plugins/Runtime/OpenXRHandTracking/Source/OpenXRHandTracking/Private/OpenXRHandTrackingLiveLink.cpp
@@ -69,7 +69,7 @@ void FOpenXRHandTracking::SetupLiveLinkData()
 	BoneParents.Reserve(EHandKeypointCount);
 	BoneKeypoints.Reserve(EHandKeypointCount);
 
-	const UEnum* EnumPtr = FindObject<UEnum>(ANY_PACKAGE, TEXT("EHandKeypoint"), true);
+	const UEnum* EnumPtr = FindObject<UEnum>(nullptr, TEXT("/Script/HeadMountedDisplay.EHandKeypoint"), true);
 	check(EnumPtr != nullptr);
 	// Iterate through all of the Keypoints building the skeleton info for it
 	for (int32 Keypoint = 0; Keypoint < EHandKeypointCount; Keypoint++)
diff --git a/Engine/Plugins/Runtime/OpenXRHandTracking/Source/OpenXRHandTracking/Private/OpenXRHandTrackingLiveLinkRemapAsset.cpp b/Engine/Plugins/Runtime/OpenXRHandTracking/Source/OpenXRHandTracking/Private/OpenXRHandTrackingLiveLinkRemapAsset.cpp
--- a/Engine/Plugins/Runtime/OpenXRHandTracking/Source/OpenXRHandTracking/Private/OpenXRHandTrackingLiveLinkRemapAsset.cpp
+++ b/Engine/Plugins/Runtime/OpenXRHandTracking/Source/OpenXRHandTracking/Private/OpenXRHandTrackingLiveLinkRemapAsset.cpp
@@ -21,7 +21,7 @@ void UOpenXRHandTrackingLiveLinkRemapAsset::PostInitProperties()
 	if (HandTrackingBoneNameMap.Num() == 0)
 	{
 		// Fill the first time use of the bone name map
-		const UEnum* EnumPtr = FindObject<UEnum>(ANY_PACKAGE, TEXT("EHandKeypoint"), true);
+		const UEnum* EnumPtr = FindObject<UEnum>(nullptr, TEXT("/Script/HeadMountedDisplay.EHandKeypoint"), true);
 		check(EnumPtr != nullptr);
 		// Iterate through all of the Keypoints building the skeleton info for it
 		for (int32 Keypoint = 0; Keypoint < EHandKeypointCount; Keypoint++)
diff --git a/Engine/Plugins/Runtime/ResonanceAudio/Source/ResonanceAudio/Private/ResonanceAudioSettings.h b/Engine/Plugins/Runtime/ResonanceAudio/Source/ResonanceAudio/Private/ResonanceAudioSettings.h
--- a/Engine/Plugins/Runtime/ResonanceAudio/Source/ResonanceAudio/Private/ResonanceAudioSettings.h
+++ b/Engine/Plugins/Runtime/ResonanceAudio/Source/ResonanceAudio/Private/ResonanceAudioSettings.h
@@ -17,7 +17,7 @@ public:
 	UResonanceAudioSettings();
 
 	// Reference to submix where reverb plugin audio is routed.
-	UPROPERTY(GlobalConfig, EditAnywhere, Category = "Reverb", meta = (AllowedClasses = "SoundSubmix"))
+	UPROPERTY(GlobalConfig, EditAnywhere, Category = "Reverb", meta = (AllowedClasses = "/Script/Engine.SoundSubmix"))
 	FSoftObjectPath OutputSubmix;
 
 	// Global Quality mode to use for directional sound sources.
@@ -25,10 +25,10 @@ public:
 	ERaQualityMode QualityMode;
 
 	// Default settings for global reverb: This is overridden when a player enters Audio Volumes.
-	UPROPERTY(GlobalConfig, EditAnywhere, Category = General, meta = (AllowedClasses = "ResonanceAudioReverbPluginPreset"))
+	UPROPERTY(GlobalConfig, EditAnywhere, Category = General, meta = (AllowedClasses = "/Script/ResonanceAudioEditor.ResonanceAudioReverbPluginPreset"))
 	FSoftObjectPath GlobalReverbPreset;
 
 	// Default settings for global source settings
-	UPROPERTY(GlobalConfig, EditAnywhere, Category = General, meta = (AllowedClasses = "ResonanceAudioSpatializationSourceSettings"))
+	UPROPERTY(GlobalConfig, EditAnywhere, Category = General, meta = (AllowedClasses = "/Script/ResonanceAudio.ResonanceAudioSpatializationSourceSettings"))
 	FSoftObjectPath GlobalSourcePreset;
 };
diff --git a/Engine/Plugins/Runtime/SignificanceManager/Source/SignificanceManager/Public/SignificanceManager.h b/Engine/Plugins/Runtime/SignificanceManager/Source/SignificanceManager/Public/SignificanceManager.h
--- a/Engine/Plugins/Runtime/SignificanceManager/Source/SignificanceManager/Public/SignificanceManager.h
+++ b/Engine/Plugins/Runtime/SignificanceManager/Source/SignificanceManager/Public/SignificanceManager.h
@@ -244,7 +244,7 @@ private:
 	TArray<FSequentialPostWorkPair> ObjWithSequentialPostWork;
 
 	// Game specific significance class to instantiate
-	UPROPERTY(globalconfig, noclear, EditAnywhere, Category=DefaultClasses, meta=(MetaClass="SignificanceManager", DisplayName="Significance Manager Class"))
+	UPROPERTY(globalconfig, noclear, EditAnywhere, Category=DefaultClasses, meta=(MetaClass="/Script/SignificanceManager.SignificanceManager", DisplayName="Significance Manager Class"))
 	FSoftClassPath SignificanceManagerClassName;
 
 	// Callback function registered with HUD to supply debug info when ShowDebug SignificanceManager has been entered on the console
diff --git a/Engine/Plugins/Runtime/StateTree/Source/StateTreeEditorModule/Public/StateTreeState.h b/Engine/Plugins/Runtime/StateTree/Source/StateTreeEditorModule/Public/StateTreeState.h
--- a/Engine/Plugins/Runtime/StateTree/Source/StateTreeEditorModule/Public/StateTreeState.h
+++ b/Engine/Plugins/Runtime/StateTree/Source/StateTreeEditorModule/Public/StateTreeState.h
@@ -74,7 +74,7 @@ struct STATETREEEDITORMODULE_API FStateTreeTransition
 	float GateDelay = 0.0f;
 	
 	UPROPERTY(EditDefaultsOnly, Category = Transition, meta = (BaseStruct = "StateTreeConditionBase", BaseClass = "StateTreeConditionBlueprintBase"))
-	TArray<FStateTreeEditorNode> Conditions;
+	TArray<FStateTreeConditionItem> Conditions;
 };
 
 USTRUCT()
@@ -193,18 +193,18 @@ public:
 	UPROPERTY()
 	FGuid ID;
 
-	UPROPERTY(EditDefaultsOnly, Category = "Enter Conditions", meta = (BaseStruct = "StateTreeConditionBase", BaseClass = "StateTreeConditionBlueprintBase"))
+	UPROPERTY(EditDefaultsOnly, Category = "Enter Conditions", meta = (BaseStruct = "/Script/StateTreeModule.StateTreeConditionBase", BaseClass = "/Script/StateTreeModule.StateTreeConditionBlueprintBase"))
 	TArray<FStateTreeEditorNode> EnterConditions;
 
 	UE_DEPRECATED(5.1, "Evaluators are moved into UStateTreeEditorData. This property will be removed for 5.1.")
-	UPROPERTY(meta = (DeprecatedProperty, BaseStruct = "StateTreeEvaluatorBase", BaseClass = "StateTreeEvaluatorBlueprintBase"))
+	UPROPERTY(meta = (DeprecatedProperty, BaseStruct = "/Script/StateTreeModule.StateTreeEvaluatorBase", BaseClass = "/Script/StateTreeModule.StateTreeEvaluatorBlueprintBase"))
 	TArray<FStateTreeEditorNode> Evaluators_DEPRECATED;
 
-	UPROPERTY(EditDefaultsOnly, Category = "Tasks", meta = (BaseStruct = "StateTreeTaskBase", BaseClass = "StateTreeTaskBlueprintBase"))
+	UPROPERTY(EditDefaultsOnly, Category = "Tasks", meta = (BaseStruct = "/Script/StateTreeModule.StateTreeTaskBase", BaseClass = "/Script/StateTreeModule.StateTreeTaskBlueprintBase"))
 	TArray<FStateTreeEditorNode> Tasks;
 
 	// Single item used when schema calls for single task per state.
-	UPROPERTY(EditDefaultsOnly, Category = "Task", meta = (BaseStruct = "StateTreeTaskBase", BaseClass = "StateTreeTaskBlueprintBase"))
+	UPROPERTY(EditDefaultsOnly, Category = "Task", meta = (BaseStruct = "/Script/StateTreeModule.StateTreeTaskBase", BaseClass = "/Script/StateTreeModule.StateTreeTaskBlueprintBase"))
 	FStateTreeEditorNode SingleTask;
 
 	UPROPERTY(EditDefaultsOnly, Category = "Transitions")
diff --git a/Engine/Plugins/Runtime/Steam/SteamAudio/Source/SteamAudio/Private/SteamAudioSettings.h b/Engine/Plugins/Runtime/Steam/SteamAudio/Source/SteamAudio/Private/SteamAudioSettings.h
--- a/Engine/Plugins/Runtime/Steam/SteamAudio/Source/SteamAudio/Private/SteamAudioSettings.h
+++ b/Engine/Plugins/Runtime/Steam/SteamAudio/Source/SteamAudio/Private/SteamAudioSettings.h
@@ -33,7 +33,7 @@ public:
 	virtual bool CanEditChange(const FProperty* InProperty) const override;
 #endif
 
-	UPROPERTY(GlobalConfig, EditAnywhere, Category = General, meta = (AllowedClasses = "SoundSubmix"))
+	UPROPERTY(GlobalConfig, EditAnywhere, Category = General, meta = (AllowedClasses = "/Script/Engine.SoundSubmix"))
 	FSoftObjectPath OutputSubmix;
 
 	// Which audio engine to use.
diff --git a/Engine/Plugins/Runtime/WindowsMixedReality/Source/WindowsMixedRealityHandTracking/Private/WindowsMixedRealityHandTracking.cpp b/Engine/Plugins/Runtime/WindowsMixedReality/Source/WindowsMixedRealityHandTracking/Private/WindowsMixedRealityHandTracking.cpp
--- a/Engine/Plugins/Runtime/WindowsMixedReality/Source/WindowsMixedRealityHandTracking/Private/WindowsMixedRealityHandTracking.cpp
+++ b/Engine/Plugins/Runtime/WindowsMixedReality/Source/WindowsMixedRealityHandTracking/Private/WindowsMixedRealityHandTracking.cpp
@@ -247,7 +247,7 @@ bool FWindowsMixedRealityHandTracking::GetControllerOrientationAndPosition(const
 		// This can only be done in the game thread since it uses the UEnum directly
 		if (IsInGameThread())
 		{
-			const UEnum* EnumPtr = FindObject<UEnum>(ANY_PACKAGE, TEXT("EWMRHandKeypoint"), true);
+			const UEnum* EnumPtr = FindObject<UEnum>(nullptr, TEXT("/Script/WindowsMixedRealityHandTracking.EWMRHandKeypoint"), true);
 			check(EnumPtr != nullptr);
 			bool bUseRightHand = false;
 			FString SourceString = MotionSource.ToString();
@@ -339,7 +339,7 @@ void FWindowsMixedRealityHandTracking::EnumerateSources(TArray<FMotionController
 
 	SourcesOut.Empty(EWMRHandKeypointCount);
 
-	const UEnum* EnumPtr = FindObject<UEnum>(ANY_PACKAGE, TEXT("EWMRHandKeypoint"), true);
+	const UEnum* EnumPtr = FindObject<UEnum>(nullptr, TEXT("/Script/WindowsMixedRealityHandTracking.EWMRHandKeypoint"), true);
 	check(EnumPtr != nullptr);
 	for (int32 Keypoint = 0; Keypoint < EWMRHandKeypointCount; Keypoint++)
 	{
diff --git a/Engine/Plugins/Runtime/WindowsMixedReality/Source/WindowsMixedRealityHandTracking/Private/WindowsMixedRealityHandTrackingLiveLink.cpp b/Engine/Plugins/Runtime/WindowsMixedReality/Source/WindowsMixedRealityHandTracking/Private/WindowsMixedRealityHandTrackingLiveLink.cpp
--- a/Engine/Plugins/Runtime/WindowsMixedReality/Source/WindowsMixedRealityHandTracking/Private/WindowsMixedRealityHandTrackingLiveLink.cpp
+++ b/Engine/Plugins/Runtime/WindowsMixedReality/Source/WindowsMixedRealityHandTracking/Private/WindowsMixedRealityHandTrackingLiveLink.cpp
@@ -69,7 +69,7 @@ void FWindowsMixedRealityHandTracking::SetupLiveLinkData()
 	BoneParents.Reserve(EWMRHandKeypointCount);
 	BoneKeypoints.Reserve(EWMRHandKeypointCount);
 
-	const UEnum* EnumPtr = FindObject<UEnum>(ANY_PACKAGE, TEXT("EWMRHandKeypoint"), true);
+	const UEnum* EnumPtr = FindObject<UEnum>(nullptr, TEXT("/Script/WindowsMixedRealityHandTracking.EWMRHandKeypoint"), true);
 	check(EnumPtr != nullptr);
 	// Iterate through all of the Keypoints building the skeleton info for it
 	for (int32 Keypoint = 0; Keypoint < EWMRHandKeypointCount; Keypoint++)
diff --git a/Engine/Plugins/Runtime/WindowsMixedReality/Source/WindowsMixedRealityHandTracking/Private/WindowsMixedRealityHandTrackingLiveLinkRemapAsset.cpp b/Engine/Plugins/Runtime/WindowsMixedReality/Source/WindowsMixedRealityHandTracking/Private/WindowsMixedRealityHandTrackingLiveLinkRemapAsset.cpp
--- a/Engine/Plugins/Runtime/WindowsMixedReality/Source/WindowsMixedRealityHandTracking/Private/WindowsMixedRealityHandTrackingLiveLinkRemapAsset.cpp
+++ b/Engine/Plugins/Runtime/WindowsMixedReality/Source/WindowsMixedRealityHandTracking/Private/WindowsMixedRealityHandTrackingLiveLinkRemapAsset.cpp
@@ -21,7 +21,7 @@ void UWindowsMixedRealityHandTrackingLiveLinkRemapAsset::PostInitProperties()
 	if (HandTrackingBoneNameMap.Num() == 0)
 	{
 		// Fill the first time use of the bone name map
-		const UEnum* EnumPtr = FindObject<UEnum>(ANY_PACKAGE, TEXT("EWMRHandKeypoint"), true);
+		const UEnum* EnumPtr = FindObject<UEnum>(nullptr, TEXT("/Script/WindowsMixedRealityHandTracking.EWMRHandKeypoint"), true);
 		check(EnumPtr != nullptr);
 		// Iterate through all of the Keypoints building the skeleton info for it
 		for (int32 Keypoint = 0; Keypoint < EWMRHandKeypointCount; Keypoint++)
