Deprecate Paper2D editor Input functions in favor of the newer version

#rnx

diff --git a/Engine/Plugins/2D/Paper2D/Source/Paper2DEditor/Private/FlipbookEditor/FlipbookEditorViewportClient.cpp b/Engine/Plugins/2D/Paper2D/Source/Paper2DEditor/Private/FlipbookEditor/FlipbookEditorViewportClient.cpp
--- a/Engine/Plugins/2D/Paper2D/Source/Paper2DEditor/Private/FlipbookEditor/FlipbookEditorViewportClient.cpp
+++ b/Engine/Plugins/2D/Paper2D/Source/Paper2DEditor/Private/FlipbookEditor/FlipbookEditorViewportClient.cpp
@@ -109,12 +109,12 @@ void FFlipbookEditorViewportClient::Tick(float DeltaSeconds)
 	}
 }
 
-bool FFlipbookEditorViewportClient::InputKey(FViewport* InViewport, int32 ControllerId, FKey Key, EInputEvent Event, float AmountDepressed, bool bGamepad)
+bool FFlipbookEditorViewportClient::InputKey(const FInputKeyEventArgs& EventArgs)
 {
 	bool bHandled = false;
 
 	// Pass keys to standard controls, if we didn't consume input
-	return (bHandled) ? true : FEditorViewportClient::InputKey(InViewport, ControllerId, Key, Event, AmountDepressed, bGamepad);
+	return (bHandled) ? true : FEditorViewportClient::InputKey(EventArgs);
 }
 
 FLinearColor FFlipbookEditorViewportClient::GetBackgroundColor() const
diff --git a/Engine/Plugins/2D/Paper2D/Source/Paper2DEditor/Private/FlipbookEditor/FlipbookEditorViewportClient.h b/Engine/Plugins/2D/Paper2D/Source/Paper2DEditor/Private/FlipbookEditor/FlipbookEditorViewportClient.h
--- a/Engine/Plugins/2D/Paper2D/Source/Paper2DEditor/Private/FlipbookEditor/FlipbookEditorViewportClient.h
+++ b/Engine/Plugins/2D/Paper2D/Source/Paper2DEditor/Private/FlipbookEditor/FlipbookEditorViewportClient.h
@@ -28,7 +28,7 @@ public:
 	// End of FViewportClient interface
 
 	// FEditorViewportClient interface
-	virtual bool InputKey(FViewport* Viewport, int32 ControllerId, FKey Key, EInputEvent Event, float AmountDepressed, bool bGamepad) override;
+	virtual bool InputKey(const FInputKeyEventArgs& EventArgs) override;
 	virtual FLinearColor GetBackgroundColor() const override;
 	// End of FEditorViewportClient interface
 
diff --git a/Engine/Plugins/2D/Paper2D/Source/Paper2DEditor/Private/SpriteEditor/SpriteEditorViewportClient.cpp b/Engine/Plugins/2D/Paper2D/Source/Paper2DEditor/Private/SpriteEditor/SpriteEditorViewportClient.cpp
--- a/Engine/Plugins/2D/Paper2D/Source/Paper2DEditor/Private/SpriteEditor/SpriteEditorViewportClient.cpp
+++ b/Engine/Plugins/2D/Paper2D/Source/Paper2DEditor/Private/SpriteEditor/SpriteEditorViewportClient.cpp
@@ -839,10 +839,10 @@ UPaperSprite* FSpriteEditorViewportClient::CreateNewSprite(const FIntPoint& TopL
 	return CreatedSprite;
 }
 
-bool FSpriteEditorViewportClient::InputKey(FViewport* InViewport, int32 ControllerId, FKey Key, EInputEvent Event, float AmountDepressed, bool bGamepad)
+bool FSpriteEditorViewportClient::InputKey(const FInputKeyEventArgs& EventArgs)
 {
 	bool bHandled = false;
-	FInputEventState InputState(InViewport, Key, Event);
+	FInputEventState InputState(EventArgs.Viewport, EventArgs.Key, EventArgs.Event);
 
 	// Handle marquee tracking in source region edit mode
 	if (IsInSourceRegionEditMode())
@@ -851,7 +851,7 @@ bool FSpriteEditorViewportClient::InputKey(FViewport* InViewport, int32 Controll
 		check(GeometryEditMode);
 
 		const bool bMarqueeStartModifier = InputState.IsCtrlButtonPressed();
-		if (GeometryEditMode->ProcessMarquee(InViewport, Key, Event, bMarqueeStartModifier))
+		if (GeometryEditMode->ProcessMarquee(EventArgs.Viewport, EventArgs.Key, EventArgs.Event, bMarqueeStartModifier))
 		{
 			FIntPoint TextureSpaceStartPos;
 			FIntPoint TextureSpaceDimensions;
@@ -864,7 +864,7 @@ bool FSpriteEditorViewportClient::InputKey(FViewport* InViewport, int32 Controll
 	}
 	
 	// Pass keys to standard controls, if we didn't consume input
-	return (bHandled) ? true : FEditorViewportClient::InputKey(InViewport, ControllerId, Key, Event, AmountDepressed, bGamepad);
+	return (bHandled) ? true : FEditorViewportClient::InputKey(EventArgs);
 }
 
 void FSpriteEditorViewportClient::TrackingStarted(const struct FInputEventState& InInputState, bool bIsDragging, bool bNudge)
diff --git a/Engine/Plugins/2D/Paper2D/Source/Paper2DEditor/Private/SpriteEditor/SpriteEditorViewportClient.h b/Engine/Plugins/2D/Paper2D/Source/Paper2DEditor/Private/SpriteEditor/SpriteEditorViewportClient.h
--- a/Engine/Plugins/2D/Paper2D/Source/Paper2DEditor/Private/SpriteEditor/SpriteEditorViewportClient.h
+++ b/Engine/Plugins/2D/Paper2D/Source/Paper2DEditor/Private/SpriteEditor/SpriteEditorViewportClient.h
@@ -41,7 +41,7 @@ public:
 
 	// FEditorViewportClient interface
 	virtual void ProcessClick(FSceneView& View, HHitProxy* HitProxy, FKey Key, EInputEvent Event, uint32 HitX, uint32 HitY) override;
-	virtual bool InputKey(FViewport* Viewport, int32 ControllerId, FKey Key, EInputEvent Event, float AmountDepressed, bool bGamepad) override;
+	virtual bool InputKey(const FInputKeyEventArgs& EventArgs) override;
 	virtual void TrackingStarted(const struct FInputEventState& InInputState, bool bIsDragging, bool bNudge) override;
 	virtual void TrackingStopped() override;
 	virtual FLinearColor GetBackgroundColor() const override;
