PaperSprite baked data needs to be refreshed before rebuilding the render data or the preview won't be up to date

#rn Fixed an issue where the preview in the Paper2D Sprite Editor wouldn't update when adjusting the UV Coordinates

diff --git a/Engine/Plugins/2D/Paper2D/Source/Paper2D/Private/PaperSprite.cpp b/Engine/Plugins/2D/Paper2D/Source/Paper2D/Private/PaperSprite.cpp
--- a/Engine/Plugins/2D/Paper2D/Source/Paper2D/Private/PaperSprite.cpp
+++ b/Engine/Plugins/2D/Paper2D/Source/Paper2D/Private/PaperSprite.cpp
@@ -636,6 +636,11 @@ void UPaperSprite::PostEditChangeProperty(FPropertyChangedEvent& PropertyChanged
 		bBothModified = false;
 	}
 
+	if (bBothModified)
+	{
+		RefreshBakedData();
+	}
+
 	if (bCollisionDataModified || bBothModified)
 	{
 		RebuildCollisionData();
@@ -646,10 +651,7 @@ void UPaperSprite::PostEditChangeProperty(FPropertyChangedEvent& PropertyChanged
 		RebuildRenderData();
 	}
 
-	if (bBothModified)
-	{
-		RefreshBakedData();
-	}
+	
 
 	Super::PostEditChangeProperty(PropertyChangedEvent);
 }
@@ -1453,6 +1455,7 @@ void UPaperSprite::InitializeSprite(const FSpriteAssetInitParameters& InitParams
 	}
 
 	RefreshBakedData();
+	CustomPivotPoint = GetPivotPosition();
 }
 
 void UPaperSprite::SetTrim(bool bTrimmed, const FVector2D& OriginInSourceImage, const FVector2D& SourceImageDimension, bool bRebuildData /*= true*/)
