Disable physics on Paper2D cursor component to mitigate out of order destruction of state leading to crashes when closing tool tabs

diff --git a/Engine/Plugins/2D/Paper2D/Source/Paper2DEditor/Private/TileMapEditing/EdModeTileMap.cpp b/Engine/Plugins/2D/Paper2D/Source/Paper2DEditor/Private/TileMapEditing/EdModeTileMap.cpp
--- a/Engine/Plugins/2D/Paper2D/Source/Paper2DEditor/Private/TileMapEditing/EdModeTileMap.cpp
+++ b/Engine/Plugins/2D/Paper2D/Source/Paper2DEditor/Private/TileMapEditing/EdModeTileMap.cpp
@@ -175,6 +175,7 @@ void FEdModeTileMap::Enter()
 	CursorPreviewComponent->bShowPerLayerGridWhenSelected = false;
 	CursorPreviewComponent->bShowPerLayerGridWhenUnselected = false;
 	CursorPreviewComponent->bShowOutlineWhenUnselected = false;
+	CursorPreviewComponent->SetCollisionEnabled(ECollisionEnabled::NoCollision);
 	CursorPreviewComponent->UpdateBounds();
 	CursorPreviewComponent->AddToRoot();
 	CursorPreviewComponent->RegisterComponentWithWorld(World);
