[Paper2d]
* Ran iwyu to remove includes not needed

diff --git a/Engine/Plugins/2D/Paper2D/Source/Paper2D/Classes/IntMargin.h b/Engine/Plugins/2D/Paper2D/Source/Paper2D/Classes/IntMargin.h
--- a/Engine/Plugins/2D/Paper2D/Source/Paper2D/Classes/IntMargin.h
+++ b/Engine/Plugins/2D/Paper2D/Source/Paper2D/Classes/IntMargin.h
@@ -2,8 +2,6 @@
 
 #pragma once
 
-#include "CoreMinimal.h"
-#include "UObject/ObjectMacros.h"
 #include "IntMargin.generated.h"
 
 /**
@@ -126,3 +124,8 @@ public:
 		return FIntPoint(Left + Right, Top + Bottom);
 	}
 };
+
+#if UE_ENABLE_INCLUDE_ORDER_DEPRECATED_IN_5_2
+#include "CoreMinimal.h"
+#include "UObject/ObjectMacros.h"
+#endif
diff --git a/Engine/Plugins/2D/Paper2D/Source/Paper2D/Classes/PaperCharacter.h b/Engine/Plugins/2D/Paper2D/Source/Paper2D/Classes/PaperCharacter.h
--- a/Engine/Plugins/2D/Paper2D/Source/Paper2D/Classes/PaperCharacter.h
+++ b/Engine/Plugins/2D/Paper2D/Source/Paper2D/Classes/PaperCharacter.h
@@ -2,11 +2,11 @@
 
 #pragma once
 
-#include "CoreMinimal.h"
-#include "UObject/ObjectMacros.h"
 #include "GameFramework/Character.h"
 #include "PaperCharacter.generated.h"
 
+class UPaperFlipbookComponent;
+
 // APaperCharacter behaves like ACharacter, but uses a UPaperFlipbookComponent instead of a USkeletalMeshComponent as a visual representation
 // Note: The variable named Mesh will not be set up on this actor!
 UCLASS()
@@ -28,3 +28,8 @@ public:
 	/** Returns Sprite subobject **/
 	FORCEINLINE class UPaperFlipbookComponent* GetSprite() const { return Sprite; }
 };
+
+#if UE_ENABLE_INCLUDE_ORDER_DEPRECATED_IN_5_2
+#include "CoreMinimal.h"
+#include "UObject/ObjectMacros.h"
+#endif
diff --git a/Engine/Plugins/2D/Paper2D/Source/Paper2D/Classes/PaperFlipbook.h b/Engine/Plugins/2D/Paper2D/Source/Paper2D/Classes/PaperFlipbook.h
--- a/Engine/Plugins/2D/Paper2D/Source/Paper2D/Classes/PaperFlipbook.h
+++ b/Engine/Plugins/2D/Paper2D/Source/Paper2D/Classes/PaperFlipbook.h
@@ -2,13 +2,11 @@
 
 #pragma once
 
-#include "CoreMinimal.h"
-#include "UObject/ObjectMacros.h"
-#include "UObject/Object.h"
-#include "Engine/EngineTypes.h"
-#include "UObject/ScriptMacros.h"
 #include "PaperFlipbook.generated.h"
 
+struct FComponentSocketDescription;
+struct FPropertyChangedEvent;
+
 class UMaterialInterface;
 class UPaperSprite;
 
@@ -191,3 +189,11 @@ public:
 		return SourceFlipbook;
 	}
 };
+
+#if UE_ENABLE_INCLUDE_ORDER_DEPRECATED_IN_5_2
+#include "CoreMinimal.h"
+#include "Engine/EngineTypes.h"
+#include "UObject/Object.h"
+#include "UObject/ObjectMacros.h"
+#include "UObject/ScriptMacros.h"
+#endif
diff --git a/Engine/Plugins/2D/Paper2D/Source/Paper2D/Classes/PaperFlipbookActor.h b/Engine/Plugins/2D/Paper2D/Source/Paper2D/Classes/PaperFlipbookActor.h
--- a/Engine/Plugins/2D/Paper2D/Source/Paper2D/Classes/PaperFlipbookActor.h
+++ b/Engine/Plugins/2D/Paper2D/Source/Paper2D/Classes/PaperFlipbookActor.h
@@ -2,11 +2,11 @@
 
 #pragma once
 
-#include "CoreMinimal.h"
-#include "UObject/ObjectMacros.h"
 #include "GameFramework/Actor.h"
 #include "PaperFlipbookActor.generated.h"
 
+class UPaperFlipbookComponent;
+
 /**
  * An instance of a UPaperFlipbook in a level.
  *
@@ -34,3 +34,8 @@ public:
 	FORCEINLINE class UPaperFlipbookComponent* GetRenderComponent() const { return RenderComponent; }
 };
 
+
+#if UE_ENABLE_INCLUDE_ORDER_DEPRECATED_IN_5_2
+#include "CoreMinimal.h"
+#include "UObject/ObjectMacros.h"
+#endif
diff --git a/Engine/Plugins/2D/Paper2D/Source/Paper2D/Classes/PaperFlipbookComponent.h b/Engine/Plugins/2D/Paper2D/Source/Paper2D/Classes/PaperFlipbookComponent.h
--- a/Engine/Plugins/2D/Paper2D/Source/Paper2D/Classes/PaperFlipbookComponent.h
+++ b/Engine/Plugins/2D/Paper2D/Source/Paper2D/Classes/PaperFlipbookComponent.h
@@ -2,13 +2,13 @@
 
 #pragma once
 
-#include "CoreMinimal.h"
-#include "UObject/ObjectMacros.h"
-#include "Engine/EngineTypes.h"
-#include "Engine/TextureStreamingTypes.h"
 #include "Components/MeshComponent.h"
 #include "PaperFlipbookComponent.generated.h"
 
+class FStreamingTextureLevelContext;
+struct FComponentSocketDescription;
+struct FStreamingRenderAssetPrimitiveInfo;
+
 class FPrimitiveSceneProxy;
 class UBodySetup;
 class UPaperFlipbook;
@@ -212,3 +212,10 @@ public:
 	virtual UBodySetup* GetBodySetup() override;
 	// End of UPrimitiveComponent interface
 };
+
+#if UE_ENABLE_INCLUDE_ORDER_DEPRECATED_IN_5_2
+#include "CoreMinimal.h"
+#include "Engine/EngineTypes.h"
+#include "Engine/TextureStreamingTypes.h"
+#include "UObject/ObjectMacros.h"
+#endif
diff --git a/Engine/Plugins/2D/Paper2D/Source/Paper2D/Classes/PaperGroupedSpriteActor.h b/Engine/Plugins/2D/Paper2D/Source/Paper2D/Classes/PaperGroupedSpriteActor.h
--- a/Engine/Plugins/2D/Paper2D/Source/Paper2D/Classes/PaperGroupedSpriteActor.h
+++ b/Engine/Plugins/2D/Paper2D/Source/Paper2D/Classes/PaperGroupedSpriteActor.h
@@ -2,8 +2,6 @@
 
 #pragma once
 
-#include "CoreMinimal.h"
-#include "UObject/ObjectMacros.h"
 #include "GameFramework/Actor.h"
 #include "PaperGroupedSpriteActor.generated.h"
 
@@ -34,3 +32,8 @@ public:
 	/** Returns RenderComponent subobject **/
 	FORCEINLINE UPaperGroupedSpriteComponent* GetRenderComponent() const { return RenderComponent; }
 };
+
+#if UE_ENABLE_INCLUDE_ORDER_DEPRECATED_IN_5_2
+#include "CoreMinimal.h"
+#include "UObject/ObjectMacros.h"
+#endif
diff --git a/Engine/Plugins/2D/Paper2D/Source/Paper2D/Classes/PaperGroupedSpriteComponent.h b/Engine/Plugins/2D/Paper2D/Source/Paper2D/Classes/PaperGroupedSpriteComponent.h
--- a/Engine/Plugins/2D/Paper2D/Source/Paper2D/Classes/PaperGroupedSpriteComponent.h
+++ b/Engine/Plugins/2D/Paper2D/Source/Paper2D/Classes/PaperGroupedSpriteComponent.h
@@ -2,10 +2,7 @@
 
 #pragma once
 
-#include "CoreMinimal.h"
-#include "UObject/ObjectMacros.h"
 #include "Components/MeshComponent.h"
-#include "Physics/PhysicsInterfaceDeclares.h"
 
 #include "PaperGroupedSpriteComponent.generated.h"
 
@@ -185,3 +182,9 @@ protected:
 
 	friend FGroupedSpriteSceneProxy;
 };
+
+#if UE_ENABLE_INCLUDE_ORDER_DEPRECATED_IN_5_2
+#include "CoreMinimal.h"
+#include "Physics/PhysicsInterfaceDeclares.h"
+#include "UObject/ObjectMacros.h"
+#endif
diff --git a/Engine/Plugins/2D/Paper2D/Source/Paper2D/Classes/PaperRuntimeSettings.h b/Engine/Plugins/2D/Paper2D/Source/Paper2D/Classes/PaperRuntimeSettings.h
--- a/Engine/Plugins/2D/Paper2D/Source/Paper2D/Classes/PaperRuntimeSettings.h
+++ b/Engine/Plugins/2D/Paper2D/Source/Paper2D/Classes/PaperRuntimeSettings.h
@@ -2,9 +2,6 @@
 
 #pragma once
 
-#include "CoreMinimal.h"
-#include "UObject/ObjectMacros.h"
-#include "UObject/Object.h"
 #include "PaperRuntimeSettings.generated.h"
 
 /**
@@ -27,3 +24,9 @@ class PAPER2D_API UPaperRuntimeSettings : public UObject
 	UPROPERTY(EditAnywhere, config, Category=Settings)
 	bool bResizeSpriteDataToMatchTextures;
 };
+
+#if UE_ENABLE_INCLUDE_ORDER_DEPRECATED_IN_5_2
+#include "CoreMinimal.h"
+#include "UObject/Object.h"
+#include "UObject/ObjectMacros.h"
+#endif
diff --git a/Engine/Plugins/2D/Paper2D/Source/Paper2D/Classes/PaperSprite.h b/Engine/Plugins/2D/Paper2D/Source/Paper2D/Classes/PaperSprite.h
--- a/Engine/Plugins/2D/Paper2D/Source/Paper2D/Classes/PaperSprite.h
+++ b/Engine/Plugins/2D/Paper2D/Source/Paper2D/Classes/PaperSprite.h
@@ -2,15 +2,14 @@
 
 #pragma once
 
-#include "CoreMinimal.h"
-#include "UObject/ObjectMacros.h"
-#include "UObject/Object.h"
-#include "Engine/EngineTypes.h"
 #include "Interfaces/Interface_CollisionDataProvider.h"
 #include "SpriteEditorOnlyTypes.h"
 #include "Slate/SlateTextureAtlasInterface.h"
 #include "PaperSprite.generated.h"
 
+class UTexture2D;
+struct FComponentSocketDescription;
+
 class UMaterialInterface;
 class UPaperSpriteAtlas;
 
@@ -341,3 +340,10 @@ protected:
 	friend class FSingleTileEditorViewportClient;
 };
 
+
+#if UE_ENABLE_INCLUDE_ORDER_DEPRECATED_IN_5_2
+#include "CoreMinimal.h"
+#include "Engine/EngineTypes.h"
+#include "UObject/Object.h"
+#include "UObject/ObjectMacros.h"
+#endif
diff --git a/Engine/Plugins/2D/Paper2D/Source/Paper2D/Classes/PaperSpriteActor.h b/Engine/Plugins/2D/Paper2D/Source/Paper2D/Classes/PaperSpriteActor.h
--- a/Engine/Plugins/2D/Paper2D/Source/Paper2D/Classes/PaperSpriteActor.h
+++ b/Engine/Plugins/2D/Paper2D/Source/Paper2D/Classes/PaperSpriteActor.h
@@ -2,11 +2,11 @@
 
 #pragma once
 
-#include "CoreMinimal.h"
-#include "UObject/ObjectMacros.h"
 #include "GameFramework/Actor.h"
 #include "PaperSpriteActor.generated.h"
 
+class UPaperSpriteComponent;
+
 /**
  * An instance of a UPaperSprite in a level.
  *
@@ -32,3 +32,8 @@ public:
 	/** Returns RenderComponent subobject **/
 	FORCEINLINE class UPaperSpriteComponent* GetRenderComponent() const { return RenderComponent; }
 };
+
+#if UE_ENABLE_INCLUDE_ORDER_DEPRECATED_IN_5_2
+#include "CoreMinimal.h"
+#include "UObject/ObjectMacros.h"
+#endif
diff --git a/Engine/Plugins/2D/Paper2D/Source/Paper2D/Classes/PaperSpriteAtlas.h b/Engine/Plugins/2D/Paper2D/Source/Paper2D/Classes/PaperSpriteAtlas.h
--- a/Engine/Plugins/2D/Paper2D/Source/Paper2D/Classes/PaperSpriteAtlas.h
+++ b/Engine/Plugins/2D/Paper2D/Source/Paper2D/Classes/PaperSpriteAtlas.h
@@ -2,13 +2,12 @@
 
 #pragma once
 
-#include "CoreMinimal.h"
-#include "UObject/ObjectMacros.h"
-#include "UObject/Object.h"
-#include "Misc/Guid.h"
-#include "PaperSprite.h"
 #include "PaperSpriteAtlas.generated.h"
 
+class UTexture;
+enum TextureCompressionSettings : int;
+enum TextureFilter : int;
+
 USTRUCT()
 struct FPaperSpriteAtlasSlot
 {
@@ -81,11 +80,11 @@ class UPaperSpriteAtlas : public UObject
 
 	// Compression settings to use on atlas texture
 	UPROPERTY(EditAnywhere, Category=AtlasTexture)
-	TEnumAsByte<enum TextureCompressionSettings> CompressionSettings;
+	TEnumAsByte<TextureCompressionSettings> CompressionSettings;
 
 	// Texture filtering mode when sampling these textures
 	UPROPERTY(EditAnywhere, Category=AtlasTexture)
-	TEnumAsByte<enum TextureFilter> Filter;
+	TEnumAsByte<TextureFilter> Filter;
 
 	// List of generated atlas textures
 	UPROPERTY(VisibleAnywhere, Category=AtlasTexture)
@@ -128,3 +127,11 @@ public:
 	virtual void GetAssetRegistryTags(TArray<FAssetRegistryTag>& OutTags) const override;
 	// End of UObject interface
 };
+
+#if UE_ENABLE_INCLUDE_ORDER_DEPRECATED_IN_5_2
+#include "CoreMinimal.h"
+#include "Misc/Guid.h"
+#include "PaperSprite.h"
+#include "UObject/Object.h"
+#include "UObject/ObjectMacros.h"
+#endif
diff --git a/Engine/Plugins/2D/Paper2D/Source/Paper2D/Classes/PaperSpriteBlueprintLibrary.h b/Engine/Plugins/2D/Paper2D/Source/Paper2D/Classes/PaperSpriteBlueprintLibrary.h
--- a/Engine/Plugins/2D/Paper2D/Source/Paper2D/Classes/PaperSpriteBlueprintLibrary.h
+++ b/Engine/Plugins/2D/Paper2D/Source/Paper2D/Classes/PaperSpriteBlueprintLibrary.h
@@ -2,12 +2,11 @@
 
 #pragma once
 
-#include "CoreMinimal.h"
-#include "UObject/ObjectMacros.h"
 #include "Kismet/BlueprintFunctionLibrary.h"
-#include "Styling/SlateBrush.h"
 #include "PaperSpriteBlueprintLibrary.generated.h"
 
+struct FSlateBrush;
+
 class UPaperSprite;
 
 /**
@@ -22,3 +21,9 @@ public:
 	UFUNCTION(BlueprintPure, Category="Widget|Brush")
 	static FSlateBrush MakeBrushFromSprite(UPaperSprite* Sprite, int32 Width, int32 Height);
 };
+
+#if UE_ENABLE_INCLUDE_ORDER_DEPRECATED_IN_5_2
+#include "CoreMinimal.h"
+#include "Styling/SlateBrush.h"
+#include "UObject/ObjectMacros.h"
+#endif
diff --git a/Engine/Plugins/2D/Paper2D/Source/Paper2D/Classes/PaperSpriteComponent.h b/Engine/Plugins/2D/Paper2D/Source/Paper2D/Classes/PaperSpriteComponent.h
--- a/Engine/Plugins/2D/Paper2D/Source/Paper2D/Classes/PaperSpriteComponent.h
+++ b/Engine/Plugins/2D/Paper2D/Source/Paper2D/Classes/PaperSpriteComponent.h
@@ -2,14 +2,14 @@
 
 #pragma once
 
-#include "CoreMinimal.h"
-#include "UObject/ObjectMacros.h"
-#include "Engine/EngineTypes.h"
-#include "Engine/TextureStreamingTypes.h"
 #include "Components/MeshComponent.h"
 
 #include "PaperSpriteComponent.generated.h"
 
+class FStreamingTextureLevelContext;
+struct FComponentSocketDescription;
+struct FStreamingRenderAssetPrimitiveInfo;
+
 class FPrimitiveSceneProxy;
 class UPaperSprite;
 class UTexture;
@@ -102,3 +102,10 @@ public:
 protected:
 	friend class FPaperSpriteSceneProxy;
 };
+
+#if UE_ENABLE_INCLUDE_ORDER_DEPRECATED_IN_5_2
+#include "CoreMinimal.h"
+#include "Engine/EngineTypes.h"
+#include "Engine/TextureStreamingTypes.h"
+#include "UObject/ObjectMacros.h"
+#endif
diff --git a/Engine/Plugins/2D/Paper2D/Source/Paper2D/Classes/PaperTerrainActor.h b/Engine/Plugins/2D/Paper2D/Source/Paper2D/Classes/PaperTerrainActor.h
--- a/Engine/Plugins/2D/Paper2D/Source/Paper2D/Classes/PaperTerrainActor.h
+++ b/Engine/Plugins/2D/Paper2D/Source/Paper2D/Classes/PaperTerrainActor.h
@@ -2,12 +2,13 @@
 
 #pragma once
 
-#include "CoreMinimal.h"
-#include "UObject/ObjectMacros.h"
 #include "GameFramework/Actor.h"
 
 #include "PaperTerrainActor.generated.h"
 
+class UPaperTerrainComponent;
+class UPaperTerrainSplineComponent;
+
 /**
  * An instance of a piece of 2D terrain in the level
  */
@@ -41,3 +42,8 @@ public:
 	/** Returns RenderComponent subobject **/
 	FORCEINLINE class UPaperTerrainComponent* GetRenderComponent() const { return RenderComponent; }
 };
+
+#if UE_ENABLE_INCLUDE_ORDER_DEPRECATED_IN_5_2
+#include "CoreMinimal.h"
+#include "UObject/ObjectMacros.h"
+#endif
diff --git a/Engine/Plugins/2D/Paper2D/Source/Paper2D/Classes/PaperTerrainComponent.h b/Engine/Plugins/2D/Paper2D/Source/Paper2D/Classes/PaperTerrainComponent.h
--- a/Engine/Plugins/2D/Paper2D/Source/Paper2D/Classes/PaperTerrainComponent.h
+++ b/Engine/Plugins/2D/Paper2D/Source/Paper2D/Classes/PaperTerrainComponent.h
@@ -2,12 +2,13 @@
 
 #pragma once
 
-#include "CoreMinimal.h"
-#include "UObject/ObjectMacros.h"
 #include "Components/PrimitiveComponent.h"
-#include "SpriteDrawCall.h"
 #include "PaperTerrainComponent.generated.h"
 
+class UPaperSprite;
+namespace ESpriteCollisionMode { enum Type : int; }
+struct FSpriteDrawCallRecord;
+
 class FPrimitiveSceneProxy;
 class UBodySetup;
 class UMaterialInterface;
@@ -135,3 +136,9 @@ protected:
 	FTransform GetTransformAtDistance(float InDistance) const;
 };
 
+
+#if UE_ENABLE_INCLUDE_ORDER_DEPRECATED_IN_5_2
+#include "CoreMinimal.h"
+#include "SpriteDrawCall.h"
+#include "UObject/ObjectMacros.h"
+#endif
diff --git a/Engine/Plugins/2D/Paper2D/Source/Paper2D/Classes/PaperTileLayer.h b/Engine/Plugins/2D/Paper2D/Source/Paper2D/Classes/PaperTileLayer.h
--- a/Engine/Plugins/2D/Paper2D/Source/Paper2D/Classes/PaperTileLayer.h
+++ b/Engine/Plugins/2D/Paper2D/Source/Paper2D/Classes/PaperTileLayer.h
@@ -2,9 +2,6 @@
 
 #pragma once
 
-#include "CoreMinimal.h"
-#include "UObject/ObjectMacros.h"
-#include "UObject/Object.h"
 
 #include "PaperTileLayer.generated.h"
 
@@ -266,3 +263,9 @@ protected:
 	void ReallocateAndCopyMap();
 
 };
+
+#if UE_ENABLE_INCLUDE_ORDER_DEPRECATED_IN_5_2
+#include "CoreMinimal.h"
+#include "UObject/Object.h"
+#include "UObject/ObjectMacros.h"
+#endif
diff --git a/Engine/Plugins/2D/Paper2D/Source/Paper2D/Classes/PaperTileMap.h b/Engine/Plugins/2D/Paper2D/Source/Paper2D/Classes/PaperTileMap.h
--- a/Engine/Plugins/2D/Paper2D/Source/Paper2D/Classes/PaperTileMap.h
+++ b/Engine/Plugins/2D/Paper2D/Source/Paper2D/Classes/PaperTileMap.h
@@ -2,12 +2,11 @@
 
 #pragma once
 
-#include "CoreMinimal.h"
-#include "UObject/ObjectMacros.h"
-#include "UObject/Object.h"
-#include "SpriteEditorOnlyTypes.h"
 #include "PaperTileMap.generated.h"
 
+namespace ESpriteCollisionMode { enum Type : int; }
+struct FPropertyChangedEvent;
+
 class UBodySetup;
 class UMaterialInterface;
 class UPaperTileLayer;
@@ -242,3 +241,10 @@ public:
 protected:
 	virtual void UpdateBodySetup();
 };
+
+#if UE_ENABLE_INCLUDE_ORDER_DEPRECATED_IN_5_2
+#include "CoreMinimal.h"
+#include "SpriteEditorOnlyTypes.h"
+#include "UObject/Object.h"
+#include "UObject/ObjectMacros.h"
+#endif
diff --git a/Engine/Plugins/2D/Paper2D/Source/Paper2D/Classes/PaperTileMapActor.h b/Engine/Plugins/2D/Paper2D/Source/Paper2D/Classes/PaperTileMapActor.h
--- a/Engine/Plugins/2D/Paper2D/Source/Paper2D/Classes/PaperTileMapActor.h
+++ b/Engine/Plugins/2D/Paper2D/Source/Paper2D/Classes/PaperTileMapActor.h
@@ -2,12 +2,12 @@
 
 #pragma once
 
-#include "CoreMinimal.h"
-#include "UObject/ObjectMacros.h"
 #include "GameFramework/Actor.h"
 
 #include "PaperTileMapActor.generated.h"
 
+class UPaperTileMapComponent;
+
 /**
  * An instance of a UPaperTileMap in a level.
  *
@@ -33,3 +33,8 @@ public:
 	/** Returns RenderComponent subobject **/
 	FORCEINLINE class UPaperTileMapComponent* GetRenderComponent() const { return RenderComponent; }
 };
+
+#if UE_ENABLE_INCLUDE_ORDER_DEPRECATED_IN_5_2
+#include "CoreMinimal.h"
+#include "UObject/ObjectMacros.h"
+#endif
diff --git a/Engine/Plugins/2D/Paper2D/Source/Paper2D/Classes/PaperTileMapComponent.h b/Engine/Plugins/2D/Paper2D/Source/Paper2D/Classes/PaperTileMapComponent.h
--- a/Engine/Plugins/2D/Paper2D/Source/Paper2D/Classes/PaperTileMapComponent.h
+++ b/Engine/Plugins/2D/Paper2D/Source/Paper2D/Classes/PaperTileMapComponent.h
@@ -2,13 +2,13 @@
 
 #pragma once
 
-#include "CoreMinimal.h"
-#include "UObject/ObjectMacros.h"
-#include "PaperTileLayer.h"
 #include "Components/MeshComponent.h"
-#include "DynamicMeshBuilder.h"
 #include "PaperTileMapComponent.generated.h"
 
+class UPaperTileLayer;
+struct FDynamicMeshVertex;
+struct FPaperTileInfo;
+
 class FPaperTileMapRenderSceneProxy;
 class FPrimitiveSceneProxy;
 class UPaperTileMap;
@@ -235,3 +235,10 @@ public:
 	void GetRenderingStats(int32& OutNumTriangles, int32& OutNumBatches) const;
 #endif
 };
+
+#if UE_ENABLE_INCLUDE_ORDER_DEPRECATED_IN_5_2
+#include "CoreMinimal.h"
+#include "DynamicMeshBuilder.h"
+#include "PaperTileLayer.h"
+#include "UObject/ObjectMacros.h"
+#endif
diff --git a/Engine/Plugins/2D/Paper2D/Source/Paper2D/Classes/TileMapBlueprintLibrary.h b/Engine/Plugins/2D/Paper2D/Source/Paper2D/Classes/TileMapBlueprintLibrary.h
--- a/Engine/Plugins/2D/Paper2D/Source/Paper2D/Classes/TileMapBlueprintLibrary.h
+++ b/Engine/Plugins/2D/Paper2D/Source/Paper2D/Classes/TileMapBlueprintLibrary.h
@@ -2,12 +2,11 @@
 
 #pragma once
 
-#include "CoreMinimal.h"
-#include "UObject/ObjectMacros.h"
-#include "PaperTileLayer.h"
 #include "Kismet/BlueprintFunctionLibrary.h"
 #include "TileMapBlueprintLibrary.generated.h"
 
+struct FPaperTileInfo;
+
 class UPaperTileSet;
 
 /**
@@ -37,3 +36,9 @@ public:
 	UFUNCTION(BlueprintPure, Category=Sprite, meta=(AdvancedDisplay=2))
 	static FPaperTileInfo MakeTile(int32 TileIndex, UPaperTileSet* TileSet, bool bFlipH, bool bFlipV, bool bFlipD);
 };
+
+#if UE_ENABLE_INCLUDE_ORDER_DEPRECATED_IN_5_2
+#include "CoreMinimal.h"
+#include "PaperTileLayer.h"
+#include "UObject/ObjectMacros.h"
+#endif
diff --git a/Engine/Plugins/2D/Paper2D/Source/Paper2D/Private/BitmapUtils.h b/Engine/Plugins/2D/Paper2D/Source/Paper2D/Private/BitmapUtils.h
--- a/Engine/Plugins/2D/Paper2D/Source/Paper2D/Private/BitmapUtils.h
+++ b/Engine/Plugins/2D/Paper2D/Source/Paper2D/Private/BitmapUtils.h
@@ -2,9 +2,9 @@
 
 #pragma once
 
-#include "CoreMinimal.h"
 #include "Engine/Texture2D.h"
 #include "AlphaBitmap.h"
+#include "Math/IntRect.h"
 #include "Misc/ScopedSlowTask.h"
 
 //////////////////////////////////////////////////////////////////////////
diff --git a/Engine/Plugins/2D/Paper2D/Source/Paper2D/Private/GroupedSpriteSceneProxy.h b/Engine/Plugins/2D/Paper2D/Source/Paper2D/Private/GroupedSpriteSceneProxy.h
--- a/Engine/Plugins/2D/Paper2D/Source/Paper2D/Private/GroupedSpriteSceneProxy.h
+++ b/Engine/Plugins/2D/Paper2D/Source/Paper2D/Private/GroupedSpriteSceneProxy.h
@@ -2,10 +2,10 @@
 
 #pragma once
 
-#include "CoreMinimal.h"
-#include "Components/InstancedStaticMeshComponent.h"
 #include "PaperRenderSceneProxy.h"
 
+struct FSpriteDrawCallRecord;
+
 class FMeshElementCollector;
 class UBodySetup;
 class UPaperGroupedSpriteComponent;
diff --git a/Engine/Plugins/2D/Paper2D/Source/Paper2D/Private/Paper2DPrivate.h b/Engine/Plugins/2D/Paper2D/Source/Paper2D/Private/Paper2DPrivate.h
--- a/Engine/Plugins/2D/Paper2D/Source/Paper2D/Private/Paper2DPrivate.h
+++ b/Engine/Plugins/2D/Paper2D/Source/Paper2D/Private/Paper2DPrivate.h
@@ -2,7 +2,7 @@
 
 #pragma once
 
-#include "CoreMinimal.h"
+#include "Logging/LogMacros.h"
 
 DECLARE_LOG_CATEGORY_EXTERN(LogPaper2D, Verbose, All);
 
diff --git a/Engine/Plugins/2D/Paper2D/Source/Paper2D/Private/PaperFlipbook.cpp b/Engine/Plugins/2D/Paper2D/Source/Paper2D/Private/PaperFlipbook.cpp
--- a/Engine/Plugins/2D/Paper2D/Source/Paper2D/Private/PaperFlipbook.cpp
+++ b/Engine/Plugins/2D/Paper2D/Source/Paper2D/Private/PaperFlipbook.cpp
@@ -10,7 +10,6 @@
 #include UE_INLINE_GENERATED_CPP_BY_NAME(PaperFlipbook)
 
 #if WITH_EDITORONLY_DATA
-#include "UObject/UObjectHash.h"
 #include "ComponentReregisterContext.h"
 #endif
 
diff --git a/Engine/Plugins/2D/Paper2D/Source/Paper2D/Private/PaperFlipbookComponent.cpp b/Engine/Plugins/2D/Paper2D/Source/Paper2D/Private/PaperFlipbookComponent.cpp
--- a/Engine/Plugins/2D/Paper2D/Source/Paper2D/Private/PaperFlipbookComponent.cpp
+++ b/Engine/Plugins/2D/Paper2D/Source/Paper2D/Private/PaperFlipbookComponent.cpp
@@ -3,13 +3,11 @@
 #include "PaperFlipbookComponent.h"
 #include "RenderingThread.h"
 #include "Engine/CollisionProfile.h"
-#include "SpriteDrawCall.h"
 #include "PaperFlipbookSceneProxy.h"
 #include "PaperFlipbook.h"
 #include "PaperCustomVersion.h"
 #include "Net/UnrealNetwork.h"
 #include "ContentStreaming.h"
-#include "Logging/TokenizedMessage.h"
 #include "Logging/MessageLog.h"
 #include "Misc/MapErrors.h"
 #include "Misc/UObjectToken.h"
diff --git a/Engine/Plugins/2D/Paper2D/Source/Paper2D/Private/PaperGroupedSpriteComponent.cpp b/Engine/Plugins/2D/Paper2D/Source/Paper2D/Private/PaperGroupedSpriteComponent.cpp
--- a/Engine/Plugins/2D/Paper2D/Source/Paper2D/Private/PaperGroupedSpriteComponent.cpp
+++ b/Engine/Plugins/2D/Paper2D/Source/Paper2D/Private/PaperGroupedSpriteComponent.cpp
@@ -7,7 +7,6 @@
 #include "GroupedSpriteSceneProxy.h"
 #include "AI/NavigationSystemHelpers.h"
 
-#include "Logging/TokenizedMessage.h"
 #include "Logging/MessageLog.h"
 #include "Misc/UObjectToken.h"
 
diff --git a/Engine/Plugins/2D/Paper2D/Source/Paper2D/Private/PaperRenderSceneProxy.cpp b/Engine/Plugins/2D/Paper2D/Source/Paper2D/Private/PaperRenderSceneProxy.cpp
--- a/Engine/Plugins/2D/Paper2D/Source/Paper2D/Private/PaperRenderSceneProxy.cpp
+++ b/Engine/Plugins/2D/Paper2D/Source/Paper2D/Private/PaperRenderSceneProxy.cpp
@@ -1,14 +1,12 @@
 // Copyright Epic Games, Inc. All Rights Reserved.
 
 #include "PaperRenderSceneProxy.h"
-#include "Containers/ResourceArray.h"
+#include "PrimitiveViewRelevance.h"
 #include "SceneManagement.h"
 #include "MaterialDomain.h"
 #include "Materials/Material.h"
 #include "Materials/MaterialRenderProxy.h"
 #include "PhysicsEngine/BodySetup.h"
-#include "EngineGlobals.h"
-#include "Components/PrimitiveComponent.h"
 #include "Components/MeshComponent.h"
 #include "Engine/Engine.h"
 #include "SceneInterface.h"
diff --git a/Engine/Plugins/2D/Paper2D/Source/Paper2D/Private/PaperRenderSceneProxy.h b/Engine/Plugins/2D/Paper2D/Source/Paper2D/Private/PaperRenderSceneProxy.h
--- a/Engine/Plugins/2D/Paper2D/Source/Paper2D/Private/PaperRenderSceneProxy.h
+++ b/Engine/Plugins/2D/Paper2D/Source/Paper2D/Private/PaperRenderSceneProxy.h
@@ -2,20 +2,16 @@
 
 #pragma once
 
-#include "CoreMinimal.h"
-#include "Engine/EngineTypes.h"
-#include "RenderResource.h"
+#include "Materials/MaterialRelevance.h"
 #include "SpriteDrawCall.h"
 #include "Materials/MaterialInterface.h"
-#include "PackedNormal.h"
-#include "PrimitiveViewRelevance.h"
 #include "PrimitiveSceneProxy.h"
-#include "LocalVertexFactory.h"
 #include "Paper2DModule.h"
-#include "DynamicMeshBuilder.h"
 #include "PaperSpriteVertexBuffer.h"
 #include "TextureResource.h"
 
+class FTexture;
+
 class FMeshElementCollector;
 class UBodySetup;
 class UPrimitiveComponent;
diff --git a/Engine/Plugins/2D/Paper2D/Source/Paper2D/Private/PaperSprite.cpp b/Engine/Plugins/2D/Paper2D/Source/Paper2D/Private/PaperSprite.cpp
--- a/Engine/Plugins/2D/Paper2D/Source/Paper2D/Private/PaperSprite.cpp
+++ b/Engine/Plugins/2D/Paper2D/Source/Paper2D/Private/PaperSprite.cpp
@@ -1,12 +1,10 @@
 // Copyright Epic Games, Inc. All Rights Reserved.
 
 #include "PaperSprite.h"
+#include "BodySetupEnums.h"
 #include "UObject/ConstructorHelpers.h"
 #include "Materials/MaterialInterface.h"
 #include "Engine/CollisionProfile.h"
-#include "PhysicsEngine/ConvexElem.h"
-#include "PhysicsEngine/BoxElem.h"
-#include "PhysicsEngine/SphereElem.h"
 #include "PhysicsEngine/BodySetup.h"
 
 #include "PaperCustomVersion.h"
@@ -23,8 +21,6 @@
 #include UE_INLINE_GENERATED_CPP_BY_NAME(PaperSprite)
 
 #if WITH_EDITOR
-#include "UObject/UObjectHash.h"
-#include "UObject/UObjectIterator.h"
 #endif
 
 #if WITH_EDITOR
@@ -66,7 +62,6 @@ static void UpdateGeometryToBeBoxPositionRelative(FSpriteGeometryCollection& Geo
 #if WITH_EDITOR
 
 #include "PaperSpriteAtlas.h"
-#include "AlphaBitmap.h"
 #include "BitmapUtils.h"
 #include "ComponentReregisterContext.h"
 #include "PaperRuntimeSettings.h"
diff --git a/Engine/Plugins/2D/Paper2D/Source/Paper2D/Private/PaperSpriteAtlas.cpp b/Engine/Plugins/2D/Paper2D/Source/Paper2D/Private/PaperSpriteAtlas.cpp
--- a/Engine/Plugins/2D/Paper2D/Source/Paper2D/Private/PaperSpriteAtlas.cpp
+++ b/Engine/Plugins/2D/Paper2D/Source/Paper2D/Private/PaperSpriteAtlas.cpp
@@ -1,6 +1,7 @@
 // Copyright Epic Games, Inc. All Rights Reserved.
 
 #include "PaperSpriteAtlas.h"
+#include "Engine/Texture.h"
 
 #include UE_INLINE_GENERATED_CPP_BY_NAME(PaperSpriteAtlas)
 
diff --git a/Engine/Plugins/2D/Paper2D/Source/Paper2D/Private/PaperSpriteComponent.cpp b/Engine/Plugins/2D/Paper2D/Source/Paper2D/Private/PaperSpriteComponent.cpp
--- a/Engine/Plugins/2D/Paper2D/Source/Paper2D/Private/PaperSpriteComponent.cpp
+++ b/Engine/Plugins/2D/Paper2D/Source/Paper2D/Private/PaperSpriteComponent.cpp
@@ -3,12 +3,10 @@
 #include "PaperSpriteComponent.h"
 #include "RenderingThread.h"
 #include "Engine/CollisionProfile.h"
-#include "SpriteDrawCall.h"
 #include "PaperSpriteSceneProxy.h"
 #include "PaperCustomVersion.h"
 #include "PhysicsEngine/BodySetup.h"
 #include "ContentStreaming.h"
-#include "Logging/TokenizedMessage.h"
 #include "Logging/MessageLog.h"
 #include "Misc/MapErrors.h"
 #include "Misc/UObjectToken.h"
diff --git a/Engine/Plugins/2D/Paper2D/Source/Paper2D/Private/PaperSpriteSceneProxy.cpp b/Engine/Plugins/2D/Paper2D/Source/Paper2D/Private/PaperSpriteSceneProxy.cpp
--- a/Engine/Plugins/2D/Paper2D/Source/Paper2D/Private/PaperSpriteSceneProxy.cpp
+++ b/Engine/Plugins/2D/Paper2D/Source/Paper2D/Private/PaperSpriteSceneProxy.cpp
@@ -5,7 +5,6 @@
 #include "Materials/MaterialRenderProxy.h"
 #include "SceneManagement.h"
 #include "PhysicsEngine/BodySetup.h"
-#include "EngineGlobals.h"
 #include "Engine/Engine.h"
 #include "PaperSpriteComponent.h"
 
diff --git a/Engine/Plugins/2D/Paper2D/Source/Paper2D/Private/PaperSpriteSceneProxy.h b/Engine/Plugins/2D/Paper2D/Source/Paper2D/Private/PaperSpriteSceneProxy.h
--- a/Engine/Plugins/2D/Paper2D/Source/Paper2D/Private/PaperSpriteSceneProxy.h
+++ b/Engine/Plugins/2D/Paper2D/Source/Paper2D/Private/PaperSpriteSceneProxy.h
@@ -2,8 +2,6 @@
 
 #pragma once
 
-#include "CoreMinimal.h"
-#include "SpriteDrawCall.h"
 #include "PaperRenderSceneProxy.h"
 
 class FMeshElementCollector;
diff --git a/Engine/Plugins/2D/Paper2D/Source/Paper2D/Private/PaperSpriteVertexBuffer.cpp b/Engine/Plugins/2D/Paper2D/Source/Paper2D/Private/PaperSpriteVertexBuffer.cpp
--- a/Engine/Plugins/2D/Paper2D/Source/Paper2D/Private/PaperSpriteVertexBuffer.cpp
+++ b/Engine/Plugins/2D/Paper2D/Source/Paper2D/Private/PaperSpriteVertexBuffer.cpp
@@ -1,13 +1,8 @@
 // Copyright Epic Games, Inc. All Rights Reserved.
 
 #include "PaperSpriteVertexBuffer.h"
-#include "Materials/Material.h"
-#include "SceneManagement.h"
-#include "PhysicsEngine/BodySetup.h"
-#include "EngineGlobals.h"
-#include "Engine/Engine.h"
-#include "PaperSpriteComponent.h"
 #include "DataDrivenShaderPlatformInfo.h"
+#include "RenderingThread.h"
 
 //////////////////////////////////////////////////////////////////////////
 // FPaperSpriteVertexBuffer
diff --git a/Engine/Plugins/2D/Paper2D/Source/Paper2D/Private/PaperSpriteVertexBuffer.h b/Engine/Plugins/2D/Paper2D/Source/Paper2D/Private/PaperSpriteVertexBuffer.h
--- a/Engine/Plugins/2D/Paper2D/Source/Paper2D/Private/PaperSpriteVertexBuffer.h
+++ b/Engine/Plugins/2D/Paper2D/Source/Paper2D/Private/PaperSpriteVertexBuffer.h
@@ -2,12 +2,7 @@
 
 #pragma once
 
-#include "CoreMinimal.h"
-#include "Engine/EngineTypes.h"
 #include "LocalVertexFactory.h"
-#include "RenderResource.h"
-#include "PackedNormal.h"
-#include "SceneManagement.h"
 #include "DynamicMeshBuilder.h"
 
 //////////////////////////////////////////////////////////////////////////
diff --git a/Engine/Plugins/2D/Paper2D/Source/Paper2D/Private/PaperTileLayer.cpp b/Engine/Plugins/2D/Paper2D/Source/Paper2D/Private/PaperTileLayer.cpp
--- a/Engine/Plugins/2D/Paper2D/Source/Paper2D/Private/PaperTileLayer.cpp
+++ b/Engine/Plugins/2D/Paper2D/Source/Paper2D/Private/PaperTileLayer.cpp
@@ -2,7 +2,6 @@
 
 #include "PaperTileLayer.h"
 #include "Paper2DModule.h"
-#include "SpriteEditorOnlyTypes.h"
 #include "PaperTileMap.h"
 #include "PaperTileSet.h"
 
diff --git a/Engine/Plugins/2D/Paper2D/Source/Paper2D/Private/PaperTileMap.cpp b/Engine/Plugins/2D/Paper2D/Source/Paper2D/Private/PaperTileMap.cpp
--- a/Engine/Plugins/2D/Paper2D/Source/Paper2D/Private/PaperTileMap.cpp
+++ b/Engine/Plugins/2D/Paper2D/Source/Paper2D/Private/PaperTileMap.cpp
@@ -1,6 +1,7 @@
 // Copyright Epic Games, Inc. All Rights Reserved.
 
 #include "PaperTileMap.h"
+#include "BodySetupEnums.h"
 #include "UObject/ConstructorHelpers.h"
 #include "Materials/MaterialInterface.h"
 #include "EditorFramework/AssetImportData.h"
@@ -15,8 +16,6 @@
 
 #if WITH_EDITOR
 #include "PaperTileMapComponent.h"
-#include "UObject/UObjectHash.h"
-#include "UObject/UObjectIterator.h"
 #include "ComponentReregisterContext.h"
 #endif
 
diff --git a/Engine/Plugins/2D/Paper2D/Source/Paper2D/Private/PaperTileMapActor.cpp b/Engine/Plugins/2D/Paper2D/Source/Paper2D/Private/PaperTileMapActor.cpp
--- a/Engine/Plugins/2D/Paper2D/Source/Paper2D/Private/PaperTileMapActor.cpp
+++ b/Engine/Plugins/2D/Paper2D/Source/Paper2D/Private/PaperTileMapActor.cpp
@@ -1,6 +1,7 @@
 // Copyright Epic Games, Inc. All Rights Reserved.
 
 #include "PaperTileMapActor.h"
+#include "Async/TaskGraphInterfaces.h"
 #include "PaperTileMapComponent.h"
 
 #include UE_INLINE_GENERATED_CPP_BY_NAME(PaperTileMapActor)
diff --git a/Engine/Plugins/2D/Paper2D/Source/Paper2D/Private/PaperTileMapComponent.cpp b/Engine/Plugins/2D/Paper2D/Source/Paper2D/Private/PaperTileMapComponent.cpp
--- a/Engine/Plugins/2D/Paper2D/Source/Paper2D/Private/PaperTileMapComponent.cpp
+++ b/Engine/Plugins/2D/Paper2D/Source/Paper2D/Private/PaperTileMapComponent.cpp
@@ -1,11 +1,11 @@
 // Copyright Epic Games, Inc. All Rights Reserved.
 
 #include "PaperTileMapComponent.h"
+#include "PaperTileLayer.h"
 #include "UObject/ConstructorHelpers.h"
 #include "MaterialDomain.h"
 #include "Materials/Material.h"
 #include "Paper2DPrivate.h"
-#include "PaperRenderSceneProxy.h"
 #include "PaperTileMapRenderSceneProxy.h"
 #include "PaperCustomVersion.h"
 #include "PhysicsEngine/BodySetup.h"
diff --git a/Engine/Plugins/2D/Paper2D/Source/Paper2D/Private/PaperTileMapRenderSceneProxy.cpp b/Engine/Plugins/2D/Paper2D/Source/Paper2D/Private/PaperTileMapRenderSceneProxy.cpp
--- a/Engine/Plugins/2D/Paper2D/Source/Paper2D/Private/PaperTileMapRenderSceneProxy.cpp
+++ b/Engine/Plugins/2D/Paper2D/Source/Paper2D/Private/PaperTileMapRenderSceneProxy.cpp
@@ -4,7 +4,6 @@
 #include "Materials/Material.h"
 #include "SceneManagement.h"
 #include "PhysicsEngine/BodySetup.h"
-#include "EngineGlobals.h"
 #include "Engine/Engine.h"
 #include "Materials/MaterialRenderProxy.h"
 #include "PaperTileMap.h"
diff --git a/Engine/Plugins/2D/Paper2D/Source/Paper2D/Private/PaperTileSet.cpp b/Engine/Plugins/2D/Paper2D/Source/Paper2D/Private/PaperTileSet.cpp
--- a/Engine/Plugins/2D/Paper2D/Source/Paper2D/Private/PaperTileSet.cpp
+++ b/Engine/Plugins/2D/Paper2D/Source/Paper2D/Private/PaperTileSet.cpp
@@ -183,7 +183,6 @@ void UPaperTileSet::PostLoad()
 #if WITH_EDITOR
 
 #include "PaperTileMapComponent.h"
-#include "UObject/UObjectHash.h"
 #include "ComponentReregisterContext.h"
 
 void UPaperTileSet::PostEditChangeProperty(FPropertyChangedEvent& PropertyChangedEvent)
diff --git a/Engine/Plugins/2D/Paper2D/Source/Paper2D/Private/Terrain/PaperTerrainComponent.cpp b/Engine/Plugins/2D/Paper2D/Source/Paper2D/Private/Terrain/PaperTerrainComponent.cpp
--- a/Engine/Plugins/2D/Paper2D/Source/Paper2D/Private/Terrain/PaperTerrainComponent.cpp
+++ b/Engine/Plugins/2D/Paper2D/Source/Paper2D/Private/Terrain/PaperTerrainComponent.cpp
@@ -4,14 +4,10 @@
 #include "UObject/ConstructorHelpers.h"
 #include "MaterialDomain.h"
 #include "Materials/Material.h"
-#include "Engine/Polys.h"
-#include "Components/SplineComponent.h"
 #include "PaperCustomVersion.h"
 
 #include "PaperRenderSceneProxy.h"
 #include "GeomTools.h"
-#include "PhysicsEngine/ConvexElem.h"
-#include "PhysicsEngine/BoxElem.h"
 #include "PhysicsEngine/BodySetup.h"
 #include "PaperTerrainSplineComponent.h"
 #include "PaperTerrainMaterial.h"
diff --git a/Engine/Plugins/2D/Paper2D/Source/Paper2D/Public/AlphaBitmap.h b/Engine/Plugins/2D/Paper2D/Source/Paper2D/Public/AlphaBitmap.h
--- a/Engine/Plugins/2D/Paper2D/Source/Paper2D/Public/AlphaBitmap.h
+++ b/Engine/Plugins/2D/Paper2D/Source/Paper2D/Public/AlphaBitmap.h
@@ -2,7 +2,6 @@
 
 #pragma once
 
-#include "CoreMinimal.h"
 #include "Engine/Texture.h"
 
 //////////////////////////////////////////////////////////////////////////
@@ -360,3 +359,7 @@ public:
 	int32 Height;
 	uint8 DefaultValue;
 };
+
+#if UE_ENABLE_INCLUDE_ORDER_DEPRECATED_IN_5_2
+#include "CoreMinimal.h"
+#endif
diff --git a/Engine/Plugins/2D/Paper2D/Source/Paper2D/Public/Paper2DModule.h b/Engine/Plugins/2D/Paper2D/Source/Paper2D/Public/Paper2DModule.h
--- a/Engine/Plugins/2D/Paper2D/Source/Paper2D/Public/Paper2DModule.h
+++ b/Engine/Plugins/2D/Paper2D/Source/Paper2D/Public/Paper2DModule.h
@@ -2,9 +2,9 @@
 
 #pragma once
 
-#include "CoreMinimal.h"
-#include "Stats/Stats.h"
+#include "Math/Vector.h"
 #include "Modules/ModuleInterface.h"
+#include "Stats/Stats2.h"
 
 #define PAPER2D_MODULE_NAME "Paper2D"
 
@@ -35,3 +35,8 @@ extern PAPER2D_API FVector PaperAxisZ;
 
 DECLARE_STATS_GROUP(TEXT("Paper2D"), STATGROUP_Paper2D, STATCAT_Advanced);
 DECLARE_CYCLE_STAT_EXTERN(TEXT("SetSprite (RT)"), STAT_PaperRender_SetSpriteRT, STATGROUP_Paper2D, PAPER2D_API);
+
+#if UE_ENABLE_INCLUDE_ORDER_DEPRECATED_IN_5_2
+#include "CoreMinimal.h"
+#include "Stats/Stats.h"
+#endif
diff --git a/Engine/Plugins/2D/Paper2D/Source/Paper2DEditor/Classes/PaperImporterSettings.h b/Engine/Plugins/2D/Paper2D/Source/Paper2DEditor/Classes/PaperImporterSettings.h
--- a/Engine/Plugins/2D/Paper2D/Source/Paper2DEditor/Classes/PaperImporterSettings.h
+++ b/Engine/Plugins/2D/Paper2D/Source/Paper2DEditor/Classes/PaperImporterSettings.h
@@ -2,13 +2,12 @@
 
 #pragma once
 
-#include "CoreMinimal.h"
-#include "UObject/ObjectMacros.h"
-#include "UObject/Object.h"
-#include "UObject/SoftObjectPath.h"
-#include "Engine/Texture.h"
 #include "PaperImporterSettings.generated.h"
 
+class UTexture;
+enum TextureCompressionSettings : int;
+enum TextureGroup : int;
+
 class UMaterialInterface;
 class UPaperTileMap;
 class UPaperTileSet;
@@ -158,3 +157,11 @@ public:
 	// Returns the default pixels/uu setting
 	float GetDefaultPixelsPerUnrealUnit() const { return DefaultPixelsPerUnrealUnit; }
 };
+
+#if UE_ENABLE_INCLUDE_ORDER_DEPRECATED_IN_5_2
+#include "CoreMinimal.h"
+#include "Engine/Texture.h"
+#include "UObject/Object.h"
+#include "UObject/ObjectMacros.h"
+#include "UObject/SoftObjectPath.h"
+#endif
diff --git a/Engine/Plugins/2D/Paper2D/Source/Paper2DEditor/Classes/TileMapAssetImportData.h b/Engine/Plugins/2D/Paper2D/Source/Paper2DEditor/Classes/TileMapAssetImportData.h
--- a/Engine/Plugins/2D/Paper2D/Source/Paper2DEditor/Classes/TileMapAssetImportData.h
+++ b/Engine/Plugins/2D/Paper2D/Source/Paper2DEditor/Classes/TileMapAssetImportData.h
@@ -2,9 +2,6 @@
 
 #pragma once
 
-#include "CoreMinimal.h"
-#include "UObject/ObjectMacros.h"
-#include "Engine/Texture.h"
 #include "EditorFramework/AssetImportData.h"
 #include "TileMapAssetImportData.generated.h"
 
@@ -38,3 +35,9 @@ class PAPER2DEDITOR_API UTileMapAssetImportData : public UAssetImportData
 
 	static UTileMapAssetImportData* GetImportDataForTileMap(class UPaperTileMap* TileMap);
 };
+
+#if UE_ENABLE_INCLUDE_ORDER_DEPRECATED_IN_5_2
+#include "CoreMinimal.h"
+#include "Engine/Texture.h"
+#include "UObject/ObjectMacros.h"
+#endif
diff --git a/Engine/Plugins/2D/Paper2D/Source/Paper2DEditor/Private/Atlasing/AtlasAssetTypeActions.cpp b/Engine/Plugins/2D/Paper2D/Source/Paper2DEditor/Private/Atlasing/AtlasAssetTypeActions.cpp
--- a/Engine/Plugins/2D/Paper2D/Source/Paper2DEditor/Private/Atlasing/AtlasAssetTypeActions.cpp
+++ b/Engine/Plugins/2D/Paper2D/Source/Paper2DEditor/Private/Atlasing/AtlasAssetTypeActions.cpp
@@ -1,7 +1,6 @@
 // Copyright Epic Games, Inc. All Rights Reserved.
 
 #include "Atlasing/AtlasAssetTypeActions.h"
-#include "AssetRegistry/AssetData.h"
 
 #include "PaperSpriteAtlas.h"
 
diff --git a/Engine/Plugins/2D/Paper2D/Source/Paper2DEditor/Private/Atlasing/AtlasAssetTypeActions.h b/Engine/Plugins/2D/Paper2D/Source/Paper2DEditor/Private/Atlasing/AtlasAssetTypeActions.h
--- a/Engine/Plugins/2D/Paper2D/Source/Paper2DEditor/Private/Atlasing/AtlasAssetTypeActions.h
+++ b/Engine/Plugins/2D/Paper2D/Source/Paper2DEditor/Private/Atlasing/AtlasAssetTypeActions.h
@@ -2,8 +2,6 @@
 
 #pragma once
 
-#include "CoreMinimal.h"
-#include "Toolkits/IToolkitHost.h"
 #include "AssetTypeActions_Base.h"
 
 struct FAssetData;
diff --git a/Engine/Plugins/2D/Paper2D/Source/Paper2DEditor/Private/Atlasing/PaperAtlasGenerator.cpp b/Engine/Plugins/2D/Paper2D/Source/Paper2DEditor/Private/Atlasing/PaperAtlasGenerator.cpp
--- a/Engine/Plugins/2D/Paper2D/Source/Paper2DEditor/Private/Atlasing/PaperAtlasGenerator.cpp
+++ b/Engine/Plugins/2D/Paper2D/Source/Paper2DEditor/Private/Atlasing/PaperAtlasGenerator.cpp
@@ -3,13 +3,11 @@
 #include "Atlasing/PaperAtlasGenerator.h"
 #include "Paper2DEditorLog.h"
 #include "Misc/MessageDialog.h"
-#include "Engine/Texture2D.h"
 
 
 //////////////////////////////////////////////////////////////////////////
 // 
 
-#include "PaperSpriteAtlas.h"
 #include "Atlasing/PaperAtlasTextureHelpers.h"
 #include "Atlasing/PaperAtlasHelpers.h"
 
diff --git a/Engine/Plugins/2D/Paper2D/Source/Paper2DEditor/Private/Atlasing/PaperAtlasHelpers.h b/Engine/Plugins/2D/Paper2D/Source/Paper2DEditor/Private/Atlasing/PaperAtlasHelpers.h
--- a/Engine/Plugins/2D/Paper2D/Source/Paper2DEditor/Private/Atlasing/PaperAtlasHelpers.h
+++ b/Engine/Plugins/2D/Paper2D/Source/Paper2DEditor/Private/Atlasing/PaperAtlasHelpers.h
@@ -2,15 +2,14 @@
 
 #pragma once
 
-#include "CoreMinimal.h"
-#include "Misc/Guid.h"
-#include "UObject/SoftObjectPath.h"
 #include "PaperSprite.h"
 #include "PaperSpriteAtlas.h"
 #include "Modules/ModuleManager.h"
 #include "AssetRegistry/AssetData.h"
 #include "AssetRegistry/AssetRegistryModule.h"
 
+class IModuleInterface;
+
 // Loads and returns a list of sprites referenced in the atlas group (i.e. sprites used in the previous build)
 static void LoadAllReferencedSprites(UPaperSpriteAtlas* AtlasGroup, TArray<UPaperSprite*>& AllSprites)
 {
diff --git a/Engine/Plugins/2D/Paper2D/Source/Paper2DEditor/Private/Atlasing/PaperAtlasTextureHelpers.cpp b/Engine/Plugins/2D/Paper2D/Source/Paper2DEditor/Private/Atlasing/PaperAtlasTextureHelpers.cpp
--- a/Engine/Plugins/2D/Paper2D/Source/Paper2DEditor/Private/Atlasing/PaperAtlasTextureHelpers.cpp
+++ b/Engine/Plugins/2D/Paper2D/Source/Paper2DEditor/Private/Atlasing/PaperAtlasTextureHelpers.cpp
@@ -2,6 +2,8 @@
 
 #include "Atlasing/PaperAtlasTextureHelpers.h"
 #include "Paper2DEditorLog.h"
+#include "PaperSprite.h"
+#include "PaperSpriteAtlas.h"
 
 //////////////////////////////////////////////////////////////////////////
 // 
diff --git a/Engine/Plugins/2D/Paper2D/Source/Paper2DEditor/Private/Atlasing/PaperAtlasTextureHelpers.h b/Engine/Plugins/2D/Paper2D/Source/Paper2DEditor/Private/Atlasing/PaperAtlasTextureHelpers.h
--- a/Engine/Plugins/2D/Paper2D/Source/Paper2DEditor/Private/Atlasing/PaperAtlasTextureHelpers.h
+++ b/Engine/Plugins/2D/Paper2D/Source/Paper2DEditor/Private/Atlasing/PaperAtlasTextureHelpers.h
@@ -2,8 +2,12 @@
 
 #pragma once
 
-#include "CoreMinimal.h"
-#include "PaperSpriteAtlas.h"
+#include "Containers/Array.h"
+
+class UPaperSprite;
+class UTexture;
+enum class EPaperSpriteAtlasPadding : uint8;
+struct FPaperSpriteAtlasSlot;
 
 struct FPaperAtlasTextureHelpers
 {
diff --git a/Engine/Plugins/2D/Paper2D/Source/Paper2DEditor/Private/Atlasing/PaperSpriteAtlasFactory.h b/Engine/Plugins/2D/Paper2D/Source/Paper2DEditor/Private/Atlasing/PaperSpriteAtlasFactory.h
--- a/Engine/Plugins/2D/Paper2D/Source/Paper2DEditor/Private/Atlasing/PaperSpriteAtlasFactory.h
+++ b/Engine/Plugins/2D/Paper2D/Source/Paper2DEditor/Private/Atlasing/PaperSpriteAtlasFactory.h
@@ -2,8 +2,6 @@
 
 #pragma once
 
-#include "CoreMinimal.h"
-#include "UObject/ObjectMacros.h"
 #include "Factories/Factory.h"
 #include "PaperSpriteAtlasFactory.generated.h"
 
diff --git a/Engine/Plugins/2D/Paper2D/Source/Paper2DEditor/Private/ContentBrowserExtensions/ContentBrowserExtensions.cpp b/Engine/Plugins/2D/Paper2D/Source/Paper2DEditor/Private/ContentBrowserExtensions/ContentBrowserExtensions.cpp
--- a/Engine/Plugins/2D/Paper2D/Source/Paper2DEditor/Private/ContentBrowserExtensions/ContentBrowserExtensions.cpp
+++ b/Engine/Plugins/2D/Paper2D/Source/Paper2DEditor/Private/ContentBrowserExtensions/ContentBrowserExtensions.cpp
@@ -1,15 +1,8 @@
 // Copyright Epic Games, Inc. All Rights Reserved.
 
 #include "ContentBrowserExtensions/ContentBrowserExtensions.h"
-#include "Modules/ModuleManager.h"
 #include "Misc/PackageName.h"
-#include "Textures/SlateIcon.h"
-#include "Framework/Commands/UIAction.h"
-#include "Framework/MultiBox/MultiBoxExtender.h"
 #include "Framework/MultiBox/MultiBoxBuilder.h"
-#include "Styling/AppStyle.h"
-#include "Engine/Texture2D.h"
-#include "AssetRegistry/AssetData.h"
 #include "IContentBrowserSingleton.h"
 #include "ContentBrowserModule.h"
 #include "PaperSprite.h"
diff --git a/Engine/Plugins/2D/Paper2D/Source/Paper2DEditor/Private/ExtractSprites/PaperExtractSpritesSettings.h b/Engine/Plugins/2D/Paper2D/Source/Paper2DEditor/Private/ExtractSprites/PaperExtractSpritesSettings.h
--- a/Engine/Plugins/2D/Paper2D/Source/Paper2DEditor/Private/ExtractSprites/PaperExtractSpritesSettings.h
+++ b/Engine/Plugins/2D/Paper2D/Source/Paper2DEditor/Private/ExtractSprites/PaperExtractSpritesSettings.h
@@ -2,10 +2,6 @@
 
 #pragma once
 
-#include "CoreMinimal.h"
-#include "UObject/ObjectMacros.h"
-#include "UObject/UObjectGlobals.h"
-#include "UObject/Object.h"
 #include "PaperExtractSpritesSettings.generated.h"
 
 UENUM()
diff --git a/Engine/Plugins/2D/Paper2D/Source/Paper2DEditor/Private/ExtractSprites/SPaperExtractSpritesDialog.cpp b/Engine/Plugins/2D/Paper2D/Source/Paper2DEditor/Private/ExtractSprites/SPaperExtractSpritesDialog.cpp
--- a/Engine/Plugins/2D/Paper2D/Source/Paper2DEditor/Private/ExtractSprites/SPaperExtractSpritesDialog.cpp
+++ b/Engine/Plugins/2D/Paper2D/Source/Paper2DEditor/Private/ExtractSprites/SPaperExtractSpritesDialog.cpp
@@ -1,20 +1,16 @@
 // Copyright Epic Games, Inc. All Rights Reserved.
 
 #include "ExtractSprites/SPaperExtractSpritesDialog.h"
-#include "Widgets/SBoxPanel.h"
-#include "Widgets/SWindow.h"
-#include "Widgets/SViewport.h"
+#include "DetailsViewArgs.h"
 #include "Misc/FeedbackContext.h"
+#include "IDetailsView.h"
 #include "Misc/ScopedSlowTask.h"
-#include "Modules/ModuleManager.h"
 #include "Misc/PackageName.h"
-#include "Layout/WidgetPath.h"
 #include "Framework/Application/SlateApplication.h"
-#include "Widgets/Layout/SBorder.h"
+#include "PaperEditorViewportClient.h"
 #include "Widgets/Layout/SUniformGridPanel.h"
 #include "Widgets/Input/SButton.h"
 #include "Framework/Docking/TabManager.h"
-#include "Styling/AppStyle.h"
 #include "CanvasItem.h"
 #include "ExtractSprites/PaperExtractSpritesSettings.h"
 
diff --git a/Engine/Plugins/2D/Paper2D/Source/Paper2DEditor/Private/ExtractSprites/SPaperExtractSpritesDialog.h b/Engine/Plugins/2D/Paper2D/Source/Paper2DEditor/Private/ExtractSprites/SPaperExtractSpritesDialog.h
--- a/Engine/Plugins/2D/Paper2D/Source/Paper2DEditor/Private/ExtractSprites/SPaperExtractSpritesDialog.h
+++ b/Engine/Plugins/2D/Paper2D/Source/Paper2DEditor/Private/ExtractSprites/SPaperExtractSpritesDialog.h
@@ -2,14 +2,11 @@
 
 #pragma once
 
-#include "CoreMinimal.h"
-#include "Input/Reply.h"
-#include "Widgets/DeclarativeSyntaxSupport.h"
-#include "Widgets/SCompoundWidget.h"
-#include "Engine/Texture2D.h"
-#include "IDetailsView.h"
 #include "SPaperEditorViewport.h"
 
+class UTexture2D;
+struct FPropertyChangedEvent;
+
 class FPaperExtractSpritesViewportClient;
 struct Rect;
 
diff --git a/Engine/Plugins/2D/Paper2D/Source/Paper2DEditor/Private/FlipbookAssetTypeActions.cpp b/Engine/Plugins/2D/Paper2D/Source/Paper2DEditor/Private/FlipbookAssetTypeActions.cpp
--- a/Engine/Plugins/2D/Paper2D/Source/Paper2DEditor/Private/FlipbookAssetTypeActions.cpp
+++ b/Engine/Plugins/2D/Paper2D/Source/Paper2DEditor/Private/FlipbookAssetTypeActions.cpp
@@ -2,7 +2,6 @@
 
 #include "FlipbookAssetTypeActions.h"
 
-#include "PaperFlipbook.h"
 #include "FlipbookEditor/FlipbookEditor.h"
 
 #define LOCTEXT_NAMESPACE "AssetTypeActions"
diff --git a/Engine/Plugins/2D/Paper2D/Source/Paper2DEditor/Private/FlipbookAssetTypeActions.h b/Engine/Plugins/2D/Paper2D/Source/Paper2DEditor/Private/FlipbookAssetTypeActions.h
--- a/Engine/Plugins/2D/Paper2D/Source/Paper2DEditor/Private/FlipbookAssetTypeActions.h
+++ b/Engine/Plugins/2D/Paper2D/Source/Paper2DEditor/Private/FlipbookAssetTypeActions.h
@@ -2,8 +2,6 @@
 
 #pragma once
 
-#include "CoreMinimal.h"
-#include "Toolkits/IToolkitHost.h"
 #include "AssetTypeActions_Base.h"
 
 class FFlipbookAssetTypeActions : public FAssetTypeActions_Base
diff --git a/Engine/Plugins/2D/Paper2D/Source/Paper2DEditor/Private/FlipbookEditor/FlipbookEditor.cpp b/Engine/Plugins/2D/Paper2D/Source/Paper2DEditor/Private/FlipbookEditor/FlipbookEditor.cpp
--- a/Engine/Plugins/2D/Paper2D/Source/Paper2DEditor/Private/FlipbookEditor/FlipbookEditor.cpp
+++ b/Engine/Plugins/2D/Paper2D/Source/Paper2DEditor/Private/FlipbookEditor/FlipbookEditor.cpp
@@ -1,10 +1,8 @@
 // Copyright Epic Games, Inc. All Rights Reserved.
 
 #include "FlipbookEditor/FlipbookEditor.h"
-#include "EditorViewportClient.h"
-#include "UObject/Package.h"
+#include "Framework/Commands/UICommandList.h"
 #include "Modules/ModuleManager.h"
-#include "Styling/AppStyle.h"
 #include "SSingleObjectDetailsPanel.h"
 
 
@@ -14,13 +12,11 @@
 #include "FlipbookEditor/FlipbookEditorViewportClient.h"
 #include "FlipbookEditor/FlipbookEditorCommands.h"
 #include "Paper2DEditorModule.h"
-#include "SCommonEditorViewportToolbarBase.h"
 #include "FlipbookEditor/SFlipbookEditorViewportToolbar.h"
 #include "SScrubControlPanel.h"
 #include "FlipbookEditor/SFlipbookTimeline.h"
 #include "Widgets/Docking/SDockTab.h"
 #include "Framework/Commands/GenericCommands.h"
-#include "Subsystems/AssetEditorSubsystem.h"
 
 #define LOCTEXT_NAMESPACE "FlipbookEditor"
 
diff --git a/Engine/Plugins/2D/Paper2D/Source/Paper2DEditor/Private/FlipbookEditor/FlipbookEditor.h b/Engine/Plugins/2D/Paper2D/Source/Paper2DEditor/Private/FlipbookEditor/FlipbookEditor.h
--- a/Engine/Plugins/2D/Paper2D/Source/Paper2DEditor/Private/FlipbookEditor/FlipbookEditor.h
+++ b/Engine/Plugins/2D/Paper2D/Source/Paper2DEditor/Private/FlipbookEditor/FlipbookEditor.h
@@ -2,14 +2,12 @@
 
 #pragma once
 
-#include "CoreMinimal.h"
-#include "Input/Reply.h"
-#include "UObject/GCObject.h"
-#include "Framework/Docking/TabManager.h"
 #include "Toolkits/AssetEditorToolkit.h"
 #include "PaperFlipbook.h"
 #include "ITransportControl.h"
 
+class FSpawnTabArgs;
+
 class SFlipbookEditorViewport;
 class UPaperFlipbookComponent;
 
diff --git a/Engine/Plugins/2D/Paper2D/Source/Paper2DEditor/Private/FlipbookEditor/FlipbookEditorCommands.h b/Engine/Plugins/2D/Paper2D/Source/Paper2DEditor/Private/FlipbookEditor/FlipbookEditorCommands.h
--- a/Engine/Plugins/2D/Paper2D/Source/Paper2DEditor/Private/FlipbookEditor/FlipbookEditorCommands.h
+++ b/Engine/Plugins/2D/Paper2D/Source/Paper2DEditor/Private/FlipbookEditor/FlipbookEditorCommands.h
@@ -2,9 +2,9 @@
 
 #pragma once
 
-#include "CoreMinimal.h"
 #include "Framework/Commands/Commands.h"
 #include "PaperStyle.h"
+#include "Styling/ISlateStyle.h"
 
 class FFlipbookEditorCommands : public TCommands<FFlipbookEditorCommands>
 {
diff --git a/Engine/Plugins/2D/Paper2D/Source/Paper2DEditor/Private/FlipbookEditor/FlipbookEditorSettings.h b/Engine/Plugins/2D/Paper2D/Source/Paper2DEditor/Private/FlipbookEditor/FlipbookEditorSettings.h
--- a/Engine/Plugins/2D/Paper2D/Source/Paper2DEditor/Private/FlipbookEditor/FlipbookEditorSettings.h
+++ b/Engine/Plugins/2D/Paper2D/Source/Paper2DEditor/Private/FlipbookEditor/FlipbookEditorSettings.h
@@ -2,9 +2,6 @@
 
 #pragma once
 
-#include "CoreMinimal.h"
-#include "UObject/ObjectMacros.h"
-#include "UObject/Object.h"
 #include "FlipbookEditorSettings.generated.h"
 
 // Settings for the Paper2D sprite editor
diff --git a/Engine/Plugins/2D/Paper2D/Source/Paper2DEditor/Private/FlipbookEditor/FlipbookEditorViewportClient.cpp b/Engine/Plugins/2D/Paper2D/Source/Paper2DEditor/Private/FlipbookEditor/FlipbookEditorViewportClient.cpp
--- a/Engine/Plugins/2D/Paper2D/Source/Paper2DEditor/Private/FlipbookEditor/FlipbookEditorViewportClient.cpp
+++ b/Engine/Plugins/2D/Paper2D/Source/Paper2DEditor/Private/FlipbookEditor/FlipbookEditorViewportClient.cpp
@@ -1,10 +1,9 @@
 // Copyright Epic Games, Inc. All Rights Reserved.
 
 #include "FlipbookEditor/FlipbookEditorViewportClient.h"
+#include "PaperFlipbook.h"
 #include "PaperFlipbookComponent.h"
 #include "CanvasItem.h"
-#include "Engine/Engine.h"
-#include "EngineGlobals.h"
 #include "Engine/CollisionProfile.h"
 #include "Utils.h"
 
diff --git a/Engine/Plugins/2D/Paper2D/Source/Paper2DEditor/Private/FlipbookEditor/FlipbookEditorViewportClient.h b/Engine/Plugins/2D/Paper2D/Source/Paper2DEditor/Private/FlipbookEditor/FlipbookEditorViewportClient.h
--- a/Engine/Plugins/2D/Paper2D/Source/Paper2DEditor/Private/FlipbookEditor/FlipbookEditorViewportClient.h
+++ b/Engine/Plugins/2D/Paper2D/Source/Paper2DEditor/Private/FlipbookEditor/FlipbookEditorViewportClient.h
@@ -2,11 +2,7 @@
 
 #pragma once
 
-#include "CoreMinimal.h"
-#include "Misc/Attribute.h"
-#include "InputCoreTypes.h"
 #include "PreviewScene.h"
-#include "PaperFlipbook.h"
 #include "PaperEditorViewportClient.h"
 
 class FCanvas;
diff --git a/Engine/Plugins/2D/Paper2D/Source/Paper2DEditor/Private/FlipbookEditor/SFlipbookEditorViewportToolbar.cpp b/Engine/Plugins/2D/Paper2D/Source/Paper2DEditor/Private/FlipbookEditor/SFlipbookEditorViewportToolbar.cpp
--- a/Engine/Plugins/2D/Paper2D/Source/Paper2DEditor/Private/FlipbookEditor/SFlipbookEditorViewportToolbar.cpp
+++ b/Engine/Plugins/2D/Paper2D/Source/Paper2DEditor/Private/FlipbookEditor/SFlipbookEditorViewportToolbar.cpp
@@ -1,6 +1,7 @@
 // Copyright Epic Games, Inc. All Rights Reserved.
 
 #include "FlipbookEditor/SFlipbookEditorViewportToolbar.h"
+#include "Framework/Commands/UICommandList.h"
 #include "SEditorViewport.h"
 #include "FlipbookEditor/FlipbookEditorCommands.h"
 
diff --git a/Engine/Plugins/2D/Paper2D/Source/Paper2DEditor/Private/FlipbookEditor/SFlipbookEditorViewportToolbar.h b/Engine/Plugins/2D/Paper2D/Source/Paper2DEditor/Private/FlipbookEditor/SFlipbookEditorViewportToolbar.h
--- a/Engine/Plugins/2D/Paper2D/Source/Paper2DEditor/Private/FlipbookEditor/SFlipbookEditorViewportToolbar.h
+++ b/Engine/Plugins/2D/Paper2D/Source/Paper2DEditor/Private/FlipbookEditor/SFlipbookEditorViewportToolbar.h
@@ -2,9 +2,6 @@
 
 #pragma once
 
-#include "CoreMinimal.h"
-#include "Widgets/SWidget.h"
-#include "Widgets/DeclarativeSyntaxSupport.h"
 #include "SCommonEditorViewportToolbarBase.h"
 
 // In-viewport toolbar widget used in the flipbook editor
diff --git a/Engine/Plugins/2D/Paper2D/Source/Paper2DEditor/Private/FlipbookEditor/SFlipbookTimeline.cpp b/Engine/Plugins/2D/Paper2D/Source/Paper2DEditor/Private/FlipbookEditor/SFlipbookTimeline.cpp
--- a/Engine/Plugins/2D/Paper2D/Source/Paper2DEditor/Private/FlipbookEditor/SFlipbookTimeline.cpp
+++ b/Engine/Plugins/2D/Paper2D/Source/Paper2DEditor/Private/FlipbookEditor/SFlipbookTimeline.cpp
@@ -1,22 +1,14 @@
 // Copyright Epic Games, Inc. All Rights Reserved.
 
 #include "FlipbookEditor/SFlipbookTimeline.h"
-#include "Rendering/DrawElements.h"
-#include "Widgets/SBoxPanel.h"
+#include "Framework/Commands/UICommandList.h"
 #include "Widgets/SOverlay.h"
 #include "Layout/WidgetPath.h"
-#include "Framework/Application/MenuStack.h"
 #include "Framework/Application/SlateApplication.h"
-#include "Widgets/Layout/SBorder.h"
-#include "Widgets/Text/STextBlock.h"
-#include "Widgets/Layout/SBox.h"
 #include "Framework/MultiBox/MultiBoxBuilder.h"
 #include "Widgets/Layout/SScrollBox.h"
 #include "Widgets/Colors/SColorBlock.h"
-#include "Styling/AppStyle.h"
-#include "PaperFlipbook.h"
 #include "DragAndDrop/AssetDragDropOp.h"
-#include "ScopedTransaction.h"
 #include "FlipbookEditor/FlipbookEditorCommands.h"
 #include "PaperSprite.h"
 #include "FlipbookEditor/STimelineTrack.h"
diff --git a/Engine/Plugins/2D/Paper2D/Source/Paper2DEditor/Private/FlipbookEditor/SFlipbookTimeline.h b/Engine/Plugins/2D/Paper2D/Source/Paper2DEditor/Private/FlipbookEditor/SFlipbookTimeline.h
--- a/Engine/Plugins/2D/Paper2D/Source/Paper2DEditor/Private/FlipbookEditor/SFlipbookTimeline.h
+++ b/Engine/Plugins/2D/Paper2D/Source/Paper2DEditor/Private/FlipbookEditor/SFlipbookTimeline.h
@@ -2,14 +2,10 @@
 
 #pragma once
 
-#include "CoreMinimal.h"
-#include "Misc/Attribute.h"
-#include "Input/Reply.h"
-#include "Layout/Visibility.h"
-#include "Widgets/SWidget.h"
-#include "Widgets/DeclarativeSyntaxSupport.h"
 #include "Widgets/SCompoundWidget.h"
 
+class UPaperFlipbook;
+
 class FPaintArgs;
 class FSlateWindowElementList;
 class FUICommandList;
diff --git a/Engine/Plugins/2D/Paper2D/Source/Paper2DEditor/Private/FlipbookEditor/SFlipbookTrackHandle.h b/Engine/Plugins/2D/Paper2D/Source/Paper2DEditor/Private/FlipbookEditor/SFlipbookTrackHandle.h
--- a/Engine/Plugins/2D/Paper2D/Source/Paper2DEditor/Private/FlipbookEditor/SFlipbookTrackHandle.h
+++ b/Engine/Plugins/2D/Paper2D/Source/Paper2DEditor/Private/FlipbookEditor/SFlipbookTrackHandle.h
@@ -2,12 +2,6 @@
 
 #pragma once
 
-#include "CoreMinimal.h"
-#include "Misc/Attribute.h"
-#include "InputCoreTypes.h"
-#include "Input/CursorReply.h"
-#include "Input/Reply.h"
-#include "Widgets/DeclarativeSyntaxSupport.h"
 #include "PaperFlipbook.h"
 #include "Framework/Application/SlateApplication.h"
 #include "PaperStyle.h"
diff --git a/Engine/Plugins/2D/Paper2D/Source/Paper2DEditor/Private/FlipbookEditor/STimelineHeader.h b/Engine/Plugins/2D/Paper2D/Source/Paper2DEditor/Private/FlipbookEditor/STimelineHeader.h
--- a/Engine/Plugins/2D/Paper2D/Source/Paper2DEditor/Private/FlipbookEditor/STimelineHeader.h
+++ b/Engine/Plugins/2D/Paper2D/Source/Paper2DEditor/Private/FlipbookEditor/STimelineHeader.h
@@ -2,9 +2,6 @@
 
 #pragma once
 
-#include "CoreMinimal.h"
-#include "Misc/Attribute.h"
-#include "Widgets/DeclarativeSyntaxSupport.h"
 #include "Widgets/SCompoundWidget.h"
 #include "Widgets/SBoxPanel.h"
 #include "Widgets/Text/STextBlock.h"
diff --git a/Engine/Plugins/2D/Paper2D/Source/Paper2DEditor/Private/FlipbookEditor/STimelineTrack.cpp b/Engine/Plugins/2D/Paper2D/Source/Paper2DEditor/Private/FlipbookEditor/STimelineTrack.cpp
--- a/Engine/Plugins/2D/Paper2D/Source/Paper2DEditor/Private/FlipbookEditor/STimelineTrack.cpp
+++ b/Engine/Plugins/2D/Paper2D/Source/Paper2DEditor/Private/FlipbookEditor/STimelineTrack.cpp
@@ -1,19 +1,10 @@
 // Copyright Epic Games, Inc. All Rights Reserved.
 
 #include "FlipbookEditor/STimelineTrack.h"
-#include "Widgets/SBoxPanel.h"
-#include "Widgets/SOverlay.h"
 #include "Layout/WidgetPath.h"
-#include "Framework/Application/MenuStack.h"
-#include "Framework/Application/SlateApplication.h"
 #include "Framework/Commands/UICommandList.h"
-#include "Widgets/Layout/SBorder.h"
-#include "Widgets/Text/STextBlock.h"
-#include "Widgets/Layout/SBox.h"
 #include "Framework/MultiBox/MultiBoxBuilder.h"
-#include "AssetRegistry/AssetData.h"
 #include "Framework/Commands/GenericCommands.h"
-#include "PaperStyle.h"
 #include "FlipbookEditor/FlipbookEditorCommands.h"
 #include "FlipbookEditor/SFlipbookTrackHandle.h"
 #include "DragAndDrop/AssetDragDropOp.h"
@@ -22,6 +13,8 @@
 #include "PaperSprite.h"
 
 #include "Subsystems/AssetEditorSubsystem.h"
+#include "Subsystems/ImportSubsystem.h"
+#include "Widgets/SWindow.h"
 
 #define LOCTEXT_NAMESPACE "FlipbookEditor"
 
diff --git a/Engine/Plugins/2D/Paper2D/Source/Paper2DEditor/Private/FlipbookEditor/STimelineTrack.h b/Engine/Plugins/2D/Paper2D/Source/Paper2DEditor/Private/FlipbookEditor/STimelineTrack.h
--- a/Engine/Plugins/2D/Paper2D/Source/Paper2DEditor/Private/FlipbookEditor/STimelineTrack.h
+++ b/Engine/Plugins/2D/Paper2D/Source/Paper2DEditor/Private/FlipbookEditor/STimelineTrack.h
@@ -2,19 +2,12 @@
 
 #pragma once 
 
-#include "CoreMinimal.h"
-#include "Misc/Attribute.h"
-#include "Input/DragAndDrop.h"
-#include "Input/Reply.h"
-#include "Layout/Margin.h"
-#include "Widgets/SWidget.h"
-#include "Widgets/DeclarativeSyntaxSupport.h"
-#include "Widgets/SCompoundWidget.h"
 #include "FlipbookEditor/SFlipbookTimeline.h"
-#include "Types/SlateStructs.h"
 #include "PaperFlipbook.h"
 #include "ScopedTransaction.h"
 
+struct FOptionalSize;
+
 class FUICommandList;
 class SHorizontalBox;
 
diff --git a/Engine/Plugins/2D/Paper2D/Source/Paper2DEditor/Private/GroupedSprites/GroupedSpriteDetailsCustomization.cpp b/Engine/Plugins/2D/Paper2D/Source/Paper2DEditor/Private/GroupedSprites/GroupedSpriteDetailsCustomization.cpp
--- a/Engine/Plugins/2D/Paper2D/Source/Paper2DEditor/Private/GroupedSprites/GroupedSpriteDetailsCustomization.cpp
+++ b/Engine/Plugins/2D/Paper2D/Source/Paper2DEditor/Private/GroupedSprites/GroupedSpriteDetailsCustomization.cpp
@@ -1,8 +1,6 @@
 // Copyright Epic Games, Inc. All Rights Reserved.
 
 #include "GroupedSprites/GroupedSpriteDetailsCustomization.h"
-#include "Layout/Margin.h"
-#include "Widgets/DeclarativeSyntaxSupport.h"
 #include "Widgets/SBoxPanel.h"
 #include "Widgets/Layout/SWrapBox.h"
 #include "Widgets/Layout/SBox.h"
@@ -10,7 +8,6 @@
 #include "Editor.h"
 #include "DetailLayoutBuilder.h"
 #include "DetailWidgetRow.h"
-#include "DetailCategoryBuilder.h"
 
 #include "PaperGroupedSpriteComponent.h"
 #include "GroupedSprites/PaperGroupedSpriteUtilities.h"
diff --git a/Engine/Plugins/2D/Paper2D/Source/Paper2DEditor/Private/GroupedSprites/GroupedSpriteDetailsCustomization.h b/Engine/Plugins/2D/Paper2D/Source/Paper2DEditor/Private/GroupedSprites/GroupedSpriteDetailsCustomization.h
--- a/Engine/Plugins/2D/Paper2D/Source/Paper2DEditor/Private/GroupedSprites/GroupedSpriteDetailsCustomization.h
+++ b/Engine/Plugins/2D/Paper2D/Source/Paper2DEditor/Private/GroupedSprites/GroupedSpriteDetailsCustomization.h
@@ -2,9 +2,10 @@
 
 #pragma once
 
-#include "CoreMinimal.h"
-#include "Input/Reply.h"
 #include "IDetailCustomization.h"
+#include "UObject/WeakObjectPtrTemplates.h"
+
+class FReply;
 
 class IDetailLayoutBuilder;
 
diff --git a/Engine/Plugins/2D/Paper2D/Source/Paper2DEditor/Private/GroupedSprites/PaperGroupedSpriteUtilities.cpp b/Engine/Plugins/2D/Paper2D/Source/Paper2DEditor/Private/GroupedSprites/PaperGroupedSpriteUtilities.cpp
--- a/Engine/Plugins/2D/Paper2D/Source/Paper2DEditor/Private/GroupedSprites/PaperGroupedSpriteUtilities.cpp
+++ b/Engine/Plugins/2D/Paper2D/Source/Paper2DEditor/Private/GroupedSprites/PaperGroupedSpriteUtilities.cpp
@@ -2,10 +2,6 @@
 
 #include "GroupedSprites/PaperGroupedSpriteUtilities.h"
 #include "Paper2DEditorLog.h"
-#include "Components/SceneComponent.h"
-#include "GameFramework/Actor.h"
-#include "Engine/World.h"
-#include "Editor/EditorEngine.h"
 #include "Kismet/GameplayStatics.h"
 #include "Editor.h"
 #include "PaperGroupedSpriteActor.h"
diff --git a/Engine/Plugins/2D/Paper2D/Source/Paper2DEditor/Private/GroupedSprites/PaperGroupedSpriteUtilities.h b/Engine/Plugins/2D/Paper2D/Source/Paper2DEditor/Private/GroupedSprites/PaperGroupedSpriteUtilities.h
--- a/Engine/Plugins/2D/Paper2D/Source/Paper2DEditor/Private/GroupedSprites/PaperGroupedSpriteUtilities.h
+++ b/Engine/Plugins/2D/Paper2D/Source/Paper2DEditor/Private/GroupedSprites/PaperGroupedSpriteUtilities.h
@@ -2,9 +2,10 @@
 
 #pragma once
 
-#include "CoreMinimal.h"
-#include "Templates/SubclassOf.h"
-#include "Components/ActorComponent.h"
+#include "Math/Box.h"
+
+class UActorComponent;
+template <class TClass> class TSubclassOf;
 
 class AActor;
 
diff --git a/Engine/Plugins/2D/Paper2D/Source/Paper2DEditor/Private/LevelEditorMenuExtensions/Paper2DLevelEditorExtensions.cpp b/Engine/Plugins/2D/Paper2D/Source/Paper2DEditor/Private/LevelEditorMenuExtensions/Paper2DLevelEditorExtensions.cpp
--- a/Engine/Plugins/2D/Paper2D/Source/Paper2DEditor/Private/LevelEditorMenuExtensions/Paper2DLevelEditorExtensions.cpp
+++ b/Engine/Plugins/2D/Paper2D/Source/Paper2DEditor/Private/LevelEditorMenuExtensions/Paper2DLevelEditorExtensions.cpp
@@ -2,19 +2,13 @@
 
 #include "LevelEditorMenuExtensions/Paper2DLevelEditorExtensions.h"
 #include "Modules/ModuleManager.h"
-#include "Textures/SlateIcon.h"
-#include "Framework/Commands/UIAction.h"
-#include "Framework/Commands/UICommandList.h"
-#include "Framework/MultiBox/MultiBoxExtender.h"
-#include "GameFramework/Actor.h"
-#include "Engine/Selection.h"
 #include "Framework/MultiBox/MultiBoxBuilder.h"
 #include "LevelEditor.h"
-#include "Editor.h"
 
 #include "PaperSpriteComponent.h"
 #include "PaperGroupedSpriteComponent.h"
 #include "GroupedSprites/PaperGroupedSpriteUtilities.h"
+#include "Selection.h"
 
 #define LOCTEXT_NAMESPACE "Paper2D"
 
diff --git a/Engine/Plugins/2D/Paper2D/Source/Paper2DEditor/Private/MeshPainting/MeshPaintSpriteAdapter.h b/Engine/Plugins/2D/Paper2D/Source/Paper2DEditor/Private/MeshPainting/MeshPaintSpriteAdapter.h
--- a/Engine/Plugins/2D/Paper2D/Source/Paper2DEditor/Private/MeshPainting/MeshPaintSpriteAdapter.h
+++ b/Engine/Plugins/2D/Paper2D/Source/Paper2DEditor/Private/MeshPainting/MeshPaintSpriteAdapter.h
@@ -2,8 +2,6 @@
 
 #pragma once
 
-#include "CoreMinimal.h"
-#include "MeshPaintModule.h"
 #include "IMeshPaintGeometryAdapter.h"
 #include "IMeshPaintGeometryAdapterFactory.h"
 
diff --git a/Engine/Plugins/2D/Paper2D/Source/Paper2DEditor/Private/Paper2DEditorModule.cpp b/Engine/Plugins/2D/Paper2D/Source/Paper2DEditor/Private/Paper2DEditorModule.cpp
--- a/Engine/Plugins/2D/Paper2D/Source/Paper2DEditor/Private/Paper2DEditorModule.cpp
+++ b/Engine/Plugins/2D/Paper2D/Source/Paper2DEditor/Private/Paper2DEditorModule.cpp
@@ -1,43 +1,30 @@
 // Copyright Epic Games, Inc. All Rights Reserved.
 
 #include "Paper2DEditorModule.h"
+#include "Misc/CoreDelegates.h"
 #include "Paper2DEditorLog.h"
-#include "Engine/Texture2D.h"
-#include "Editor.h"
 #include "EditorModeRegistry.h"
-#include "Modules/ModuleManager.h"
-#include "UObject/UObjectHash.h"
+#include "Subsystems/ImportSubsystem.h"
 #include "UObject/UObjectIterator.h"
-#include "ThumbnailRendering/ThumbnailManager.h"
 
-#include "AssetToolsModule.h"
 #include "PropertyEditorModule.h"
-#include "PaperStyle.h"
 #include "PaperEditorCommands.h"
-#include "PaperSprite.h"
 #include "PaperEditorShared/SpriteGeometryEditMode.h"
 #include "PaperTileSet.h"
 
 #include "ContentBrowserExtensions/ContentBrowserExtensions.h"
 #include "LevelEditorMenuExtensions/Paper2DLevelEditorExtensions.h"
-#include "PaperTileMap.h"
 #include "PaperImporterSettings.h"
 
 // Sprite support
-#include "IAssetTypeActions.h"
 #include "SpriteAssetTypeActions.h"
-#include "ComponentAssetBroker.h"
-#include "PaperSpriteComponent.h"
 #include "PaperSpriteAssetBroker.h"
-#include "PaperSpriteThumbnailRenderer.h"
 #include "SpriteEditor/SpriteDetailsCustomization.h"
 #include "SpriteEditor/SpriteComponentDetailsCustomization.h"
 #include "SpriteEditor/SpriteEditorSettings.h"
 
 // Flipbook support
 #include "FlipbookAssetTypeActions.h"
-#include "PaperFlipbookComponent.h"
-#include "PaperFlipbook.h"
 #include "PaperFlipbookAssetBroker.h"
 #include "PaperFlipbookThumbnailRenderer.h"
 #include "FlipbookEditor/FlipbookComponentDetailsCustomization.h"
@@ -51,7 +38,6 @@
 
 // Tile map support
 #include "TileMapEditing/TileMapAssetTypeActions.h"
-#include "PaperTileMapComponent.h"
 #include "TileMapEditing/PaperTileMapAssetBroker.h"
 #include "TileMapEditing/EdModeTileMap.h"
 #include "TileMapEditing/PaperTileMapDetailsCustomization.h"
diff --git a/Engine/Plugins/2D/Paper2D/Source/Paper2DEditor/Private/PaperEditorCommands.h b/Engine/Plugins/2D/Paper2D/Source/Paper2DEditor/Private/PaperEditorCommands.h
--- a/Engine/Plugins/2D/Paper2D/Source/Paper2DEditor/Private/PaperEditorCommands.h
+++ b/Engine/Plugins/2D/Paper2D/Source/Paper2DEditor/Private/PaperEditorCommands.h
@@ -2,9 +2,9 @@
 
 #pragma once
 
-#include "CoreMinimal.h"
 #include "Framework/Commands/Commands.h"
 #include "PaperStyle.h"
+#include "Styling/ISlateStyle.h"
 
 class FPaperEditorCommands : public TCommands<FPaperEditorCommands>
 {
diff --git a/Engine/Plugins/2D/Paper2D/Source/Paper2DEditor/Private/PaperEditorShared/AssetEditorSelectedItem.h b/Engine/Plugins/2D/Paper2D/Source/Paper2DEditor/Private/PaperEditorShared/AssetEditorSelectedItem.h
--- a/Engine/Plugins/2D/Paper2D/Source/Paper2DEditor/Private/PaperEditorShared/AssetEditorSelectedItem.h
+++ b/Engine/Plugins/2D/Paper2D/Source/Paper2DEditor/Private/PaperEditorShared/AssetEditorSelectedItem.h
@@ -2,9 +2,9 @@
 
 #pragma once
 
-#include "CoreMinimal.h"
 #include "GenericPlatform/ICursor.h"
 #include "HitProxies.h"
+#include "Math/Rotator.h"
 #include "UnrealWidgetFwd.h"
 
 //////////////////////////////////////////////////////////////////////////
diff --git a/Engine/Plugins/2D/Paper2D/Source/Paper2DEditor/Private/PaperEditorShared/SocketEditing.cpp b/Engine/Plugins/2D/Paper2D/Source/Paper2DEditor/Private/PaperEditorShared/SocketEditing.cpp
--- a/Engine/Plugins/2D/Paper2D/Source/Paper2DEditor/Private/PaperEditorShared/SocketEditing.cpp
+++ b/Engine/Plugins/2D/Paper2D/Source/Paper2DEditor/Private/PaperEditorShared/SocketEditing.cpp
@@ -1,16 +1,15 @@
 // Copyright Epic Games, Inc. All Rights Reserved.
 
 #include "PaperEditorShared/SocketEditing.h"
-#include "Engine/EngineTypes.h"
 #include "SceneManagement.h"
 #include "SceneView.h"
 #include "CanvasItem.h"
-#include "Components/Viewport.h"
 #include "Paper2DModule.h"
 #include "PaperSpriteComponent.h"
 #include "PaperSprite.h"
 #include "PaperEditorShared/SpriteGeometryEditMode.h"
 #include "CanvasTypes.h"
+#include "UnrealClient.h"
 
 //////////////////////////////////////////////////////////////////////////
 // HSpriteSelectableObjectHitProxy
diff --git a/Engine/Plugins/2D/Paper2D/Source/Paper2DEditor/Private/PaperEditorShared/SocketEditing.h b/Engine/Plugins/2D/Paper2D/Source/Paper2DEditor/Private/PaperEditorShared/SocketEditing.h
--- a/Engine/Plugins/2D/Paper2D/Source/Paper2DEditor/Private/PaperEditorShared/SocketEditing.h
+++ b/Engine/Plugins/2D/Paper2D/Source/Paper2DEditor/Private/PaperEditorShared/SocketEditing.h
@@ -2,9 +2,7 @@
 
 #pragma once
 
-#include "CoreMinimal.h"
 #include "UObject/WeakObjectPtr.h"
-#include "UnrealWidgetFwd.h"
 #include "PaperEditorShared/AssetEditorSelectedItem.h"
 
 class FCanvas;
diff --git a/Engine/Plugins/2D/Paper2D/Source/Paper2DEditor/Private/PaperEditorShared/SpriteGeometryEditMode.cpp b/Engine/Plugins/2D/Paper2D/Source/Paper2DEditor/Private/PaperEditorShared/SpriteGeometryEditMode.cpp
--- a/Engine/Plugins/2D/Paper2D/Source/Paper2DEditor/Private/PaperEditorShared/SpriteGeometryEditMode.cpp
+++ b/Engine/Plugins/2D/Paper2D/Source/Paper2DEditor/Private/PaperEditorShared/SpriteGeometryEditMode.cpp
@@ -7,8 +7,6 @@
 #include "CanvasItem.h"
 #include "CanvasTypes.h"
 #include "EditorModeManager.h"
-#include "PaperEditorShared/AssetEditorSelectedItem.h"
-#include "SpriteEditorOnlyTypes.h"
 #include "SpriteEditor/SpriteEditorSelections.h"
 #include "PaperEditorShared/SpriteGeometryEditCommands.h"
 #include "PhysicsEngine/BodySetup.h"
diff --git a/Engine/Plugins/2D/Paper2D/Source/Paper2DEditor/Private/PaperEditorShared/SpriteGeometryEditMode.h b/Engine/Plugins/2D/Paper2D/Source/Paper2DEditor/Private/PaperEditorShared/SpriteGeometryEditMode.h
--- a/Engine/Plugins/2D/Paper2D/Source/Paper2DEditor/Private/PaperEditorShared/SpriteGeometryEditMode.h
+++ b/Engine/Plugins/2D/Paper2D/Source/Paper2DEditor/Private/PaperEditorShared/SpriteGeometryEditMode.h
@@ -2,8 +2,6 @@
 
 #pragma once
 
-#include "CoreMinimal.h"
-#include "InputCoreTypes.h"
 #include "EdMode.h"
 #include "PaperEditorShared/SpriteGeometryEditing.h"
 
diff --git a/Engine/Plugins/2D/Paper2D/Source/Paper2DEditor/Private/PaperEditorShared/SpriteGeometryEditing.cpp b/Engine/Plugins/2D/Paper2D/Source/Paper2DEditor/Private/PaperEditorShared/SpriteGeometryEditing.cpp
--- a/Engine/Plugins/2D/Paper2D/Source/Paper2DEditor/Private/PaperEditorShared/SpriteGeometryEditing.cpp
+++ b/Engine/Plugins/2D/Paper2D/Source/Paper2DEditor/Private/PaperEditorShared/SpriteGeometryEditing.cpp
@@ -1,16 +1,13 @@
 // Copyright Epic Games, Inc. All Rights Reserved.
 
 #include "PaperEditorShared/SpriteGeometryEditing.h"
+#include "GlobalRenderResources.h"
 #include "Materials/Material.h"
 #include "CanvasItem.h"
 #include "CanvasTypes.h"
 #include "Engine/Engine.h"
-#include "EngineGlobals.h"
 #include "SceneManagement.h"
 #include "SceneView.h"
-#include "UnrealWidgetFwd.h"
-#include "PaperEditorShared/AssetEditorSelectedItem.h"
-#include "SpriteEditorOnlyTypes.h"
 #include "SpriteEditor/SpriteEditorSelections.h"
 #include "GeomTools.h"
 #include "DynamicMeshBuilder.h"
diff --git a/Engine/Plugins/2D/Paper2D/Source/Paper2DEditor/Private/PaperEditorShared/SpriteGeometryEditing.h b/Engine/Plugins/2D/Paper2D/Source/Paper2DEditor/Private/PaperEditorShared/SpriteGeometryEditing.h
--- a/Engine/Plugins/2D/Paper2D/Source/Paper2DEditor/Private/PaperEditorShared/SpriteGeometryEditing.h
+++ b/Engine/Plugins/2D/Paper2D/Source/Paper2DEditor/Private/PaperEditorShared/SpriteGeometryEditing.h
@@ -2,10 +2,11 @@
 
 #pragma once
 
-#include "CoreMinimal.h"
-#include "Engine/EngineBaseTypes.h"
+#include "Math/Vector2D.h"
 #include "UObject/GCObject.h"
 
+enum EInputEvent : int;
+
 class FCanvas;
 class FPrimitiveDrawInterface;
 class FSceneView;
diff --git a/Engine/Plugins/2D/Paper2D/Source/Paper2DEditor/Private/PaperEditorViewportClient.cpp b/Engine/Plugins/2D/Paper2D/Source/Paper2DEditor/Private/PaperEditorViewportClient.cpp
--- a/Engine/Plugins/2D/Paper2D/Source/Paper2DEditor/Private/PaperEditorViewportClient.cpp
+++ b/Engine/Plugins/2D/Paper2D/Source/Paper2DEditor/Private/PaperEditorViewportClient.cpp
@@ -1,8 +1,8 @@
 // Copyright Epic Games, Inc. All Rights Reserved.
 
 #include "PaperEditorViewportClient.h"
+#include "Editor/UnrealEdTypes.h"
 #include "Paper2DModule.h"
-#include "AssetEditorModeManager.h"
 #include "CanvasTypes.h"
 #include "Engine/Texture2D.h"
 #include "ImageUtils.h"
diff --git a/Engine/Plugins/2D/Paper2D/Source/Paper2DEditor/Private/PaperEditorViewportClient.h b/Engine/Plugins/2D/Paper2D/Source/Paper2DEditor/Private/PaperEditorViewportClient.h
--- a/Engine/Plugins/2D/Paper2D/Source/Paper2DEditor/Private/PaperEditorViewportClient.h
+++ b/Engine/Plugins/2D/Paper2D/Source/Paper2DEditor/Private/PaperEditorViewportClient.h
@@ -2,9 +2,7 @@
 
 #pragma once
 
-#include "CoreMinimal.h"
 #include "EditorViewportClient.h"
-#include "SEditorViewport.h"
 
 class FCanvas;
 
diff --git a/Engine/Plugins/2D/Paper2D/Source/Paper2DEditor/Private/PaperFlipbookThumbnailRenderer.cpp b/Engine/Plugins/2D/Paper2D/Source/Paper2DEditor/Private/PaperFlipbookThumbnailRenderer.cpp
--- a/Engine/Plugins/2D/Paper2D/Source/Paper2DEditor/Private/PaperFlipbookThumbnailRenderer.cpp
+++ b/Engine/Plugins/2D/Paper2D/Source/Paper2DEditor/Private/PaperFlipbookThumbnailRenderer.cpp
@@ -1,10 +1,10 @@
 // Copyright Epic Games, Inc. All Rights Reserved.
 
 #include "PaperFlipbookThumbnailRenderer.h"
+#include "Math/BoxSphereBounds.h"
 #include "Misc/App.h"
 #include "CanvasItem.h"
 #include "Engine/Engine.h"
-#include "EngineGlobals.h"
 #include "PaperFlipbook.h"
 
 #include UE_INLINE_GENERATED_CPP_BY_NAME(PaperFlipbookThumbnailRenderer)
diff --git a/Engine/Plugins/2D/Paper2D/Source/Paper2DEditor/Private/PaperJSONHelpers.cpp b/Engine/Plugins/2D/Paper2D/Source/Paper2DEditor/Private/PaperJSONHelpers.cpp
--- a/Engine/Plugins/2D/Paper2D/Source/Paper2DEditor/Private/PaperJSONHelpers.cpp
+++ b/Engine/Plugins/2D/Paper2D/Source/Paper2DEditor/Private/PaperJSONHelpers.cpp
@@ -1,6 +1,7 @@
 // Copyright Epic Games, Inc. All Rights Reserved.
 
 #include "PaperJSONHelpers.h"
+#include "Dom/JsonObject.h"
 
 //////////////////////////////////////////////////////////////////////////
 // FPaperJSONHelpers
diff --git a/Engine/Plugins/2D/Paper2D/Source/Paper2DEditor/Private/PaperSpriteFactory.cpp b/Engine/Plugins/2D/Paper2D/Source/Paper2DEditor/Private/PaperSpriteFactory.cpp
--- a/Engine/Plugins/2D/Paper2D/Source/Paper2DEditor/Private/PaperSpriteFactory.cpp
+++ b/Engine/Plugins/2D/Paper2D/Source/Paper2DEditor/Private/PaperSpriteFactory.cpp
@@ -3,9 +3,7 @@
 #include "PaperSpriteFactory.h"
 #include "Modules/ModuleManager.h"
 #include "Misc/PackageName.h"
-#include "Engine/Texture2D.h"
 #include "AssetRegistry/AssetData.h"
-#include "SpriteEditorOnlyTypes.h"
 #include "PaperImporterSettings.h"
 #include "PaperSprite.h"
 #include "AssetRegistry/AssetRegistryModule.h"
diff --git a/Engine/Plugins/2D/Paper2D/Source/Paper2DEditor/Private/PaperSpriteThumbnailRenderer.cpp b/Engine/Plugins/2D/Paper2D/Source/Paper2DEditor/Private/PaperSpriteThumbnailRenderer.cpp
--- a/Engine/Plugins/2D/Paper2D/Source/Paper2DEditor/Private/PaperSpriteThumbnailRenderer.cpp
+++ b/Engine/Plugins/2D/Paper2D/Source/Paper2DEditor/Private/PaperSpriteThumbnailRenderer.cpp
@@ -1,9 +1,7 @@
 // Copyright Epic Games, Inc. All Rights Reserved.
 
 #include "PaperSpriteThumbnailRenderer.h"
-#include "Engine/EngineTypes.h"
 #include "CanvasItem.h"
-#include "Engine/Texture2D.h"
 #include "Paper2DModule.h"
 #include "CanvasTypes.h"
 #include "PaperSprite.h"
diff --git a/Engine/Plugins/2D/Paper2D/Source/Paper2DEditor/Private/PaperStyle.cpp b/Engine/Plugins/2D/Paper2D/Source/Paper2DEditor/Private/PaperStyle.cpp
--- a/Engine/Plugins/2D/Paper2D/Source/Paper2DEditor/Private/PaperStyle.cpp
+++ b/Engine/Plugins/2D/Paper2D/Source/Paper2DEditor/Private/PaperStyle.cpp
@@ -1,12 +1,15 @@
 // Copyright Epic Games, Inc. All Rights Reserved.
 
 #include "PaperStyle.h"
+#include "Brushes/SlateBoxBrush.h"
 #include "Styling/SlateStyleRegistry.h"
+#include "Brushes/SlateImageBrush.h"
 #include "Styling/SlateTypes.h"
+#include "Misc/Paths.h"
 #include "Styling/CoreStyle.h"
-#include "Styling/AppStyle.h"
 #include "Interfaces/IPluginManager.h"
 #include "SlateOptMacros.h"
+#include "Styling/SlateStyle.h"
 
 
 #define IMAGE_PLUGIN_BRUSH( RelativePath, ... ) FSlateImageBrush( FPaperStyle::InContent( RelativePath, ".png" ), __VA_ARGS__ )
diff --git a/Engine/Plugins/2D/Paper2D/Source/Paper2DEditor/Private/PaperStyle.h b/Engine/Plugins/2D/Paper2D/Source/Paper2DEditor/Private/PaperStyle.h
--- a/Engine/Plugins/2D/Paper2D/Source/Paper2DEditor/Private/PaperStyle.h
+++ b/Engine/Plugins/2D/Paper2D/Source/Paper2DEditor/Private/PaperStyle.h
@@ -2,9 +2,7 @@
 
 #pragma once
 
-#include "CoreMinimal.h"
-#include "Styling/ISlateStyle.h"
-#include "Styling/SlateStyle.h"
+#include "Templates/SharedPointer.h"
 
 class FPaperStyle
 {
diff --git a/Engine/Plugins/2D/Paper2D/Source/Paper2DEditor/Private/PaperTileSetThumbnailRenderer.cpp b/Engine/Plugins/2D/Paper2D/Source/Paper2DEditor/Private/PaperTileSetThumbnailRenderer.cpp
--- a/Engine/Plugins/2D/Paper2D/Source/Paper2DEditor/Private/PaperTileSetThumbnailRenderer.cpp
+++ b/Engine/Plugins/2D/Paper2D/Source/Paper2DEditor/Private/PaperTileSetThumbnailRenderer.cpp
@@ -1,11 +1,8 @@
 // Copyright Epic Games, Inc. All Rights Reserved.
 
 #include "PaperTileSetThumbnailRenderer.h"
-#include "Engine/Texture2D.h"
 #include "CanvasTypes.h"
-#include "IntMargin.h"
 #include "PaperTileSet.h"
-#include "RenderUtils.h"
 #include "TextureResource.h"
 #include "GlobalRenderResources.h"
 
diff --git a/Engine/Plugins/2D/Paper2D/Source/Paper2DEditor/Private/SPaperEditorViewport.cpp b/Engine/Plugins/2D/Paper2D/Source/Paper2DEditor/Private/SPaperEditorViewport.cpp
--- a/Engine/Plugins/2D/Paper2D/Source/Paper2DEditor/Private/SPaperEditorViewport.cpp
+++ b/Engine/Plugins/2D/Paper2D/Source/Paper2DEditor/Private/SPaperEditorViewport.cpp
@@ -1,15 +1,13 @@
 // Copyright Epic Games, Inc. All Rights Reserved.
 
 #include "SPaperEditorViewport.h"
-#include "Rendering/DrawElements.h"
+#include "Math/ColorList.h"
 #include "Widgets/SBoxPanel.h"
-#include "Widgets/SViewport.h"
-#include "Application/ThrottleManager.h"
 #include "Framework/Application/SlateApplication.h"
+#include "PaperEditorViewportClient.h"
 #include "Widgets/Layout/SBorder.h"
+#include "Widgets/SOverlay.h"
 #include "Widgets/Text/STextBlock.h"
-#include "Styling/AppStyle.h"
-#include "Editor.h"
 #include "PaperStyle.h"
 
 const int32 DefaultZoomLevel = 7;
diff --git a/Engine/Plugins/2D/Paper2D/Source/Paper2DEditor/Private/SPaperEditorViewport.h b/Engine/Plugins/2D/Paper2D/Source/Paper2DEditor/Private/SPaperEditorViewport.h
--- a/Engine/Plugins/2D/Paper2D/Source/Paper2DEditor/Private/SPaperEditorViewport.h
+++ b/Engine/Plugins/2D/Paper2D/Source/Paper2DEditor/Private/SPaperEditorViewport.h
@@ -2,16 +2,8 @@
 
 #pragma once
 
-#include "CoreMinimal.h"
-#include "Styling/SlateColor.h"
-#include "Layout/SlateRect.h"
-#include "Input/CursorReply.h"
-#include "Input/Reply.h"
 #include "Animation/CurveSequence.h"
-#include "Widgets/DeclarativeSyntaxSupport.h"
-#include "EditorViewportClient.h"
 #include "SEditorViewport.h"
-#include "PaperEditorViewportClient.h"
 #include "Framework/MarqueeRect.h"
 
 class FPaintArgs;
diff --git a/Engine/Plugins/2D/Paper2D/Source/Paper2DEditor/Private/SpriteAssetTypeActions.cpp b/Engine/Plugins/2D/Paper2D/Source/Paper2DEditor/Private/SpriteAssetTypeActions.cpp
--- a/Engine/Plugins/2D/Paper2D/Source/Paper2DEditor/Private/SpriteAssetTypeActions.cpp
+++ b/Engine/Plugins/2D/Paper2D/Source/Paper2DEditor/Private/SpriteAssetTypeActions.cpp
@@ -1,19 +1,16 @@
 // Copyright Epic Games, Inc. All Rights Reserved.
 
 #include "SpriteAssetTypeActions.h"
-#include "Framework/MultiBox/MultiBoxBuilder.h"
-#include "ToolMenus.h"
 #include "Misc/PackageName.h"
 #include "Misc/FeedbackContext.h"
-#include "Styling/AppStyle.h"
 
 #include "SpriteEditor/SpriteEditor.h"
 #include "IContentBrowserSingleton.h"
 #include "ContentBrowserModule.h"
 #include "PaperFlipbookHelpers.h"
-#include "PaperFlipbook.h"
 #include "PaperFlipbookFactory.h"
 #include "PaperSprite.h"
+#include "ToolMenuSection.h"
 
 #define LOCTEXT_NAMESPACE "AssetTypeActions"
 
diff --git a/Engine/Plugins/2D/Paper2D/Source/Paper2DEditor/Private/SpriteEditor/SSpriteEditorViewportToolbar.cpp b/Engine/Plugins/2D/Paper2D/Source/Paper2DEditor/Private/SpriteEditor/SSpriteEditorViewportToolbar.cpp
--- a/Engine/Plugins/2D/Paper2D/Source/Paper2DEditor/Private/SpriteEditor/SSpriteEditorViewportToolbar.cpp
+++ b/Engine/Plugins/2D/Paper2D/Source/Paper2DEditor/Private/SpriteEditor/SSpriteEditorViewportToolbar.cpp
@@ -1,6 +1,7 @@
 // Copyright Epic Games, Inc. All Rights Reserved.
 
 #include "SpriteEditor/SSpriteEditorViewportToolbar.h"
+#include "Framework/Commands/UICommandList.h"
 #include "SpriteEditor/SpriteEditorCommands.h"
 #include "PaperEditorShared/SpriteGeometryEditCommands.h"
 #include "SEditorViewport.h"
diff --git a/Engine/Plugins/2D/Paper2D/Source/Paper2DEditor/Private/SpriteEditor/SSpriteEditorViewportToolbar.h b/Engine/Plugins/2D/Paper2D/Source/Paper2DEditor/Private/SpriteEditor/SSpriteEditorViewportToolbar.h
--- a/Engine/Plugins/2D/Paper2D/Source/Paper2DEditor/Private/SpriteEditor/SSpriteEditorViewportToolbar.h
+++ b/Engine/Plugins/2D/Paper2D/Source/Paper2DEditor/Private/SpriteEditor/SSpriteEditorViewportToolbar.h
@@ -2,9 +2,6 @@
 
 #pragma once
 
-#include "CoreMinimal.h"
-#include "Widgets/DeclarativeSyntaxSupport.h"
-#include "Widgets/SWidget.h"
 #include "SCommonEditorViewportToolbarBase.h"
 
 // In-viewport toolbar widget used in the sprite editor
diff --git a/Engine/Plugins/2D/Paper2D/Source/Paper2DEditor/Private/SpriteEditor/SSpriteList.cpp b/Engine/Plugins/2D/Paper2D/Source/Paper2DEditor/Private/SpriteEditor/SSpriteList.cpp
--- a/Engine/Plugins/2D/Paper2D/Source/Paper2DEditor/Private/SpriteEditor/SSpriteList.cpp
+++ b/Engine/Plugins/2D/Paper2D/Source/Paper2DEditor/Private/SpriteEditor/SSpriteList.cpp
@@ -4,6 +4,7 @@
 #include "IContentBrowserSingleton.h"
 #include "ContentBrowserModule.h"
 #include "PaperSprite.h"
+#include "SpriteEditor/SpriteEditor.h"
 
 //////////////////////////////////////////////////////////////////////////
 // SSpriteList
diff --git a/Engine/Plugins/2D/Paper2D/Source/Paper2DEditor/Private/SpriteEditor/SSpriteList.h b/Engine/Plugins/2D/Paper2D/Source/Paper2DEditor/Private/SpriteEditor/SSpriteList.h
--- a/Engine/Plugins/2D/Paper2D/Source/Paper2DEditor/Private/SpriteEditor/SSpriteList.h
+++ b/Engine/Plugins/2D/Paper2D/Source/Paper2DEditor/Private/SpriteEditor/SSpriteList.h
@@ -2,10 +2,7 @@
 
 #pragma once
 
-#include "CoreMinimal.h"
-#include "Widgets/DeclarativeSyntaxSupport.h"
 #include "Widgets/SCompoundWidget.h"
-#include "SpriteEditor/SpriteEditor.h"
 #include "ContentBrowserDelegates.h"
 
 class UTexture2D;
diff --git a/Engine/Plugins/2D/Paper2D/Source/Paper2DEditor/Private/SpriteEditor/SpriteComponentDetailsCustomization.cpp b/Engine/Plugins/2D/Paper2D/Source/Paper2DEditor/Private/SpriteEditor/SpriteComponentDetailsCustomization.cpp
--- a/Engine/Plugins/2D/Paper2D/Source/Paper2DEditor/Private/SpriteEditor/SpriteComponentDetailsCustomization.cpp
+++ b/Engine/Plugins/2D/Paper2D/Source/Paper2DEditor/Private/SpriteEditor/SpriteComponentDetailsCustomization.cpp
@@ -1,12 +1,10 @@
 // Copyright Epic Games, Inc. All Rights Reserved.
 
 #include "SpriteEditor/SpriteComponentDetailsCustomization.h"
-#include "Widgets/DeclarativeSyntaxSupport.h"
 #include "Widgets/SBoxPanel.h"
 #include "Widgets/Input/SButton.h"
 #include "DetailLayoutBuilder.h"
 #include "DetailWidgetRow.h"
-#include "DetailCategoryBuilder.h"
 #include "GroupedSprites/PaperGroupedSpriteUtilities.h"
 
 #define LOCTEXT_NAMESPACE "SpriteEditor"
diff --git a/Engine/Plugins/2D/Paper2D/Source/Paper2DEditor/Private/SpriteEditor/SpriteComponentDetailsCustomization.h b/Engine/Plugins/2D/Paper2D/Source/Paper2DEditor/Private/SpriteEditor/SpriteComponentDetailsCustomization.h
--- a/Engine/Plugins/2D/Paper2D/Source/Paper2DEditor/Private/SpriteEditor/SpriteComponentDetailsCustomization.h
+++ b/Engine/Plugins/2D/Paper2D/Source/Paper2DEditor/Private/SpriteEditor/SpriteComponentDetailsCustomization.h
@@ -2,9 +2,10 @@
 
 #pragma once
 
-#include "CoreMinimal.h"
-#include "Input/Reply.h"
 #include "IDetailCustomization.h"
+#include "UObject/WeakObjectPtrTemplates.h"
+
+class FReply;
 
 class IDetailLayoutBuilder;
 
diff --git a/Engine/Plugins/2D/Paper2D/Source/Paper2DEditor/Private/SpriteEditor/SpriteDetailsCustomization.cpp b/Engine/Plugins/2D/Paper2D/Source/Paper2DEditor/Private/SpriteEditor/SpriteDetailsCustomization.cpp
--- a/Engine/Plugins/2D/Paper2D/Source/Paper2DEditor/Private/SpriteEditor/SpriteDetailsCustomization.cpp
+++ b/Engine/Plugins/2D/Paper2D/Source/Paper2DEditor/Private/SpriteEditor/SpriteDetailsCustomization.cpp
@@ -1,23 +1,17 @@
 // Copyright Epic Games, Inc. All Rights Reserved.
 
 #include "SpriteEditor/SpriteDetailsCustomization.h"
-#include "Materials/MaterialInterface.h"
 #include "Widgets/Text/STextBlock.h"
 #include "Widgets/Images/SImage.h"
 #include "Widgets/Layout/SBox.h"
-#include "Styling/AppStyle.h"
 #include "Materials/Material.h"
-#include "PropertyHandle.h"
 #include "DetailLayoutBuilder.h"
 #include "IDetailChildrenBuilder.h"
-#include "IDetailPropertyRow.h"
-#include "DetailCategoryBuilder.h"
 #include "SpriteEditor/SpriteEditorViewportClient.h"
 
 #include "PhysicsEngine/BodySetup.h"
 #include "Widgets/SToolTip.h"
 #include "IDocumentation.h"
-#include "PhysicsEngine/PhysicsSettings.h"
 #include "PropertyRestriction.h"
 #include "PropertyCustomizationHelpers.h"
 #include "MaterialExpressionSpriteTextureSampler.h"
diff --git a/Engine/Plugins/2D/Paper2D/Source/Paper2DEditor/Private/SpriteEditor/SpriteDetailsCustomization.h b/Engine/Plugins/2D/Paper2D/Source/Paper2DEditor/Private/SpriteEditor/SpriteDetailsCustomization.h
--- a/Engine/Plugins/2D/Paper2D/Source/Paper2DEditor/Private/SpriteEditor/SpriteDetailsCustomization.h
+++ b/Engine/Plugins/2D/Paper2D/Source/Paper2DEditor/Private/SpriteEditor/SpriteDetailsCustomization.h
@@ -2,14 +2,12 @@
 
 #pragma once
 
-#include "CoreMinimal.h"
-#include "Misc/Attribute.h"
-#include "Layout/Visibility.h"
-#include "Widgets/SWidget.h"
-#include "PaperSprite.h"
 #include "IDetailCustomization.h"
 #include "SpriteEditor/SpriteEditor.h"
 
+namespace ESpriteCollisionMode { enum Type : int; }
+namespace ESpritePolygonMode { enum Type : int; }
+
 class FDetailWidgetRow;
 class IDetailCategoryBuilder;
 class IDetailChildrenBuilder;
diff --git a/Engine/Plugins/2D/Paper2D/Source/Paper2DEditor/Private/SpriteEditor/SpriteEditor.cpp b/Engine/Plugins/2D/Paper2D/Source/Paper2DEditor/Private/SpriteEditor/SpriteEditor.cpp
--- a/Engine/Plugins/2D/Paper2D/Source/Paper2DEditor/Private/SpriteEditor/SpriteEditor.cpp
+++ b/Engine/Plugins/2D/Paper2D/Source/Paper2DEditor/Private/SpriteEditor/SpriteEditor.cpp
@@ -1,22 +1,16 @@
 // Copyright Epic Games, Inc. All Rights Reserved.
 
 #include "SpriteEditor/SpriteEditor.h"
-#include "Framework/MultiBox/MultiBoxDefs.h"
+#include "Framework/Commands/UICommandList.h"
 #include "Widgets/Text/STextBlock.h"
-#include "Framework/MultiBox/MultiBoxBuilder.h"
-#include "EditorViewportClient.h"
-#include "Modules/ModuleManager.h"
-#include "Styling/AppStyle.h"
+#include "IDetailsView.h"
 #include "SSingleObjectDetailsPanel.h"
 
 
-#include "SEditorViewport.h"
-#include "PaperSprite.h"
 #include "SpriteEditor/SpriteEditorViewportClient.h"
 #include "SpriteEditor/SpriteEditorCommands.h"
 #include "PaperEditorShared/SpriteGeometryEditCommands.h"
 #include "Paper2DEditorModule.h"
-#include "SCommonEditorViewportToolbarBase.h"
 #include "SpriteEditor/SSpriteEditorViewportToolbar.h"
 #include "SpriteEditor/SpriteDetailsCustomization.h"
 
@@ -24,7 +18,6 @@
 #include "Widgets/Docking/SDockTab.h"
 
 #include "ExtractSprites/SPaperExtractSpritesDialog.h"
-#include "EditorModeManager.h"
 
 #define LOCTEXT_NAMESPACE "SpriteEditor"
 
diff --git a/Engine/Plugins/2D/Paper2D/Source/Paper2DEditor/Private/SpriteEditor/SpriteEditor.h b/Engine/Plugins/2D/Paper2D/Source/Paper2DEditor/Private/SpriteEditor/SpriteEditor.h
--- a/Engine/Plugins/2D/Paper2D/Source/Paper2DEditor/Private/SpriteEditor/SpriteEditor.h
+++ b/Engine/Plugins/2D/Paper2D/Source/Paper2DEditor/Private/SpriteEditor/SpriteEditor.h
@@ -2,11 +2,10 @@
 
 #pragma once
 
-#include "CoreMinimal.h"
-#include "UObject/GCObject.h"
-#include "Toolkits/IToolkitHost.h"
 #include "Toolkits/AssetEditorToolkit.h"
 
+class FSpawnTabArgs;
+
 class FToolBarBuilder;
 class SSpriteEditorViewport;
 class SSpriteList;
diff --git a/Engine/Plugins/2D/Paper2D/Source/Paper2DEditor/Private/SpriteEditor/SpriteEditorSelections.h b/Engine/Plugins/2D/Paper2D/Source/Paper2DEditor/Private/SpriteEditor/SpriteEditorSelections.h
--- a/Engine/Plugins/2D/Paper2D/Source/Paper2DEditor/Private/SpriteEditor/SpriteEditorSelections.h
+++ b/Engine/Plugins/2D/Paper2D/Source/Paper2DEditor/Private/SpriteEditor/SpriteEditorSelections.h
@@ -2,9 +2,6 @@
 
 #pragma once
 
-#include "CoreMinimal.h"
-#include "UnrealWidgetFwd.h"
-#include "SpriteEditorOnlyTypes.h"
 #include "Paper2DModule.h"
 #include "PaperEditorShared/AssetEditorSelectedItem.h"
 #include "PaperSprite.h"
diff --git a/Engine/Plugins/2D/Paper2D/Source/Paper2DEditor/Private/SpriteEditor/SpriteEditorSettings.h b/Engine/Plugins/2D/Paper2D/Source/Paper2DEditor/Private/SpriteEditor/SpriteEditorSettings.h
--- a/Engine/Plugins/2D/Paper2D/Source/Paper2DEditor/Private/SpriteEditor/SpriteEditorSettings.h
+++ b/Engine/Plugins/2D/Paper2D/Source/Paper2DEditor/Private/SpriteEditor/SpriteEditorSettings.h
@@ -2,9 +2,6 @@
 
 #pragma once
 
-#include "CoreMinimal.h"
-#include "UObject/ObjectMacros.h"
-#include "UObject/Object.h"
 #include "SpriteEditorSettings.generated.h"
 
 // Settings for the Paper2D sprite editor
diff --git a/Engine/Plugins/2D/Paper2D/Source/Paper2DEditor/Private/SpriteEditor/SpriteEditorViewportClient.cpp b/Engine/Plugins/2D/Paper2D/Source/Paper2DEditor/Private/SpriteEditor/SpriteEditorViewportClient.cpp
--- a/Engine/Plugins/2D/Paper2D/Source/Paper2DEditor/Private/SpriteEditor/SpriteEditorViewportClient.cpp
+++ b/Engine/Plugins/2D/Paper2D/Source/Paper2DEditor/Private/SpriteEditor/SpriteEditorViewportClient.cpp
@@ -1,12 +1,14 @@
 // Copyright Epic Games, Inc. All Rights Reserved.
 
 #include "SpriteEditor/SpriteEditorViewportClient.h"
+#include "InputKeyEventArgs.h"
 #include "Materials/MaterialInterface.h"
-#include "Modules/ModuleManager.h"
 #include "CanvasItem.h"
+#include "Misc/PackageName.h"
 #include "Utils.h"
 #include "PaperSpriteComponent.h"
 #include "AssetEditorModeManager.h"
+#include "SEditorViewport.h"
 #include "SceneView.h"
 #include "ScopedTransaction.h"
 #include "GlobalRenderResources.h"
@@ -14,7 +16,6 @@
 
 #include "IAssetTools.h"
 #include "AssetToolsModule.h"
-#include "AssetRegistry/ARFilter.h"
 #include "IContentBrowserSingleton.h"
 #include "ContentBrowserModule.h"
 #include "AssetRegistry/AssetRegistryModule.h"
diff --git a/Engine/Plugins/2D/Paper2D/Source/Paper2DEditor/Private/SpriteEditor/SpriteEditorViewportClient.h b/Engine/Plugins/2D/Paper2D/Source/Paper2DEditor/Private/SpriteEditor/SpriteEditorViewportClient.h
--- a/Engine/Plugins/2D/Paper2D/Source/Paper2DEditor/Private/SpriteEditor/SpriteEditorViewportClient.h
+++ b/Engine/Plugins/2D/Paper2D/Source/Paper2DEditor/Private/SpriteEditor/SpriteEditorViewportClient.h
@@ -2,13 +2,9 @@
 
 #pragma once
 
-#include "CoreMinimal.h"
-#include "InputCoreTypes.h"
 #include "PreviewScene.h"
-#include "AssetRegistry/AssetData.h"
 #include "SpriteEditor/SpriteEditor.h"
 #include "SpriteEditor/SpriteEditorSelections.h"
-#include "SEditorViewport.h"
 #include "PaperEditorViewportClient.h"
 
 class FCanvas;
diff --git a/Engine/Plugins/2D/Paper2D/Source/Paper2DEditor/Private/SpriteEditor/SpritePolygonCollectionCustomization.cpp b/Engine/Plugins/2D/Paper2D/Source/Paper2DEditor/Private/SpriteEditor/SpritePolygonCollectionCustomization.cpp
--- a/Engine/Plugins/2D/Paper2D/Source/Paper2DEditor/Private/SpriteEditor/SpritePolygonCollectionCustomization.cpp
+++ b/Engine/Plugins/2D/Paper2D/Source/Paper2DEditor/Private/SpriteEditor/SpritePolygonCollectionCustomization.cpp
@@ -1,6 +1,7 @@
 // Copyright Epic Games, Inc. All Rights Reserved.
 
 #include "SpriteEditor/SpritePolygonCollectionCustomization.h"
+#include "PropertyHandle.h"
 
 //////////////////////////////////////////////////////////////////////////
 // FSpritePolygonCollectionCustomization
diff --git a/Engine/Plugins/2D/Paper2D/Source/Paper2DEditor/Private/SpriteEditor/SpritePolygonCollectionCustomization.h b/Engine/Plugins/2D/Paper2D/Source/Paper2DEditor/Private/SpriteEditor/SpritePolygonCollectionCustomization.h
--- a/Engine/Plugins/2D/Paper2D/Source/Paper2DEditor/Private/SpriteEditor/SpritePolygonCollectionCustomization.h
+++ b/Engine/Plugins/2D/Paper2D/Source/Paper2DEditor/Private/SpriteEditor/SpritePolygonCollectionCustomization.h
@@ -2,11 +2,10 @@
 
 #pragma once
 
-#include "CoreMinimal.h"
-#include "Layout/Visibility.h"
 #include "IPropertyTypeCustomization.h"
-#include "PropertyHandle.h"
-#include "SpriteEditorOnlyTypes.h"
+
+namespace ESpritePolygonMode { enum Type : int; }
+struct EVisibility;
 
 //////////////////////////////////////////////////////////////////////////
 // FSpritePolygonCollectionCustomization
diff --git a/Engine/Plugins/2D/Paper2D/Source/Paper2DEditor/Private/TileMapEditing/EdModeTileMap.cpp b/Engine/Plugins/2D/Paper2D/Source/Paper2DEditor/Private/TileMapEditing/EdModeTileMap.cpp
--- a/Engine/Plugins/2D/Paper2D/Source/Paper2DEditor/Private/TileMapEditing/EdModeTileMap.cpp
+++ b/Engine/Plugins/2D/Paper2D/Source/Paper2DEditor/Private/TileMapEditing/EdModeTileMap.cpp
@@ -1,14 +1,16 @@
 // Copyright Epic Games, Inc. All Rights Reserved.
 
 #include "TileMapEditing/EdModeTileMap.h"
+#include "EditorViewportClient.h"
 #include "NavigationSystem.h"
 #include "Framework/Application/SlateApplication.h"
-#include "Engine/Selection.h"
 #include "EditorModeManager.h"
 #include "Paper2DModule.h"
 #include "PaperTileLayer.h"
 #include "PaperTileMapComponent.h"
+#include "PaperTileSet.h"
 #include "TileMapEditing/TileMapEdModeToolkit.h"
+#include "Selection.h"
 #include "Toolkits/ToolkitManager.h"
 #include "SceneView.h"
 #include "ScopedTransaction.h"
diff --git a/Engine/Plugins/2D/Paper2D/Source/Paper2DEditor/Private/TileMapEditing/EdModeTileMap.h b/Engine/Plugins/2D/Paper2D/Source/Paper2DEditor/Private/TileMapEditing/EdModeTileMap.h
--- a/Engine/Plugins/2D/Paper2D/Source/Paper2DEditor/Private/TileMapEditing/EdModeTileMap.h
+++ b/Engine/Plugins/2D/Paper2D/Source/Paper2DEditor/Private/TileMapEditing/EdModeTileMap.h
@@ -2,12 +2,11 @@
 
 #pragma once
 
-#include "CoreMinimal.h"
-#include "InputCoreTypes.h"
-#include "Framework/Commands/UICommandList.h"
-#include "Framework/MultiBox/MultiBoxExtender.h"
-#include "EditorViewportClient.h"
 #include "EdMode.h"
+#include "Math/IntRect.h"
+
+class FExtender;
+struct FViewportCursorLocation;
 
 class FCanvas;
 class UPaperTileLayer;
diff --git a/Engine/Plugins/2D/Paper2D/Source/Paper2DEditor/Private/TileMapEditing/PaperTileMapDetailsCustomization.cpp b/Engine/Plugins/2D/Paper2D/Source/Paper2DEditor/Private/TileMapEditing/PaperTileMapDetailsCustomization.cpp
--- a/Engine/Plugins/2D/Paper2D/Source/Paper2DEditor/Private/TileMapEditing/PaperTileMapDetailsCustomization.cpp
+++ b/Engine/Plugins/2D/Paper2D/Source/Paper2DEditor/Private/TileMapEditing/PaperTileMapDetailsCustomization.cpp
@@ -1,22 +1,20 @@
 // Copyright Epic Games, Inc. All Rights Reserved.
 
 #include "TileMapEditing/PaperTileMapDetailsCustomization.h"
-#include "Layout/Margin.h"
-#include "Widgets/DeclarativeSyntaxSupport.h"
+#include "DetailsViewArgs.h"
 #include "Widgets/SBoxPanel.h"
 #include "Framework/Commands/UICommandList.h"
+#include "PaperTileLayer.h"
 #include "Widgets/Layout/SWrapBox.h"
+#include "PaperTileMap.h"
 #include "Widgets/Text/STextBlock.h"
+#include "PaperTileMapComponent.h"
 #include "Widgets/Layout/SBox.h"
 #include "Widgets/Input/SButton.h"
-#include "Styling/AppStyle.h"
 #include "EditorModeManager.h"
 
-#include "PropertyHandle.h"
 #include "DetailLayoutBuilder.h"
 #include "DetailWidgetRow.h"
-#include "IDetailPropertyRow.h"
-#include "DetailCategoryBuilder.h"
 #include "PaperTileSet.h"
 #include "TileMapEditing/EdModeTileMap.h"
 
diff --git a/Engine/Plugins/2D/Paper2D/Source/Paper2DEditor/Private/TileMapEditing/PaperTileMapDetailsCustomization.h b/Engine/Plugins/2D/Paper2D/Source/Paper2DEditor/Private/TileMapEditing/PaperTileMapDetailsCustomization.h
--- a/Engine/Plugins/2D/Paper2D/Source/Paper2DEditor/Private/TileMapEditing/PaperTileMapDetailsCustomization.h
+++ b/Engine/Plugins/2D/Paper2D/Source/Paper2DEditor/Private/TileMapEditing/PaperTileMapDetailsCustomization.h
@@ -2,12 +2,11 @@
 
 #pragma once
 
-#include "CoreMinimal.h"
-#include "Layout/Visibility.h"
-#include "Input/Reply.h"
-#include "PaperTileMap.h"
 #include "IDetailCustomization.h"
-#include "PaperTileMapComponent.h"
+#include "UObject/WeakObjectPtr.h"
+
+class FReply;
+struct EVisibility;
 
 class IDetailLayoutBuilder;
 
diff --git a/Engine/Plugins/2D/Paper2D/Source/Paper2DEditor/Private/TileMapEditing/STileLayerItem.cpp b/Engine/Plugins/2D/Paper2D/Source/Paper2DEditor/Private/TileMapEditing/STileLayerItem.cpp
--- a/Engine/Plugins/2D/Paper2D/Source/Paper2DEditor/Private/TileMapEditing/STileLayerItem.cpp
+++ b/Engine/Plugins/2D/Paper2D/Source/Paper2DEditor/Private/TileMapEditing/STileLayerItem.cpp
@@ -4,7 +4,6 @@
 #include "Widgets/SBoxPanel.h"
 #include "Widgets/Images/SImage.h"
 #include "Widgets/Input/SButton.h"
-#include "Styling/AppStyle.h"
 #include "PaperTileLayer.h"
 #include "ScopedTransaction.h"
 #include "Widgets/Text/SInlineEditableTextBlock.h"
diff --git a/Engine/Plugins/2D/Paper2D/Source/Paper2DEditor/Private/TileMapEditing/STileLayerItem.h b/Engine/Plugins/2D/Paper2D/Source/Paper2DEditor/Private/TileMapEditing/STileLayerItem.h
--- a/Engine/Plugins/2D/Paper2D/Source/Paper2DEditor/Private/TileMapEditing/STileLayerItem.h
+++ b/Engine/Plugins/2D/Paper2D/Source/Paper2DEditor/Private/TileMapEditing/STileLayerItem.h
@@ -2,15 +2,12 @@
 
 #pragma once
 
-#include "CoreMinimal.h"
-#include "SlateFwd.h"
-#include "Styling/SlateColor.h"
-#include "Widgets/DeclarativeSyntaxSupport.h"
-#include "Input/Reply.h"
 #include "Widgets/SCompoundWidget.h"
 #include "Framework/SlateDelegates.h"
 #include "PaperTileMap.h"
 
+class SInlineEditableTextBlock;
+
 class SButton;
 class UPaperTileLayer;
 struct FSlateBrush;
diff --git a/Engine/Plugins/2D/Paper2D/Source/Paper2DEditor/Private/TileMapEditing/STileLayerList.cpp b/Engine/Plugins/2D/Paper2D/Source/Paper2DEditor/Private/TileMapEditing/STileLayerList.cpp
--- a/Engine/Plugins/2D/Paper2D/Source/Paper2DEditor/Private/TileMapEditing/STileLayerList.cpp
+++ b/Engine/Plugins/2D/Paper2D/Source/Paper2DEditor/Private/TileMapEditing/STileLayerList.cpp
@@ -1,17 +1,16 @@
 // Copyright Epic Games, Inc. All Rights Reserved.
 
 #include "TileMapEditing/STileLayerList.h"
+#include "Framework/Commands/UICommandList.h"
 #include "PaperTileLayer.h"
-#include "UObject/PropertyPortFlags.h"
+#include "Framework/Views/TableViewMetadata.h"
 #include "Misc/NotifyHook.h"
-#include "Framework/MultiBox/MultiBoxDefs.h"
 #include "Framework/MultiBox/MultiBoxBuilder.h"
 #include "Widgets/Views/SListView.h"
 #include "Exporters/Exporter.h"
 #include "Editor.h"
 #include "PaperTileMapComponent.h"
 #include "TileMapEditing/STileLayerItem.h"
-#include "PaperStyle.h"
 #include "HAL/PlatformApplicationMisc.h"
 
 #include "ScopedTransaction.h"
diff --git a/Engine/Plugins/2D/Paper2D/Source/Paper2DEditor/Private/TileMapEditing/STileLayerList.h b/Engine/Plugins/2D/Paper2D/Source/Paper2DEditor/Private/TileMapEditing/STileLayerList.h
--- a/Engine/Plugins/2D/Paper2D/Source/Paper2DEditor/Private/TileMapEditing/STileLayerList.h
+++ b/Engine/Plugins/2D/Paper2D/Source/Paper2DEditor/Private/TileMapEditing/STileLayerList.h
@@ -2,14 +2,13 @@
 
 #pragma once
 
-#include "CoreMinimal.h"
-#include "PaperTileMap.h"
-#include "Widgets/DeclarativeSyntaxSupport.h"
-#include "Framework/Commands/UICommandList.h"
 #include "Widgets/SCompoundWidget.h"
 #include "EditorUndoClient.h"
-#include "Widgets/Views/STableViewBase.h"
-#include "Widgets/Views/STableRow.h"
+
+class ITableRow;
+class STableViewBase;
+class UPaperTileMap;
+namespace ESelectInfo { enum Type : int; }
 
 class FNotifyHook;
 class UPaperTileLayer;
