[Paper2d]

diff --git a/Engine/Plugins/2D/Paper2D/Source/Paper2DEditor/Private/TileMapEditing/STileMapEditorViewportToolbar.h b/Engine/Plugins/2D/Paper2D/Source/Paper2DEditor/Private/TileMapEditing/STileMapEditorViewportToolbar.h
--- a/Engine/Plugins/2D/Paper2D/Source/Paper2DEditor/Private/TileMapEditing/STileMapEditorViewportToolbar.h
+++ b/Engine/Plugins/2D/Paper2D/Source/Paper2DEditor/Private/TileMapEditing/STileMapEditorViewportToolbar.h
@@ -2,9 +2,6 @@
 
 #pragma once
 
-#include "CoreMinimal.h"
-#include "Widgets/DeclarativeSyntaxSupport.h"
-#include "Widgets/SWidget.h"
 #include "SCommonEditorViewportToolbarBase.h"
 
 // In-viewport toolbar widget used in the tile map editor
diff --git a/Engine/Plugins/2D/Paper2D/Source/Paper2DEditor/Private/TileMapEditing/TileMapAssetTypeActions.cpp b/Engine/Plugins/2D/Paper2D/Source/Paper2DEditor/Private/TileMapEditing/TileMapAssetTypeActions.cpp
--- a/Engine/Plugins/2D/Paper2D/Source/Paper2DEditor/Private/TileMapEditing/TileMapAssetTypeActions.cpp
+++ b/Engine/Plugins/2D/Paper2D/Source/Paper2DEditor/Private/TileMapEditing/TileMapAssetTypeActions.cpp
@@ -2,6 +2,7 @@
 
 #include "TileMapEditing/TileMapAssetTypeActions.h"
 
+#include "Math/ColorList.h"
 #include "TileMapEditing/TileMapEditor.h"
 #include "PaperTileMap.h"
 
diff --git a/Engine/Plugins/2D/Paper2D/Source/Paper2DEditor/Private/TileMapEditing/TileMapEdModeToolkit.cpp b/Engine/Plugins/2D/Paper2D/Source/Paper2DEditor/Private/TileMapEditing/TileMapEdModeToolkit.cpp
--- a/Engine/Plugins/2D/Paper2D/Source/Paper2DEditor/Private/TileMapEditing/TileMapEdModeToolkit.cpp
+++ b/Engine/Plugins/2D/Paper2D/Source/Paper2DEditor/Private/TileMapEditing/TileMapEdModeToolkit.cpp
@@ -1,20 +1,18 @@
 // Copyright Epic Games, Inc. All Rights Reserved.
 
 #include "TileMapEditing/TileMapEdModeToolkit.h"
-#include "Framework/MultiBox/MultiBoxDefs.h"
+#include "Framework/Commands/UICommandList.h"
 #include "Widgets/Text/STextBlock.h"
 #include "Framework/MultiBox/MultiBoxBuilder.h"
 #include "Framework/Application/SlateApplication.h"
-#include "Widgets/Layout/SBorder.h"
+#include "PaperTileSet.h"
 #include "Widgets/Input/SButton.h"
-#include "Styling/AppStyle.h"
 #include "PaperTileMapComponent.h"
 #include "PaperTileMap.h"
 
 #include "Toolkits/AssetEditorToolkit.h"
 #include "TileSetEditor/TileSetSelectorViewport.h"
 #include "WorkflowOrientedApp/SContentReference.h"
-#include "PaperStyle.h"
 #include "TileMapEditing/TileMapEditorCommands.h"
 #include "SAssetDropTarget.h"
 
diff --git a/Engine/Plugins/2D/Paper2D/Source/Paper2DEditor/Private/TileMapEditing/TileMapEdModeToolkit.h b/Engine/Plugins/2D/Paper2D/Source/Paper2DEditor/Private/TileMapEditing/TileMapEdModeToolkit.h
--- a/Engine/Plugins/2D/Paper2D/Source/Paper2DEditor/Private/TileMapEditing/TileMapEdModeToolkit.h
+++ b/Engine/Plugins/2D/Paper2D/Source/Paper2DEditor/Private/TileMapEditing/TileMapEdModeToolkit.h
@@ -2,13 +2,7 @@
 
 #pragma once
 
-#include "CoreMinimal.h"
-#include "Layout/Visibility.h"
-#include "Input/Reply.h"
-#include "Widgets/SWidget.h"
 #include "TileMapEditing/EdModeTileMap.h"
-#include "PaperTileSet.h"
-#include "Toolkits/IToolkitHost.h"
 #include "Toolkits/BaseToolkit.h"
 
 class SContentReference;
diff --git a/Engine/Plugins/2D/Paper2D/Source/Paper2DEditor/Private/TileMapEditing/TileMapEditor.cpp b/Engine/Plugins/2D/Paper2D/Source/Paper2DEditor/Private/TileMapEditing/TileMapEditor.cpp
--- a/Engine/Plugins/2D/Paper2D/Source/Paper2DEditor/Private/TileMapEditing/TileMapEditor.cpp
+++ b/Engine/Plugins/2D/Paper2D/Source/Paper2DEditor/Private/TileMapEditing/TileMapEditor.cpp
@@ -1,10 +1,8 @@
 // Copyright Epic Games, Inc. All Rights Reserved.
 
 #include "TileMapEditing/TileMapEditor.h"
+#include "Framework/Commands/UICommandList.h"
 #include "Widgets/Text/STextBlock.h"
-#include "EditorViewportClient.h"
-#include "UObject/Package.h"
-#include "Styling/AppStyle.h"
 #include "SEditorViewport.h"
 #include "SSingleObjectDetailsPanel.h"
 #include "PaperTileMap.h"
@@ -12,11 +10,8 @@
 
 #include "TileMapEditing/TileMapEditorViewportClient.h"
 #include "TileMapEditing/TileMapEditorCommands.h"
-#include "SCommonEditorViewportToolbarBase.h"
 #include "TileMapEditing/STileMapEditorViewportToolbar.h"
 #include "Widgets/Docking/SDockTab.h"
-#include "Subsystems/AssetEditorSubsystem.h"
-#include "EditorModeManager.h"
 
 #define LOCTEXT_NAMESPACE "TileMapEditor"
 
diff --git a/Engine/Plugins/2D/Paper2D/Source/Paper2DEditor/Private/TileMapEditing/TileMapEditor.h b/Engine/Plugins/2D/Paper2D/Source/Paper2DEditor/Private/TileMapEditing/TileMapEditor.h
--- a/Engine/Plugins/2D/Paper2D/Source/Paper2DEditor/Private/TileMapEditing/TileMapEditor.h
+++ b/Engine/Plugins/2D/Paper2D/Source/Paper2DEditor/Private/TileMapEditing/TileMapEditor.h
@@ -2,12 +2,10 @@
 
 #pragma once
 
-#include "CoreMinimal.h"
-#include "UObject/GCObject.h"
-#include "Widgets/Layout/SBorder.h"
-#include "Toolkits/IToolkitHost.h"
 #include "Toolkits/AssetEditorToolkit.h"
 
+class FSpawnTabArgs;
+
 class STileMapEditorViewport;
 class UPaperTileMap;
 
diff --git a/Engine/Plugins/2D/Paper2D/Source/Paper2DEditor/Private/TileMapEditing/TileMapEditorViewportClient.cpp b/Engine/Plugins/2D/Paper2D/Source/Paper2DEditor/Private/TileMapEditing/TileMapEditorViewportClient.cpp
--- a/Engine/Plugins/2D/Paper2D/Source/Paper2DEditor/Private/TileMapEditing/TileMapEditorViewportClient.cpp
+++ b/Engine/Plugins/2D/Paper2D/Source/Paper2DEditor/Private/TileMapEditing/TileMapEditorViewportClient.cpp
@@ -1,14 +1,13 @@
 // Copyright Epic Games, Inc. All Rights Reserved.
 
 #include "TileMapEditing/TileMapEditorViewportClient.h"
-#include "Components/PrimitiveComponent.h"
 #include "Materials/MaterialInterface.h"
 #include "PaperTileMapComponent.h"
 #include "SceneView.h"
 #include "ScopedTransaction.h"
 #include "CanvasItem.h"
-#include "Engine/Selection.h"
 #include "AssetEditorModeManager.h"
+#include "Selection.h"
 #include "TileMapEditing/EdModeTileMap.h"
 #include "PaperEditorShared/SpriteGeometryEditMode.h"
 #include "PaperTileMap.h"
diff --git a/Engine/Plugins/2D/Paper2D/Source/Paper2DEditor/Private/TileMapEditing/TileMapEditorViewportClient.h b/Engine/Plugins/2D/Paper2D/Source/Paper2DEditor/Private/TileMapEditing/TileMapEditorViewportClient.h
--- a/Engine/Plugins/2D/Paper2D/Source/Paper2DEditor/Private/TileMapEditing/TileMapEditorViewportClient.h
+++ b/Engine/Plugins/2D/Paper2D/Source/Paper2DEditor/Private/TileMapEditing/TileMapEditorViewportClient.h
@@ -2,8 +2,6 @@
 
 #pragma once
 
-#include "CoreMinimal.h"
-#include "UnrealWidgetFwd.h"
 #include "PreviewScene.h"
 #include "TileMapEditing/TileMapEditor.h"
 #include "PaperEditorViewportClient.h"
diff --git a/Engine/Plugins/2D/Paper2D/Source/Paper2DEditor/Private/TileSetAssetTypeActions.cpp b/Engine/Plugins/2D/Paper2D/Source/Paper2DEditor/Private/TileSetAssetTypeActions.cpp
--- a/Engine/Plugins/2D/Paper2D/Source/Paper2DEditor/Private/TileSetAssetTypeActions.cpp
+++ b/Engine/Plugins/2D/Paper2D/Source/Paper2DEditor/Private/TileSetAssetTypeActions.cpp
@@ -1,10 +1,8 @@
 // Copyright Epic Games, Inc. All Rights Reserved.
 
 #include "TileSetAssetTypeActions.h"
-#include "ToolMenus.h"
-#include "Engine/Texture.h"
+#include "Math/ColorList.h"
 #include "Misc/PackageName.h"
-#include "Styling/AppStyle.h"
 #include "TileSetEditor.h"
 #include "PaperTileSet.h"
 #include "PaperTileMap.h"
@@ -13,6 +11,7 @@
 #include "PaperTileMapFactory.h"
 #include "TileSetEditor/TileSetEditorSettings.h"
 #include "TileSetEditor/TileSheetPaddingFactory.h"
+#include "ToolMenuSection.h"
 
 #define LOCTEXT_NAMESPACE "AssetTypeActions"
 
diff --git a/Engine/Plugins/2D/Paper2D/Source/Paper2DEditor/Private/TileSetEditor.cpp b/Engine/Plugins/2D/Paper2D/Source/Paper2DEditor/Private/TileSetEditor.cpp
--- a/Engine/Plugins/2D/Paper2D/Source/Paper2DEditor/Private/TileSetEditor.cpp
+++ b/Engine/Plugins/2D/Paper2D/Source/Paper2DEditor/Private/TileSetEditor.cpp
@@ -1,8 +1,7 @@
 // Copyright Epic Games, Inc. All Rights Reserved.
 
 #include "TileSetEditor.h"
-#include "UObject/Package.h"
-#include "Styling/AppStyle.h"
+#include "Framework/Commands/UICommandList.h"
 #include "PaperTileSet.h"
 #include "SSingleObjectDetailsPanel.h"
 #include "Widgets/Docking/SDockTab.h"
@@ -10,12 +9,9 @@
 #include "TileSetEditor/SingleTileEditorViewport.h"
 #include "TileSetEditor/SingleTileEditorViewportClient.h"
 #include "TileSetEditor/TileSetEditorCommands.h"
-#include "IDetailCustomization.h"
 #include "TileSetEditor/TileSetDetailsCustomization.h"
 #include "PaperEditorShared/SpriteGeometryEditCommands.h"
 #include "IDetailsView.h"
-#include "Subsystems/AssetEditorSubsystem.h"
-#include "EditorModeManager.h"
 
 #define LOCTEXT_NAMESPACE "TileSetEditor"
 
diff --git a/Engine/Plugins/2D/Paper2D/Source/Paper2DEditor/Private/TileSetEditor.h b/Engine/Plugins/2D/Paper2D/Source/Paper2DEditor/Private/TileSetEditor.h
--- a/Engine/Plugins/2D/Paper2D/Source/Paper2DEditor/Private/TileSetEditor.h
+++ b/Engine/Plugins/2D/Paper2D/Source/Paper2DEditor/Private/TileSetEditor.h
@@ -2,10 +2,6 @@
 
 #pragma once
 
-#include "CoreMinimal.h"
-#include "UObject/GCObject.h"
-#include "Framework/Docking/TabManager.h"
-#include "Toolkits/IToolkitHost.h"
 #include "Toolkits/AssetEditorToolkit.h"
 
 class FSingleTileEditorViewportClient;
diff --git a/Engine/Plugins/2D/Paper2D/Source/Paper2DEditor/Private/TileSetEditor/SingleTileEditorViewport.cpp b/Engine/Plugins/2D/Paper2D/Source/Paper2DEditor/Private/TileSetEditor/SingleTileEditorViewport.cpp
--- a/Engine/Plugins/2D/Paper2D/Source/Paper2DEditor/Private/TileSetEditor/SingleTileEditorViewport.cpp
+++ b/Engine/Plugins/2D/Paper2D/Source/Paper2DEditor/Private/TileSetEditor/SingleTileEditorViewport.cpp
@@ -1,11 +1,11 @@
 // Copyright Epic Games, Inc. All Rights Reserved.
 
 #include "TileSetEditor/SingleTileEditorViewport.h"
+#include "Framework/Commands/UICommandList.h"
 #include "Widgets/SBoxPanel.h"
-#include "Widgets/SViewport.h"
+#include "Widgets/SOverlay.h"
 #include "Widgets/Text/STextBlock.h"
 #include "Widgets/Layout/SBorder.h"
-#include "PaperStyle.h"
 #include "PaperEditorShared/SpriteGeometryEditCommands.h"
 #include "TileSetEditor/SingleTileEditorViewportClient.h"
 #include "TileSetEditor/TileSetEditorCommands.h"
diff --git a/Engine/Plugins/2D/Paper2D/Source/Paper2DEditor/Private/TileSetEditor/SingleTileEditorViewport.h b/Engine/Plugins/2D/Paper2D/Source/Paper2DEditor/Private/TileSetEditor/SingleTileEditorViewport.h
--- a/Engine/Plugins/2D/Paper2D/Source/Paper2DEditor/Private/TileSetEditor/SingleTileEditorViewport.h
+++ b/Engine/Plugins/2D/Paper2D/Source/Paper2DEditor/Private/TileSetEditor/SingleTileEditorViewport.h
@@ -2,9 +2,6 @@
 
 #pragma once
 
-#include "CoreMinimal.h"
-#include "Widgets/DeclarativeSyntaxSupport.h"
-#include "EditorViewportClient.h"
 #include "SEditorViewport.h"
 #include "SCommonEditorViewportToolbarBase.h"
 
diff --git a/Engine/Plugins/2D/Paper2D/Source/Paper2DEditor/Private/TileSetEditor/SingleTileEditorViewportClient.h b/Engine/Plugins/2D/Paper2D/Source/Paper2DEditor/Private/TileSetEditor/SingleTileEditorViewportClient.h
--- a/Engine/Plugins/2D/Paper2D/Source/Paper2DEditor/Private/TileSetEditor/SingleTileEditorViewportClient.h
+++ b/Engine/Plugins/2D/Paper2D/Source/Paper2DEditor/Private/TileSetEditor/SingleTileEditorViewportClient.h
@@ -2,11 +2,12 @@
 
 #pragma once
 
-#include "CoreMinimal.h"
 #include "PreviewScene.h"
 #include "PaperEditorViewportClient.h"
 #include "SpriteEditor/SpriteEditorSelections.h"
 
+struct FInputEventState;
+
 class FCanvas;
 class FScopedTransaction;
 class FUICommandList;
diff --git a/Engine/Plugins/2D/Paper2D/Source/Paper2DEditor/Private/TileSetEditor/TileSetDetailsCustomization.cpp b/Engine/Plugins/2D/Paper2D/Source/Paper2DEditor/Private/TileSetEditor/TileSetDetailsCustomization.cpp
--- a/Engine/Plugins/2D/Paper2D/Source/Paper2DEditor/Private/TileSetEditor/TileSetDetailsCustomization.cpp
+++ b/Engine/Plugins/2D/Paper2D/Source/Paper2DEditor/Private/TileSetEditor/TileSetDetailsCustomization.cpp
@@ -1,16 +1,13 @@
 // Copyright Epic Games, Inc. All Rights Reserved.
 
 #include "TileSetEditor/TileSetDetailsCustomization.h"
-#include "Layout/Margin.h"
-#include "Widgets/DeclarativeSyntaxSupport.h"
+#include "PaperTileSet.h"
 #include "Widgets/SBoxPanel.h"
 #include "Widgets/Text/STextBlock.h"
 #include "Widgets/Layout/SBox.h"
-#include "Styling/AppStyle.h"
 #include "DetailLayoutBuilder.h"
 #include "IDetailPropertyRow.h"
 #include "DetailCategoryBuilder.h"
-#include "IDetailsView.h"
 
 #define LOCTEXT_NAMESPACE "Paper2D"
 
diff --git a/Engine/Plugins/2D/Paper2D/Source/Paper2DEditor/Private/TileSetEditor/TileSetDetailsCustomization.h b/Engine/Plugins/2D/Paper2D/Source/Paper2DEditor/Private/TileSetEditor/TileSetDetailsCustomization.h
--- a/Engine/Plugins/2D/Paper2D/Source/Paper2DEditor/Private/TileSetEditor/TileSetDetailsCustomization.h
+++ b/Engine/Plugins/2D/Paper2D/Source/Paper2DEditor/Private/TileSetEditor/TileSetDetailsCustomization.h
@@ -2,10 +2,10 @@
 
 #pragma once
 
-#include "CoreMinimal.h"
-#include "Styling/SlateColor.h"
-#include "PaperTileSet.h"
 #include "IDetailCustomization.h"
+#include "UObject/WeakObjectPtr.h"
+
+struct FSlateColor;
 
 class IDetailLayoutBuilder;
 
diff --git a/Engine/Plugins/2D/Paper2D/Source/Paper2DEditor/Private/TileSetEditor/TileSetEditorViewportClient.cpp b/Engine/Plugins/2D/Paper2D/Source/Paper2DEditor/Private/TileSetEditor/TileSetEditorViewportClient.cpp
--- a/Engine/Plugins/2D/Paper2D/Source/Paper2DEditor/Private/TileSetEditor/TileSetEditorViewportClient.cpp
+++ b/Engine/Plugins/2D/Paper2D/Source/Paper2DEditor/Private/TileSetEditor/TileSetEditorViewportClient.cpp
@@ -1,12 +1,9 @@
 // Copyright Epic Games, Inc. All Rights Reserved.
 
 #include "TileSetEditor/TileSetEditorViewportClient.h"
-#include "IntMargin.h"
 #include "PaperTileSet.h"
 #include "CanvasItem.h"
 #include "CanvasTypes.h"
-#include "Engine/Engine.h"
-#include "EngineGlobals.h"
 #include "TextureCompiler.h"
 #include "TextureResource.h"
 #include "GlobalRenderResources.h"
diff --git a/Engine/Plugins/2D/Paper2D/Source/Paper2DEditor/Private/TileSetEditor/TileSetSelectorViewport.cpp b/Engine/Plugins/2D/Paper2D/Source/Paper2DEditor/Private/TileSetEditor/TileSetSelectorViewport.cpp
--- a/Engine/Plugins/2D/Paper2D/Source/Paper2DEditor/Private/TileSetEditor/TileSetSelectorViewport.cpp
+++ b/Engine/Plugins/2D/Paper2D/Source/Paper2DEditor/Private/TileSetEditor/TileSetSelectorViewport.cpp
@@ -1,10 +1,12 @@
 // Copyright Epic Games, Inc. All Rights Reserved.
 
 #include "TileSetEditor/TileSetSelectorViewport.h"
-#include "Widgets/SViewport.h"
+#include "Framework/Commands/UICommandList.h"
 #include "TileMapEditing/EdModeTileMap.h"
+#include "PaperTileSet.h"
 #include "Slate/SceneViewport.h"
 #include "TileSetEditor/TileSetEditorCommands.h"
+#include "TileSetEditor/TileSetEditorViewportClient.h"
 
 #define LOCTEXT_NAMESPACE "TileSetEditor"
 
diff --git a/Engine/Plugins/2D/Paper2D/Source/Paper2DEditor/Private/TileSetEditor/TileSetSelectorViewport.h b/Engine/Plugins/2D/Paper2D/Source/Paper2DEditor/Private/TileSetEditor/TileSetSelectorViewport.h
--- a/Engine/Plugins/2D/Paper2D/Source/Paper2DEditor/Private/TileSetEditor/TileSetSelectorViewport.h
+++ b/Engine/Plugins/2D/Paper2D/Source/Paper2DEditor/Private/TileSetEditor/TileSetSelectorViewport.h
@@ -2,11 +2,9 @@
 
 #pragma once
 
-#include "CoreMinimal.h"
-#include "Widgets/DeclarativeSyntaxSupport.h"
-#include "PaperTileSet.h"
 #include "SPaperEditorViewport.h"
-#include "TileSetEditor/TileSetEditorViewportClient.h"
+
+class UPaperTileSet;
 
 struct FMarqueeOperation;
 
diff --git a/Engine/Plugins/2D/Paper2D/Source/Paper2DEditor/Private/TileSetEditor/TileSheetPaddingFactory.cpp b/Engine/Plugins/2D/Paper2D/Source/Paper2DEditor/Private/TileSetEditor/TileSheetPaddingFactory.cpp
--- a/Engine/Plugins/2D/Paper2D/Source/Paper2DEditor/Private/TileSetEditor/TileSheetPaddingFactory.cpp
+++ b/Engine/Plugins/2D/Paper2D/Source/Paper2DEditor/Private/TileSetEditor/TileSheetPaddingFactory.cpp
@@ -2,9 +2,6 @@
 
 #include "TileSetEditor/TileSheetPaddingFactory.h"
 #include "Paper2DEditorLog.h"
-#include "Engine/Texture.h"
-#include "Engine/Texture2D.h"
-#include "IntMargin.h"
 #include "PaperTileSet.h"
 #include "PaperSpriteAtlas.h"
 #include "Atlasing/PaperAtlasTextureHelpers.h"
diff --git a/Engine/Plugins/2D/Paper2D/Source/Paper2DEditor/Private/TileSetEditor/TileSheetPaddingFactory.h b/Engine/Plugins/2D/Paper2D/Source/Paper2DEditor/Private/TileSetEditor/TileSheetPaddingFactory.h
--- a/Engine/Plugins/2D/Paper2D/Source/Paper2DEditor/Private/TileSetEditor/TileSheetPaddingFactory.h
+++ b/Engine/Plugins/2D/Paper2D/Source/Paper2DEditor/Private/TileSetEditor/TileSheetPaddingFactory.h
@@ -2,8 +2,6 @@
 
 #pragma once
 
-#include "CoreMinimal.h"
-#include "UObject/ObjectMacros.h"
 #include "Factories/Factory.h"
 
 class UPaperTileSet;
diff --git a/Engine/Plugins/2D/Paper2D/Source/Paper2DEditor/Public/Paper2DEditorModule.h b/Engine/Plugins/2D/Paper2D/Source/Paper2DEditor/Public/Paper2DEditorModule.h
--- a/Engine/Plugins/2D/Paper2D/Source/Paper2DEditor/Public/Paper2DEditorModule.h
+++ b/Engine/Plugins/2D/Paper2D/Source/Paper2DEditor/Public/Paper2DEditorModule.h
@@ -2,9 +2,8 @@
 
 #pragma once
 
-#include "CoreMinimal.h"
-#include "Toolkits/AssetEditorToolkit.h"
 #include "Modules/ModuleInterface.h"
+#include "Templates/SharedPointer.h"
 
 #define PAPER2D_EDITOR_MODULE_NAME "Paper2DEditor"
 
@@ -23,3 +22,8 @@ public:
 	virtual uint32 GetPaper2DAssetCategory() const = 0;
 };
 
+
+#if UE_ENABLE_INCLUDE_ORDER_DEPRECATED_IN_5_2
+#include "CoreMinimal.h"
+#include "Toolkits/AssetEditorToolkit.h"
+#endif
diff --git a/Engine/Plugins/2D/Paper2D/Source/Paper2DEditor/Public/PaperFlipbookHelpers.h b/Engine/Plugins/2D/Paper2D/Source/Paper2DEditor/Public/PaperFlipbookHelpers.h
--- a/Engine/Plugins/2D/Paper2D/Source/Paper2DEditor/Public/PaperFlipbookHelpers.h
+++ b/Engine/Plugins/2D/Paper2D/Source/Paper2DEditor/Public/PaperFlipbookHelpers.h
@@ -2,7 +2,7 @@
 
 #pragma once
 
-#include "CoreMinimal.h"
+#include "Containers/Map.h"
 
 class UPaperSprite;
 
@@ -30,3 +30,7 @@ private:
 	FPaperFlipbookHelpers() {}
 };
 
+
+#if UE_ENABLE_INCLUDE_ORDER_DEPRECATED_IN_5_2
+#include "CoreMinimal.h"
+#endif
diff --git a/Engine/Plugins/2D/Paper2D/Source/Paper2DEditor/Public/PaperJSONHelpers.h b/Engine/Plugins/2D/Paper2D/Source/Paper2DEditor/Public/PaperJSONHelpers.h
--- a/Engine/Plugins/2D/Paper2D/Source/Paper2DEditor/Public/PaperJSONHelpers.h
+++ b/Engine/Plugins/2D/Paper2D/Source/Paper2DEditor/Public/PaperJSONHelpers.h
@@ -2,8 +2,9 @@
 
 #pragma once
 
-#include "CoreMinimal.h"
-#include "Dom/JsonObject.h"
+#include "Math/Vector2D.h"
+
+class FJsonValue;
 
 //////////////////////////////////////////////////////////////////////////
 // FPaperJSONHelpers
@@ -48,3 +49,8 @@ private:
 	FPaperJSONHelpers() {}
 };
 
+
+#if UE_ENABLE_INCLUDE_ORDER_DEPRECATED_IN_5_2
+#include "CoreMinimal.h"
+#include "Dom/JsonObject.h"
+#endif
diff --git a/Engine/Plugins/2D/Paper2D/Source/PaperSpriteSheetImporter/Private/PaperJsonSpriteSheetImporter.cpp b/Engine/Plugins/2D/Paper2D/Source/PaperSpriteSheetImporter/Private/PaperJsonSpriteSheetImporter.cpp
--- a/Engine/Plugins/2D/Paper2D/Source/PaperSpriteSheetImporter/Private/PaperJsonSpriteSheetImporter.cpp
+++ b/Engine/Plugins/2D/Paper2D/Source/PaperSpriteSheetImporter/Private/PaperJsonSpriteSheetImporter.cpp
@@ -4,13 +4,8 @@
 #include "PaperSpriteSheetImporterLog.h"
 #include "Misc/Paths.h"
 #include "Misc/FeedbackContext.h"
-#include "Modules/ModuleManager.h"
-#include "Serialization/JsonTypes.h"
-#include "Serialization/JsonReader.h"
 #include "Serialization/JsonSerializer.h"
 #include "EditorReimportHandler.h"
-#include "Engine/Texture2D.h"
-#include "SpriteEditorOnlyTypes.h"
 #include "PaperSprite.h"
 #include "PaperJSONHelpers.h"
 #include "IAssetTools.h"
diff --git a/Engine/Plugins/2D/Paper2D/Source/PaperSpriteSheetImporter/Private/PaperSpriteSheet.h b/Engine/Plugins/2D/Paper2D/Source/PaperSpriteSheetImporter/Private/PaperSpriteSheet.h
--- a/Engine/Plugins/2D/Paper2D/Source/PaperSpriteSheetImporter/Private/PaperSpriteSheet.h
+++ b/Engine/Plugins/2D/Paper2D/Source/PaperSpriteSheetImporter/Private/PaperSpriteSheet.h
@@ -2,12 +2,10 @@
 
 #pragma once
 
-#include "CoreMinimal.h"
-#include "UObject/ObjectMacros.h"
-#include "UObject/Object.h"
-#include "PaperSprite.h"
 #include "PaperSpriteSheet.generated.h"
 
+class UTexture2D;
+
 UCLASS(BlueprintType, meta = (DisplayThumbnail = "true"))
 class UPaperSpriteSheet : public UObject
 {
diff --git a/Engine/Plugins/2D/Paper2D/Source/PaperSpriteSheetImporter/Private/PaperSpriteSheetAssetTypeActions.cpp b/Engine/Plugins/2D/Paper2D/Source/PaperSpriteSheetImporter/Private/PaperSpriteSheetAssetTypeActions.cpp
--- a/Engine/Plugins/2D/Paper2D/Source/PaperSpriteSheetImporter/Private/PaperSpriteSheetAssetTypeActions.cpp
+++ b/Engine/Plugins/2D/Paper2D/Source/PaperSpriteSheetImporter/Private/PaperSpriteSheetAssetTypeActions.cpp
@@ -2,11 +2,8 @@
 
 #include "PaperSpriteSheetAssetTypeActions.h"
 #include "Misc/PackageName.h"
-#include "ToolMenus.h"
 #include "Misc/FeedbackContext.h"
-#include "Styling/AppStyle.h"
 #include "EditorFramework/AssetImportData.h"
-#include "PaperFlipbook.h"
 #include "PaperSprite.h"
 #include "PaperSpriteSheet.h"
 #include "IContentBrowserSingleton.h"
@@ -14,6 +11,7 @@
 #include "PaperFlipbookFactory.h"
 #include "PackageTools.h"
 #include "PaperFlipbookHelpers.h"
+#include "ToolMenuSection.h"
 
 #define LOCTEXT_NAMESPACE "AssetTypeActions"
 
diff --git a/Engine/Plugins/2D/Paper2D/Source/PaperSpriteSheetImporter/Private/PaperSpriteSheetImportFactory.cpp b/Engine/Plugins/2D/Paper2D/Source/PaperSpriteSheetImporter/Private/PaperSpriteSheetImportFactory.cpp
--- a/Engine/Plugins/2D/Paper2D/Source/PaperSpriteSheetImporter/Private/PaperSpriteSheetImportFactory.cpp
+++ b/Engine/Plugins/2D/Paper2D/Source/PaperSpriteSheetImporter/Private/PaperSpriteSheetImportFactory.cpp
@@ -3,12 +3,12 @@
 #include "PaperSpriteSheetImportFactory.h"
 #include "Misc/FileHelper.h"
 #include "Misc/Paths.h"
-#include "Modules/ModuleManager.h"
 #include "Misc/PackageName.h"
 #include "EditorFramework/AssetImportData.h"
 #include "Editor.h"
 #include "AssetToolsModule.h"
 #include "PaperSpriteSheet.h"
+#include "Subsystems/ImportSubsystem.h"
 
 #include UE_INLINE_GENERATED_CPP_BY_NAME(PaperSpriteSheetImportFactory)
 
diff --git a/Engine/Plugins/2D/Paper2D/Source/PaperSpriteSheetImporter/Private/PaperSpriteSheetImporterModule.cpp b/Engine/Plugins/2D/Paper2D/Source/PaperSpriteSheetImporter/Private/PaperSpriteSheetImporterModule.cpp
--- a/Engine/Plugins/2D/Paper2D/Source/PaperSpriteSheetImporter/Private/PaperSpriteSheetImporterModule.cpp
+++ b/Engine/Plugins/2D/Paper2D/Source/PaperSpriteSheetImporter/Private/PaperSpriteSheetImporterModule.cpp
@@ -1,10 +1,7 @@
 // Copyright Epic Games, Inc. All Rights Reserved.
 
-#include "CoreMinimal.h"
 #include "PaperSpriteSheetImporterLog.h"
-#include "Modules/ModuleManager.h"
 
-#include "AssetToolsModule.h"
 #include "PaperSpriteSheetAssetTypeActions.h"
 
 //////////////////////////////////////////////////////////////////////////
diff --git a/Engine/Plugins/2D/Paper2D/Source/PaperTiledImporter/Classes/PaperTiledImporterFactory.h b/Engine/Plugins/2D/Paper2D/Source/PaperTiledImporter/Classes/PaperTiledImporterFactory.h
--- a/Engine/Plugins/2D/Paper2D/Source/PaperTiledImporter/Classes/PaperTiledImporterFactory.h
+++ b/Engine/Plugins/2D/Paper2D/Source/PaperTiledImporter/Classes/PaperTiledImporterFactory.h
@@ -2,14 +2,16 @@
 
 #pragma once
 
-#include "CoreMinimal.h"
-#include "UObject/ObjectMacros.h"
 #include "EditorReimportHandler.h"
 #include "Factories/Factory.h"
-#include "PaperTileMap.h"
-#include "PaperTileLayer.h"
 #include "PaperTiledImporterFactory.generated.h"
 
+class UPaperTileMap;
+class UTexture2D;
+namespace ETileMapProjectionMode { enum Type : int; }
+struct FPaperTileInfo;
+struct FSpriteGeometryCollection;
+
 class FJsonObject;
 class FJsonValue;
 class UPaperTileSet;
@@ -389,3 +391,10 @@ public:
 	FTileMapFromTiled();
 	bool IsValid() const;
 };
+
+#if UE_ENABLE_INCLUDE_ORDER_DEPRECATED_IN_5_2
+#include "CoreMinimal.h"
+#include "PaperTileLayer.h"
+#include "PaperTileMap.h"
+#include "UObject/ObjectMacros.h"
+#endif
diff --git a/Engine/Plugins/2D/Paper2D/Source/PaperTiledImporter/Private/PaperTiledImporterFactory.cpp b/Engine/Plugins/2D/Paper2D/Source/PaperTiledImporter/Private/PaperTiledImporterFactory.cpp
--- a/Engine/Plugins/2D/Paper2D/Source/PaperTiledImporter/Private/PaperTiledImporterFactory.cpp
+++ b/Engine/Plugins/2D/Paper2D/Source/PaperTiledImporter/Private/PaperTiledImporterFactory.cpp
@@ -1,14 +1,15 @@
 // Copyright Epic Games, Inc. All Rights Reserved.
 
 #include "PaperTiledImporterFactory.h"
+#include "PaperTileLayer.h"
 #include "PaperTiledImporterLog.h"
 #include "Misc/FileHelper.h"
 #include "Misc/Paths.h"
-#include "Modules/ModuleManager.h"
 #include "Misc/PackageName.h"
-#include "Serialization/JsonReader.h"
+#include "PaperTileMap.h"
 #include "Serialization/JsonSerializer.h"
 #include "Editor.h"
+#include "Subsystems/ImportSubsystem.h"
 #include "TileMapAssetImportData.h"
 #include "PaperJSONHelpers.h"
 #include "IAssetTools.h"
@@ -16,7 +17,6 @@
 #include "AssetRegistry/AssetRegistryModule.h"
 #include "PaperImporterSettings.h"
 #include "PackageTools.h"
-#include "IntMargin.h"
 #include "PaperTileSet.h"
 
 #include UE_INLINE_GENERATED_CPP_BY_NAME(PaperTiledImporterFactory)
diff --git a/Engine/Plugins/2D/Paper2D/Source/SmartSnapping/Private/PlanarConstraintSnapPolicy.h b/Engine/Plugins/2D/Paper2D/Source/SmartSnapping/Private/PlanarConstraintSnapPolicy.h
--- a/Engine/Plugins/2D/Paper2D/Source/SmartSnapping/Private/PlanarConstraintSnapPolicy.h
+++ b/Engine/Plugins/2D/Paper2D/Source/SmartSnapping/Private/PlanarConstraintSnapPolicy.h
@@ -2,8 +2,8 @@
 
 #pragma once
 
-#include "CoreMinimal.h"
 #include "ISnappingPolicy.h"
+#include "Math/Plane.h"
 
 class FPrimitiveDrawInterface;
 class FSceneView;
diff --git a/Engine/Plugins/2D/Paper2D/Source/SmartSnapping/Private/SmartSnappingModule.cpp b/Engine/Plugins/2D/Paper2D/Source/SmartSnapping/Private/SmartSnappingModule.cpp
--- a/Engine/Plugins/2D/Paper2D/Source/SmartSnapping/Private/SmartSnappingModule.cpp
+++ b/Engine/Plugins/2D/Paper2D/Source/SmartSnapping/Private/SmartSnappingModule.cpp
@@ -1,11 +1,5 @@
 // Copyright Epic Games, Inc. All Rights Reserved.
 
-#include "CoreMinimal.h"
-#include "Modules/ModuleManager.h"
-#include "Textures/SlateIcon.h"
-#include "Framework/Commands/UIAction.h"
-#include "Framework/Commands/UICommandList.h"
-#include "Framework/MultiBox/MultiBoxExtender.h"
 #include "ViewportSnappingModule.h"
 #include "PlanarConstraintSnapPolicy.h"
 #include "Framework/MultiBox/MultiBoxBuilder.h"
