Made the paper2D contextual asset menu respect the class permission system.

diff --git a/Engine/Plugins/2D/Paper2D/Source/Paper2DEditor/Paper2DEditor.Build.cs b/Engine/Plugins/2D/Paper2D/Source/Paper2DEditor/Paper2DEditor.Build.cs
--- a/Engine/Plugins/2D/Paper2D/Source/Paper2DEditor/Paper2DEditor.Build.cs
+++ b/Engine/Plugins/2D/Paper2D/Source/Paper2DEditor/Paper2DEditor.Build.cs
@@ -11,6 +11,7 @@ public class Paper2DEditor : ModuleRules
 				"Core",
 				"CoreUObject",
 				"ApplicationCore",
+				"AssetTools",
 				"Slate",
 				"SlateCore",
 				"Engine",
@@ -39,14 +40,7 @@ public class Paper2DEditor : ModuleRules
 		PrivateIncludePathModuleNames.AddRange(
 			new string[] {
 				"Settings",
-				"AssetTools",
 				"LevelEditor"
 			});
-
-		DynamicallyLoadedModuleNames.AddRange(
-			new string[] {
-				"AssetTools"
-			});
-
 	}
 }
diff --git a/Engine/Plugins/2D/Paper2D/Source/Paper2DEditor/Private/ContentBrowserExtensions/ContentBrowserExtensions.cpp b/Engine/Plugins/2D/Paper2D/Source/Paper2DEditor/Private/ContentBrowserExtensions/ContentBrowserExtensions.cpp
--- a/Engine/Plugins/2D/Paper2D/Source/Paper2DEditor/Private/ContentBrowserExtensions/ContentBrowserExtensions.cpp
+++ b/Engine/Plugins/2D/Paper2D/Source/Paper2DEditor/Private/ContentBrowserExtensions/ContentBrowserExtensions.cpp
@@ -1,6 +1,7 @@
 // Copyright Epic Games, Inc. All Rights Reserved.
 
 #include "ContentBrowserExtensions/ContentBrowserExtensions.h"
+#include "Misc/EnumClassFlags.h"
 #include "Misc/PackageName.h"
 #include "Framework/MultiBox/MultiBoxBuilder.h"
 #include "IContentBrowserSingleton.h"
@@ -183,6 +184,15 @@ struct FCreateTileSetFromTextureExtension : public FContentBrowserSelectedAssetE
 	}
 };
 
+enum class ECanCreatePaper2DFlags : uint8
+{
+	None = 0,
+	PaperSprite = 1,
+	PaperTileSet = 1 << 1,
+};
+
+ENUM_CLASS_FLAGS(ECanCreatePaper2DFlags);
+
 //////////////////////////////////////////////////////////////////////////
 // FPaperContentBrowserExtensions_Impl
 
@@ -194,51 +204,55 @@ public:
 		SelectedAssetFunctor->Execute();
 	}
 
-	static void CreateSpriteActionsSubMenu(FMenuBuilder& MenuBuilder, TArray<FAssetData> SelectedAssets)
+	static void CreateSpriteActionsSubMenu(FMenuBuilder& MenuBuilder, TArray<FAssetData> SelectedAssets, ECanCreatePaper2DFlags CanCreateFlags)
 	{
 		MenuBuilder.AddSubMenu(
 			LOCTEXT("SpriteActionsSubMenuLabel", "Sprite Actions"),
 			LOCTEXT("SpriteActionsSubMenuToolTip", "Sprite-related actions for this texture."),
-			FNewMenuDelegate::CreateStatic(&FPaperContentBrowserExtensions_Impl::PopulateSpriteActionsMenu, SelectedAssets),
+			FNewMenuDelegate::CreateStatic(&FPaperContentBrowserExtensions_Impl::PopulateSpriteActionsMenu, SelectedAssets, CanCreateFlags),
 			false,
 			FSlateIcon(FPaperStyle::Get()->GetStyleSetName(), "ClassIcon.PaperSprite")
 		);
 	}
 
-	static void PopulateSpriteActionsMenu(FMenuBuilder& MenuBuilder, TArray<FAssetData> SelectedAssets)
+	static void PopulateSpriteActionsMenu(FMenuBuilder& MenuBuilder, TArray<FAssetData> SelectedAssets, ECanCreatePaper2DFlags CanCreateFlags)
 	{
-		// Create sprites
-		TSharedPtr<FCreateSpriteFromTextureExtension> SpriteCreatorFunctor = MakeShareable(new FCreateSpriteFromTextureExtension());
-		SpriteCreatorFunctor->SelectedAssets = SelectedAssets;
-
-		FUIAction Action_CreateSpritesFromTextures(
-			FExecuteAction::CreateStatic(&FPaperContentBrowserExtensions_Impl::ExecuteSelectedContentFunctor, StaticCastSharedPtr<FContentBrowserSelectedAssetExtensionBase>(SpriteCreatorFunctor)));
-		
 		const FName PaperStyleSetName = FPaperStyle::Get()->GetStyleSetName();
 
-		MenuBuilder.AddMenuEntry(
-			LOCTEXT("CB_Extension_Texture_CreateSprite", "Create Sprite"),
-			LOCTEXT("CB_Extension_Texture_CreateSprite_Tooltip", "Create sprites from selected textures"),
-			FSlateIcon(PaperStyleSetName, "AssetActions.CreateSprite"),
-			Action_CreateSpritesFromTextures,
-			NAME_None,
-			EUserInterfaceActionType::Button);
-
-		// Extract Sprites
-		TSharedPtr<FCreateSpriteFromTextureExtension> SpriteExtractorFunctor = MakeShareable(new FCreateSpriteFromTextureExtension());
-		SpriteExtractorFunctor->SelectedAssets = SelectedAssets;
-		SpriteExtractorFunctor->bExtractSprites = true;
-
-		FUIAction Action_ExtractSpritesFromTextures(
-			FExecuteAction::CreateStatic(&FPaperContentBrowserExtensions_Impl::ExecuteSelectedContentFunctor, StaticCastSharedPtr<FContentBrowserSelectedAssetExtensionBase>(SpriteExtractorFunctor)));
-
-		MenuBuilder.AddMenuEntry(
-			LOCTEXT("CB_Extension_Texture_ExtractSprites", "Extract Sprites"),
-			LOCTEXT("CB_Extension_Texture_ExtractSprites_Tooltip", "Extract sprites from selected textures"),
-			FSlateIcon(PaperStyleSetName, "AssetActions.ExtractSprites"),
-			Action_ExtractSpritesFromTextures,
-			NAME_None,
-			EUserInterfaceActionType::Button);
+		if (EnumHasAnyFlags(CanCreateFlags, ECanCreatePaper2DFlags::PaperSprite))
+		{
+			// Create sprites
+			TSharedPtr<FCreateSpriteFromTextureExtension> SpriteCreatorFunctor = MakeShareable(new FCreateSpriteFromTextureExtension());
+			SpriteCreatorFunctor->SelectedAssets = SelectedAssets;
+
+			FUIAction Action_CreateSpritesFromTextures(
+				FExecuteAction::CreateStatic(&FPaperContentBrowserExtensions_Impl::ExecuteSelectedContentFunctor, StaticCastSharedPtr<FContentBrowserSelectedAssetExtensionBase>(SpriteCreatorFunctor)));
+
+
+			MenuBuilder.AddMenuEntry(
+				LOCTEXT("CB_Extension_Texture_CreateSprite", "Create Sprite"),
+				LOCTEXT("CB_Extension_Texture_CreateSprite_Tooltip", "Create sprites from selected textures"),
+				FSlateIcon(PaperStyleSetName, "AssetActions.CreateSprite"),
+				Action_CreateSpritesFromTextures,
+				NAME_None,
+				EUserInterfaceActionType::Button);
+
+			// Extract Sprites
+			TSharedPtr<FCreateSpriteFromTextureExtension> SpriteExtractorFunctor = MakeShareable(new FCreateSpriteFromTextureExtension());
+			SpriteExtractorFunctor->SelectedAssets = SelectedAssets;
+			SpriteExtractorFunctor->bExtractSprites = true;
+
+			FUIAction Action_ExtractSpritesFromTextures(
+				FExecuteAction::CreateStatic(&FPaperContentBrowserExtensions_Impl::ExecuteSelectedContentFunctor, StaticCastSharedPtr<FContentBrowserSelectedAssetExtensionBase>(SpriteExtractorFunctor)));
+
+			MenuBuilder.AddMenuEntry(
+				LOCTEXT("CB_Extension_Texture_ExtractSprites", "Extract Sprites"),
+				LOCTEXT("CB_Extension_Texture_ExtractSprites_Tooltip", "Extract sprites from selected textures"),
+				FSlateIcon(PaperStyleSetName, "AssetActions.ExtractSprites"),
+				Action_ExtractSpritesFromTextures,
+				NAME_None,
+				EUserInterfaceActionType::Button);
+		}
 
 		// Configure the selected textures according to the project settings (same as if it got imported from a sprite sheet)
 		TSharedPtr<FConfigureTexturesForSpriteUsageExtension> TextureConfigFunctor = MakeShareable(new FConfigureTexturesForSpriteUsageExtension());
@@ -255,20 +269,24 @@ public:
 			NAME_None,
 			EUserInterfaceActionType::Button);
 
-		// Create sprites
-		TSharedPtr<FCreateTileSetFromTextureExtension> TileSetCreatorFunctor = MakeShareable(new FCreateTileSetFromTextureExtension());
-		TileSetCreatorFunctor->SelectedAssets = SelectedAssets;
 
-		FUIAction Action_CreateTileSetFromTextures(
-			FExecuteAction::CreateStatic(&FPaperContentBrowserExtensions_Impl::ExecuteSelectedContentFunctor, StaticCastSharedPtr<FContentBrowserSelectedAssetExtensionBase>(TileSetCreatorFunctor)));
-
-		MenuBuilder.AddMenuEntry(
-			LOCTEXT("CB_Extension_Texture_CreateTileSet", "Create Tile Set"),
-			LOCTEXT("CB_Extension_Texture_CreateTileSet_Tooltip", "Create tile set from selected texture"),
-			FSlateIcon(PaperStyleSetName, "AssetActions.CreateTileSet"),
-			Action_CreateTileSetFromTextures,
-			NAME_None,
-			EUserInterfaceActionType::Button);
+		if (EnumHasAnyFlags(CanCreateFlags, ECanCreatePaper2DFlags::PaperTileSet))
+		{
+			// Create tile set
+			TSharedPtr<FCreateTileSetFromTextureExtension> TileSetCreatorFunctor = MakeShareable(new FCreateTileSetFromTextureExtension());
+			TileSetCreatorFunctor->SelectedAssets = SelectedAssets;
+
+			FUIAction Action_CreateTileSetFromTextures(
+				FExecuteAction::CreateStatic(&FPaperContentBrowserExtensions_Impl::ExecuteSelectedContentFunctor, StaticCastSharedPtr<FContentBrowserSelectedAssetExtensionBase>(TileSetCreatorFunctor)));
+
+			MenuBuilder.AddMenuEntry(
+				LOCTEXT("CB_Extension_Texture_CreateTileSet", "Create Tile Set"),
+				LOCTEXT("CB_Extension_Texture_CreateTileSet_Tooltip", "Create tile set from selected texture"),
+				FSlateIcon(PaperStyleSetName, "AssetActions.CreateTileSet"),
+				Action_CreateTileSetFromTextures,
+				NAME_None,
+				EUserInterfaceActionType::Button);
+		}
 	}
 
 	static TSharedRef<FExtender> OnExtendContentBrowserAssetSelectionMenu(const TArray<FAssetData>& SelectedAssets)
@@ -283,14 +301,27 @@ public:
 			bAnyTextures = bAnyTextures || (Asset.AssetClassPath == UTexture2D::StaticClass()->GetClassPathName());
 		}
 
-		if (bAnyTextures)
+		ECanCreatePaper2DFlags CanCreateFlags= ECanCreatePaper2DFlags::None;
+
+		if (CanCreatePaperSpriteAsset())
+		{
+			EnumAddFlags(CanCreateFlags, ECanCreatePaper2DFlags::PaperSprite);
+		}
+		const bool bCanCreatePaperSprite = CanCreatePaperSpriteAsset();
+
+		if (CanCreatePaperTileSetAsset())
+		{
+			EnumAddFlags(CanCreateFlags, ECanCreatePaper2DFlags::PaperTileSet);
+		}
+
+		if (bAnyTextures && CanCreateFlags != ECanCreatePaper2DFlags::None)
 		{
 			// Add the sprite actions sub-menu extender
 			Extender->AddMenuExtension(
 				"GetAssetActions",
 				EExtensionHook::After,
 				nullptr,
-				FMenuExtensionDelegate::CreateStatic(&FPaperContentBrowserExtensions_Impl::CreateSpriteActionsSubMenu, SelectedAssets));
+				FMenuExtensionDelegate::CreateStatic(&FPaperContentBrowserExtensions_Impl::CreateSpriteActionsSubMenu, SelectedAssets, CanCreateFlags));
 		}
 
 		return Extender;
@@ -301,6 +332,16 @@ public:
 		FContentBrowserModule& ContentBrowserModule = FModuleManager::LoadModuleChecked<FContentBrowserModule>(TEXT("ContentBrowser"));
 		return ContentBrowserModule.GetAllAssetViewContextMenuExtenders();
 	}
+
+	static bool CanCreatePaperSpriteAsset()
+	{
+		 return IAssetTools::Get().GetAssetClassPathPermissionList(EAssetClassAction::CreateAsset)->PassesFilter(UPaperSprite::StaticClass()->GetClassPathName().ToString());
+	}
+
+	static bool CanCreatePaperTileSetAsset()
+	{
+		return IAssetTools::Get().GetAssetClassPathPermissionList(EAssetClassAction::CreateAsset)->PassesFilter(UPaperTileSet::StaticClass()->GetClassPathName().ToString());
+	}
 };
 
 //////////////////////////////////////////////////////////////////////////
