Do not attempt to synchronously load the source texture of a Paper2D sprite if we are already in the *ASYNC* loading thread. This can occur in the editor if cooked content is enabled, on post load of the sprite asset.

We can go back to the old way of loading the source texture if this breaks something with the "Paper2d.DisableSyncLoadInEditor" CVar flag.

diff --git a/Engine/Plugins/2D/Paper2D/Source/Paper2D/Private/PaperSprite.cpp b/Engine/Plugins/2D/Paper2D/Source/Paper2D/Private/PaperSprite.cpp
--- a/Engine/Plugins/2D/Paper2D/Source/Paper2D/Private/PaperSprite.cpp
+++ b/Engine/Plugins/2D/Paper2D/Source/Paper2D/Private/PaperSprite.cpp
@@ -1488,6 +1488,12 @@ void UPaperSprite::SetPivotMode(ESpritePivotMode::Type InPivotMode, FVector2D In
 
 #if WITH_EDITOR
 
+static TAutoConsoleVariable<bool> CVarDisableSyncLoadInEditor(
+	TEXT("Paper2d.DisableSyncLoadInEditor"),
+	true,
+	TEXT("Don't attempt to load the source texture synchronous if we are inside the AsyncLoading thread alreayd."),
+	ECVF_Cheat);
+
 UTexture2D* UPaperSprite::GetSourceTexture() const
 {
 	// Verify the cache is still valid.
@@ -1495,8 +1501,13 @@ UTexture2D* UPaperSprite::GetSourceTexture() const
 	{
 		return SourceTextureCacheNeverSerialized;
 	}
-
-	UTexture2D* SourceTexturePtr = SourceTexture.LoadSynchronous();
+	
+	// We don't want to load synchronously unless we are OUTSIDE of an async loading thread.
+	// We could be in the Async loading thread if we are in the editor with cooked content already
+	UTexture2D* SourceTexturePtr =
+			CVarDisableSyncLoadInEditor.GetValueOnGameThread() && IsInAsyncLoadingThread() ?
+			nullptr :
+			SourceTexture.LoadSynchronous();
 
 	// We need to completely load the texture if we're in post load and this texture is needed immediately,
 	// not safe to do in a game with EDL, but safe at editor time.
