Fixed an issue with the paper2d atlas not always being updated when reimporting a texture used by a sprite that use the atlas.

diff --git a/Engine/Plugins/2D/Paper2D/Source/Paper2D/Classes/PaperSprite.h b/Engine/Plugins/2D/Paper2D/Source/Paper2D/Classes/PaperSprite.h
--- a/Engine/Plugins/2D/Paper2D/Source/Paper2D/Classes/PaperSprite.h
+++ b/Engine/Plugins/2D/Paper2D/Source/Paper2D/Classes/PaperSprite.h
@@ -326,6 +326,10 @@ public:
 	virtual bool ContainsPhysicsTriMeshData(bool InUseAllTriData) const override;
 	// End of IInterface_CollisionDataProvider
 
+#if WITH_EDITOR
+	static FName GetSourceTextureMemberName();
+#endif
+
 protected:
 	void RefreshBakedData();
 
diff --git a/Engine/Plugins/2D/Paper2D/Source/Paper2D/Private/PaperSprite.cpp b/Engine/Plugins/2D/Paper2D/Source/Paper2D/Private/PaperSprite.cpp
--- a/Engine/Plugins/2D/Paper2D/Source/Paper2D/Private/PaperSprite.cpp
+++ b/Engine/Plugins/2D/Paper2D/Source/Paper2D/Private/PaperSprite.cpp
@@ -1755,6 +1755,13 @@ bool UPaperSprite::ContainsPhysicsTriMeshData(bool InUseAllTriData) const
 	return false;
 }
 
+#if WITH_EDITOR
+FName UPaperSprite::GetSourceTextureMemberName()
+{
+	return GET_MEMBER_NAME_CHECKED(UPaperSprite, SourceTexture);
+}
+#endif // WITH_EDITOR
+
 FBoxSphereBounds UPaperSprite::GetRenderBounds() const
 {
 	FBox BoundingBox(ForceInit);
diff --git a/Engine/Plugins/2D/Paper2D/Source/Paper2DEditor/Private/Paper2DEditorModule.cpp b/Engine/Plugins/2D/Paper2D/Source/Paper2DEditor/Private/Paper2DEditorModule.cpp
--- a/Engine/Plugins/2D/Paper2D/Source/Paper2DEditor/Private/Paper2DEditorModule.cpp
+++ b/Engine/Plugins/2D/Paper2D/Source/Paper2DEditor/Private/Paper2DEditorModule.cpp
@@ -3,6 +3,8 @@
 #include "Paper2DEditorModule.h"
 #include "Paper2DEditorLog.h"
 #include "EditorModeRegistry.h"
+
+#include "AssetRegistry/IAssetRegistry.h"
 #include "Subsystems/ImportSubsystem.h"
 #include "UObject/UObjectIterator.h"
 
@@ -296,9 +298,39 @@ private:
 	{
 		if (UTexture2D* Texture = Cast<UTexture2D>(InObject))
 		{
-			for (TObjectIterator<UPaperSprite> SpriteIt; SpriteIt; ++SpriteIt)
+			if (GetDefault<UPaperRuntimeSettings>()->bResizeSpriteDataToMatchTextures)
 			{
-				SpriteIt->OnObjectReimported(Texture);
+				// Ideally we wouldn't need to load the sprite to keep the atlas up to date but this will work for now.
+				IAssetRegistry& AssetRegistry = IAssetRegistry::GetChecked();
+				AssetRegistry.WaitForCompletion();
+				TArray<FAssetData> AssetsData;
+				const bool bSearchSubClasses = true;
+				AssetRegistry.GetAssetsByClass(UPaperSprite::StaticClass()->GetClassPathName(), AssetsData, bSearchSubClasses);
+
+				const FString TexturePath = Texture->GetPathName();
+				const FName SourceTextureMemberName = UPaperSprite::GetSourceTextureMemberName();
+				for (const FAssetData& SpriteAssetData : AssetsData)
+				{
+					FAssetData SavedTextureAssetData = AssetRegistry.GetAssetByObjectPath(SpriteAssetData.TagsAndValues.FindTag(SourceTextureMemberName).GetValue());
+				
+					FSoftObjectPath TextureSourcePath;
+					if (SavedTextureAssetData.IsRedirector())
+					{
+						TextureSourcePath = AssetRegistry.GetRedirectedObjectPath(SavedTextureAssetData.GetSoftObjectPath());
+					}
+					else
+					{
+						TextureSourcePath = SavedTextureAssetData.GetSoftObjectPath();
+					}
+
+					if (TextureSourcePath == TexturePath)
+					{
+						if (UPaperSprite* PaperSprite = Cast<UPaperSprite>(SpriteAssetData.GetAsset()))
+						{
+							PaperSprite->OnObjectReimported(Texture);
+						}
+					}
+				}
 			}
 		}
 	}
