The Paper2D Atlas will now support converting the source data of other texture to BGRA8 when creating the atlas.

Note the atlas still only support BGRA8 for the source data of the texture build. This only make so that it won't leave an empty space when a sprite is using a non BGRA8 texture. Also the atlas are still experimental and not supported so this change is only to unblock the users from a recent change to the texture import.

diff --git a/Engine/Plugins/2D/Paper2D/Source/Paper2DEditor/Paper2DEditor.Build.cs b/Engine/Plugins/2D/Paper2D/Source/Paper2DEditor/Paper2DEditor.Build.cs
--- a/Engine/Plugins/2D/Paper2D/Source/Paper2DEditor/Paper2DEditor.Build.cs
+++ b/Engine/Plugins/2D/Paper2D/Source/Paper2DEditor/Paper2DEditor.Build.cs
@@ -16,6 +16,7 @@ public class Paper2DEditor : ModuleRules
 				"SlateCore",
 				"Engine",
 				"InputCore",
+				"ImageCore",
 				"EditorFramework",
 				"UnrealEd", // for AssetEditorSubsystem
 				"KismetWidgets",
diff --git a/Engine/Plugins/2D/Paper2D/Source/Paper2DEditor/Private/Atlasing/PaperAtlasTextureHelpers.cpp b/Engine/Plugins/2D/Paper2D/Source/Paper2DEditor/Private/Atlasing/PaperAtlasTextureHelpers.cpp
--- a/Engine/Plugins/2D/Paper2D/Source/Paper2DEditor/Private/Atlasing/PaperAtlasTextureHelpers.cpp
+++ b/Engine/Plugins/2D/Paper2D/Source/Paper2DEditor/Private/Atlasing/PaperAtlasTextureHelpers.cpp
@@ -5,6 +5,8 @@
 #include "PaperSprite.h"
 #include "PaperSpriteAtlas.h"
 
+#include "ImageCore.h"
+
 //////////////////////////////////////////////////////////////////////////
 // 
 
@@ -32,17 +34,22 @@ bool FPaperAtlasTextureHelpers::ReadSpriteTexture(UTexture* SourceTexture, const
 
 	check(SourceTexture);
 	FTextureSource& SourceData = SourceTexture->Source;
-	if (SourceData.GetFormat() == TSF_BGRA8)
+	FImage SourceMip0;
+	const int32 MipIndex = 0;
+	if (SourceData.GetMipImage(SourceMip0, MipIndex))
 	{
-		uint32 BytesPerPixel = SourceData.GetBytesPerPixel();
-		const uint8* OffsetSource = SourceData.LockMipReadOnly(0) + (SourceXY.X + SourceXY.Y * SourceData.GetSizeX()) * BytesPerPixel;
+		// It will only convert the texture if it is not in the right format
+		SourceMip0.ChangeFormat(ERawImageFormat::BGRA8, EGammaSpace::sRGB);
+
+		uint32 BytesPerPixel = SourceMip0.GetBytesPerPixel();
+		const uint8* OffsetSource = SourceMip0.RawData.GetData() + (SourceXY.X + SourceXY.Y * SourceData.GetSizeX()) * BytesPerPixel;
 		uint8* OffsetDest = TargetBuffer.GetData();
 		CopyTextureData(OffsetSource, OffsetDest, SourceSize.X, SourceSize.Y, BytesPerPixel, SourceData.GetSizeX() * BytesPerPixel, SourceSize.X * BytesPerPixel);
-		SourceData.UnlockMip(0);
 	}
 	else
 	{
-		UE_LOG(LogPaper2DEditor, Error, TEXT("Sprite texture %s is not BGRA8, which isn't supported in atlases yet"), *SourceTexture->GetName());
+		UE_LOG(LogPaper2DEditor, Error, TEXT("Couldn't retrive the source texture of the Sprite texture %s while building the atlas for the sprite."), *SourceTexture->GetName());
+		return false;
 	}
 
 	return true;
