Fixed an issue with the paper 2d sprites not properly being updating when transacted.

diff --git a/Engine/Plugins/2D/Paper2D/Source/Paper2D/Private/PaperSprite.cpp b/Engine/Plugins/2D/Paper2D/Source/Paper2D/Private/PaperSprite.cpp
--- a/Engine/Plugins/2D/Paper2D/Source/Paper2D/Private/PaperSprite.cpp
+++ b/Engine/Plugins/2D/Paper2D/Source/Paper2D/Private/PaperSprite.cpp
@@ -491,7 +491,7 @@ void UPaperSprite::PostEditChangeProperty(FPropertyChangedEvent& PropertyChanged
 	//@TODO: Determine when this is really needed, as it is seriously expensive!
 	FSpriteReregisterContext ReregisterExistingComponents(this);
 
-	// Look for changed properties
+	// Look for changed properties, if null assume all properties might have changed.
 	const FName PropertyName = (PropertyChangedEvent.Property != nullptr) ? PropertyChangedEvent.Property->GetFName() : NAME_None;
 	const FName MemberPropertyName = (PropertyChangedEvent.MemberProperty != nullptr) ? PropertyChangedEvent.MemberProperty->GetFName() : NAME_None;
 
@@ -508,12 +508,13 @@ void UPaperSprite::PostEditChangeProperty(FPropertyChangedEvent& PropertyChanged
 	RenderGeometry.PixelsPerSubdivisionX = FMath::Max(RenderGeometry.PixelsPerSubdivisionX, 4);
 	RenderGeometry.PixelsPerSubdivisionY = FMath::Max(RenderGeometry.PixelsPerSubdivisionY, 4);
 
-	if (MemberPropertyName == GET_MEMBER_NAME_CHECKED(UPaperSprite, SourceUV))
+	if (MemberPropertyName == GET_MEMBER_NAME_CHECKED(UPaperSprite, SourceUV) || PropertyName == NAME_None)
 	{
 		SourceUV.X = FMath::Max(FMath::RoundToFloat(SourceUV.X), 0.0f);
 		SourceUV.Y = FMath::Max(FMath::RoundToFloat(SourceUV.Y), 0.0f);
 	}
-	else if (MemberPropertyName == GET_MEMBER_NAME_CHECKED(UPaperSprite, SourceDimension))
+
+	if (MemberPropertyName == GET_MEMBER_NAME_CHECKED(UPaperSprite, SourceDimension) || PropertyName == NAME_None)
 	{
 		SourceDimension.X = FMath::Max(FMath::RoundToFloat(SourceDimension.X), 0.0f);
 		SourceDimension.Y = FMath::Max(FMath::RoundToFloat(SourceDimension.Y), 0.0f);
@@ -522,39 +523,9 @@ void UPaperSprite::PostEditChangeProperty(FPropertyChangedEvent& PropertyChanged
 	// Update the pivot (roundtripping thru the function will round to a pixel position if that option is enabled)
 	CustomPivotPoint = GetPivotPosition();
 
-	bool bRenderDataModified = false;
-	bool bCollisionDataModified = false;
-	bool bBothModified = false;
-
-	if ((PropertyName == GET_MEMBER_NAME_CHECKED(UPaperSprite, SpriteCollisionDomain)) ||
-		(PropertyName == GET_MEMBER_NAME_CHECKED(UPaperSprite, BodySetup)) ||
-		(PropertyName == GET_MEMBER_NAME_CHECKED(UPaperSprite, CollisionGeometry)) )
-	{
-		bCollisionDataModified = true;
-	}
-
-	// Properties inside one of the geom structures (we don't know which one)
-// 	if ((PropertyName == GET_MEMBER_NAME_CHECKED(UPaperSprite, GeometryType)) ||
-// 		(PropertyName == GET_MEMBER_NAME_CHECKED(UPaperSprite, AlphaThreshold)) ||
-// 		)
-// 		BoxSize
-// 		BoxPosition
-// 		Vertices
-// 		VertexCount
-// 		Polygons
-// 	{
-		bBothModified = true;
-//	}
+	bool bRebuilAtlas = false;
 
-	if ((PropertyName == GET_MEMBER_NAME_CHECKED(UPaperSprite, SourceUV)) ||
-		(PropertyName == GET_MEMBER_NAME_CHECKED(UPaperSprite, SourceDimension)) ||
-		(PropertyName == GET_MEMBER_NAME_CHECKED(UPaperSprite, CustomPivotPoint)) ||
-		(PropertyName == GET_MEMBER_NAME_CHECKED(UPaperSprite, PivotMode)) )
-	{
-		bBothModified = true;
-	}
-
-	if (PropertyName == GET_MEMBER_NAME_CHECKED(UPaperSprite, SourceTexture))
+	if (PropertyName == GET_MEMBER_NAME_CHECKED(UPaperSprite, SourceTexture) || PropertyName == NAME_None)
 	{
 		const bool bResizeSprite = 
 			(SourceTexture != nullptr) &&
@@ -571,35 +542,44 @@ void UPaperSprite::PostEditChangeProperty(FPropertyChangedEvent& PropertyChanged
 			SourceTextureDimension = SourceDimension;
 		}
 
-		// Rebuild atlas group, if we have one.
-		if (AtlasGroup != nullptr)
-		{
-			AtlasGroup->PostEditChange();
-		}
-
-		bBothModified = true;
+		bRebuilAtlas = true;
 	}
 
 	if (PropertyName == GET_MEMBER_NAME_CHECKED(UPaperSprite, Sockets) ||
-		(MemberPropertyName == GET_MEMBER_NAME_CHECKED(UPaperSprite, Sockets) && PropertyName == GET_MEMBER_NAME_CHECKED(FPaperSpriteSocket, SocketName)))
+		(MemberPropertyName == GET_MEMBER_NAME_CHECKED(UPaperSprite, Sockets) && PropertyName == GET_MEMBER_NAME_CHECKED(FPaperSpriteSocket, SocketName)) ||
+		PropertyName == NAME_None)
 	{
-		ValidateSocketNames();
+		if (PropertyChangedEvent.ChangeType != EPropertyChangeType::Interactive)
+		{
+			ValidateSocketNames();
+		}
 	}
 
 	if (PropertyName == GET_MEMBER_NAME_CHECKED(UPaperSprite, AtlasGroup))
 	{
-		UPaperSpriteAtlas* PreviousAtlasGroupPtr = PreviousAtlasGroup.LoadSynchronous();
-		
-		if (PreviousAtlasGroupPtr != AtlasGroup)
+		bRebuilAtlas = true;
+	}
+
+	// Don't do rebuilds during an interactive event to make things more responsive.
+	// They'll always be followed by a ValueSet event at the end to force the change.
+	if (PropertyChangedEvent.ChangeType != EPropertyChangeType::Interactive)
+	{
+		if (bRebuilAtlas)
 		{
-			// Update previous
-			if (PreviousAtlasGroupPtr != nullptr)
+			TSoftObjectPtr<UPaperSpriteAtlas> AtlasGroupPath = AtlasGroup;
+			if (PreviousAtlasGroup != AtlasGroupPath)
 			{
-				PreviousAtlasGroupPtr->PostEditChange();
-			}
+				UPaperSpriteAtlas* PreviousAtlasGroupPtr = PreviousAtlasGroup.LoadSynchronous();
 
-			// Update cached previous atlas group
-			PreviousAtlasGroup = AtlasGroup;
+				// Update previous
+				if (PreviousAtlasGroupPtr != nullptr)
+				{
+					PreviousAtlasGroupPtr->PostEditChange();
+				}
+
+				// Update cached previous atlas group
+				PreviousAtlasGroup = AtlasGroup;
+			}
 
 			// Rebuild atlas group
 			if (AtlasGroup != nullptr)
@@ -611,46 +591,15 @@ void UPaperSprite::PostEditChangeProperty(FPropertyChangedEvent& PropertyChanged
 				BakedSourceTexture = nullptr;
 				BakedSourceUV = FVector2D(0, 0);
 				BakedSourceDimension = FVector2D(0, 0);
-				bRenderDataModified = true;
 			}
-
 		}
-	}
-
-	// The texture dimensions have changed
-	//if (NeedRescaleSpriteData())
-	//{
-		// TMP: Disabled, not sure if we want this here
-		// RescaleSpriteData(GetSourceTexture());
-		// bBothModified = true;
-	//}
 
-	// Don't do rebuilds during an interactive event to make things more responsive.
-	// They'll always be followed by a ValueSet event at the end to force the change.
-	if (PropertyChangedEvent.ChangeType == EPropertyChangeType::Interactive)
-	{
-		bCollisionDataModified = false;
-		bRenderDataModified = false;
-		bBothModified = false;
-	}
 
-	if (bBothModified)
-	{
 		RefreshBakedData();
-	}
-
-	if (bCollisionDataModified || bBothModified)
-	{
 		RebuildCollisionData();
-	}
-
-	if (bRenderDataModified || bBothModified)
-	{
 		RebuildRenderData();
 	}
 
-	
-
 	Super::PostEditChangeProperty(PropertyChangedEvent);
 }
 
