Removed operator new (TArray) usage from headers.

diff --git a/Engine/Plugins/2D/Paper2D/Source/Paper2D/Private/PaperRenderSceneProxy.h b/Engine/Plugins/2D/Paper2D/Source/Paper2D/Private/PaperRenderSceneProxy.h
--- a/Engine/Plugins/2D/Paper2D/Source/Paper2D/Private/PaperRenderSceneProxy.h
+++ b/Engine/Plugins/2D/Paper2D/Source/Paper2D/Private/PaperRenderSceneProxy.h
@@ -66,8 +66,8 @@ struct PAPER2D_API FSpriteRenderSection
 		return (Material != nullptr) && (NumVertices > 0) && (GetBaseTextureResource() != nullptr);
 	}
 
-	template <typename SourceArrayType>
-	void AddVerticesFromDrawCallRecord(const FSpriteDrawCallRecord& Record, int32 StartIndexWithinRecord, int32 NumVertsToCopy, SourceArrayType& Vertices)
+	template <typename VerticesArrayAllocator>
+	void AddVerticesFromDrawCallRecord(const FSpriteDrawCallRecord& Record, int32 StartIndexWithinRecord, int32 NumVertsToCopy, TArray<FDynamicMeshVertex, VerticesArrayAllocator>& Vertices)
 	{
 		if (NumVertices == 0)
 		{
@@ -93,25 +93,25 @@ struct PAPER2D_API FSpriteRenderSection
 			const FVector Pos((PaperAxisX * SourceVert.X) + (PaperAxisY * SourceVert.Y) + Record.Destination);
 			const FVector2f UV(SourceVert.Z, SourceVert.W);	// LWC_TODO: Precision loss
 
-			new (Vertices) FDynamicMeshVertex((FVector3f)Pos, FPaperSpriteTangents::PackedNormalX.ToFVector3f(), FPaperSpriteTangents::PackedNormalZ.ToFVector3f(), UV, VertColor);
+			Vertices.Emplace((FVector3f)Pos, FPaperSpriteTangents::PackedNormalX.ToFVector3f(), FPaperSpriteTangents::PackedNormalZ.ToFVector3f(), UV, VertColor);
 		}
 	}
 
-	template <typename SourceArrayType>
-	inline void AddVertex(float X, float Y, float U, float V, const FVector& Origin, const FColor& Color, SourceArrayType& Vertices)
+	template <typename VerticesArrayAllocator>
+	inline void AddVertex(float X, float Y, float U, float V, const FVector& Origin, const FColor& Color, TArray<FDynamicMeshVertex, VerticesArrayAllocator>& Vertices)
 	{
 		const FVector Pos((PaperAxisX * X) + (PaperAxisY * Y) + Origin);
 
-		new (Vertices) FDynamicMeshVertex(FVector3f(Pos), FPaperSpriteTangents::PackedNormalX.ToFVector3f(), FPaperSpriteTangents::PackedNormalZ.ToFVector3f(), FVector2f(U, V), Color);
+		Vertices.Emplace(FVector3f(Pos), FPaperSpriteTangents::PackedNormalX.ToFVector3f(), FPaperSpriteTangents::PackedNormalZ.ToFVector3f(), FVector2f(U, V), Color);
 		++NumVertices;
 	}
 
-	template <typename SourceArrayType>
-	inline void AddVertex(float X, float Y, float U, float V, const FVector& Origin, const FColor& Color, const FPackedNormal& TangentX, const FPackedNormal& TangentZ, SourceArrayType& Vertices)
+	template <typename VerticesArrayAllocator>
+	inline void AddVertex(float X, float Y, float U, float V, const FVector& Origin, const FColor& Color, const FPackedNormal& TangentX, const FPackedNormal& TangentZ, TArray<FDynamicMeshVertex, VerticesArrayAllocator>& Vertices)
 	{
 		const FVector Pos((PaperAxisX * X) + (PaperAxisY * Y) + Origin);
 
-		new (Vertices) FDynamicMeshVertex(FVector3f(Pos), TangentX.ToFVector3f(), TangentZ.ToFVector3f(), FVector2f(U, V), Color);
+		Vertices.Emplace(FVector3f(Pos), TangentX.ToFVector3f(), TangentZ.ToFVector3f(), FVector2f(U, V), Color);
 		++NumVertices;
 	}
 };
diff --git a/Engine/Plugins/2D/Paper2D/Source/Paper2DEditor/Private/Atlasing/PaperAtlasHelpers.h b/Engine/Plugins/2D/Paper2D/Source/Paper2DEditor/Private/Atlasing/PaperAtlasHelpers.h
--- a/Engine/Plugins/2D/Paper2D/Source/Paper2DEditor/Private/Atlasing/PaperAtlasHelpers.h
+++ b/Engine/Plugins/2D/Paper2D/Source/Paper2DEditor/Private/Atlasing/PaperAtlasHelpers.h
@@ -55,12 +55,12 @@ static int32 CreateEmptySlot(TArray<FPaperSpriteAtlasSlot>& Slots, int32 AtlasIn
 {
 	if (Width > 0 && Height > 0)
 	{
-		FPaperSpriteAtlasSlot* NewSlot = new(Slots)FPaperSpriteAtlasSlot();
-		NewSlot->AtlasIndex = AtlasIndex;
-		NewSlot->X = X;
-		NewSlot->Y = Y;
-		NewSlot->Width = Width;
-		NewSlot->Height = Height;
+		FPaperSpriteAtlasSlot& NewSlot = Slots.AddDefaulted_GetRef();
+		NewSlot.AtlasIndex = AtlasIndex;
+		NewSlot.X = X;
+		NewSlot.Y = Y;
+		NewSlot.Width = Width;
+		NewSlot.Height = Height;
 		return Slots.Num() - 1;
 	}
 	else
diff --git a/Engine/Plugins/Online/OnlineSubsystemNull/Source/Private/OnlineSessionInterfaceNull.h b/Engine/Plugins/Online/OnlineSubsystemNull/Source/Private/OnlineSessionInterfaceNull.h
--- a/Engine/Plugins/Online/OnlineSubsystemNull/Source/Private/OnlineSessionInterfaceNull.h
+++ b/Engine/Plugins/Online/OnlineSubsystemNull/Source/Private/OnlineSessionInterfaceNull.h
@@ -188,13 +188,13 @@ PACKAGE_SCOPE:
 	class FNamedOnlineSession* AddNamedSession(FName SessionName, const FOnlineSessionSettings& SessionSettings) override
 	{
 		FScopeLock ScopeLock(&SessionLock);
-		return new (Sessions) FNamedOnlineSession(SessionName, SessionSettings);
+		return &Sessions.Emplace_GetRef(SessionName, SessionSettings);
 	}
 
 	class FNamedOnlineSession* AddNamedSession(FName SessionName, const FOnlineSession& Session) override
 	{
 		FScopeLock ScopeLock(&SessionLock);
-		return new (Sessions) FNamedOnlineSession(SessionName, Session);
+		return &Sessions.Emplace_GetRef(SessionName, Session);
 	}
 
 	/**
diff --git a/Engine/Plugins/Online/OnlineSubsystemSteam/Source/Private/OnlineSessionInterfaceSteam.h b/Engine/Plugins/Online/OnlineSubsystemSteam/Source/Private/OnlineSessionInterfaceSteam.h
--- a/Engine/Plugins/Online/OnlineSubsystemSteam/Source/Private/OnlineSessionInterfaceSteam.h
+++ b/Engine/Plugins/Online/OnlineSubsystemSteam/Source/Private/OnlineSessionInterfaceSteam.h
@@ -303,7 +303,7 @@ PACKAGE_SCOPE:
 	class FNamedOnlineSession* AddNamedSession(FName SessionName, const FOnlineSessionSettings& SessionSettings) override
 	{
 		FScopeLock ScopeLock(&SessionLock);
-		return new (Sessions) FNamedOnlineSession(SessionName, SessionSettings);
+		return &Sessions.Emplace_GetRef(SessionName, SessionSettings);
 	}
 
 	/**
@@ -317,7 +317,7 @@ PACKAGE_SCOPE:
 	class FNamedOnlineSession* AddNamedSession(FName SessionName, const FOnlineSession& Session) override
 	{
 		FScopeLock ScopeLock(&SessionLock);
-		return new (Sessions) FNamedOnlineSession(SessionName, Session);
+		return &Sessions.Emplace_GetRef(SessionName, Session);
 	}
 
 	/**
diff --git a/Engine/Plugins/Online/OnlineSubsystemSteam/Source/Private/OnlineSubsystemSteamTypes.h b/Engine/Plugins/Online/OnlineSubsystemSteam/Source/Private/OnlineSubsystemSteamTypes.h
--- a/Engine/Plugins/Online/OnlineSubsystemSteam/Source/Private/OnlineSubsystemSteamTypes.h
+++ b/Engine/Plugins/Online/OnlineSubsystemSteam/Source/Private/OnlineSubsystemSteamTypes.h
@@ -620,7 +620,7 @@ public:
 
 			if (bCreateIfMissing)
 			{
-				return new (CloudMetadata) FCloudFileHeader(FileName, FileName, 0);
+				return &CloudMetadata.Emplace_GetRef(FileName, FileName, 0);
 			}
 		}
 
@@ -650,7 +650,7 @@ public:
 
 			if (bCreateIfMissing)
 			{
-				return new (CloudFileData) FCloudFile(FileName);
+				return &CloudFileData.Emplace_GetRef(FileName);
 			}
 		}
 
diff --git a/Engine/Source/Developer/MeshBuilder/Private/BuildOptimizationHelper.h b/Engine/Source/Developer/MeshBuilder/Private/BuildOptimizationHelper.h
--- a/Engine/Source/Developer/MeshBuilder/Private/BuildOptimizationHelper.h
+++ b/Engine/Source/Developer/MeshBuilder/Private/BuildOptimizationHelper.h
@@ -207,7 +207,7 @@ public:
 		VertIndexAndZ.Empty(NumVertices);
 		for (int32 VertIndex = 0; VertIndex < NumVertices; VertIndex++)
 		{
-			new(VertIndexAndZ)FIndexAndZ(VertIndex, FVector(InVertices.Position[VertIndex]));
+			VertIndexAndZ.Emplace(VertIndex, FVector(InVertices.Position[VertIndex]));
 		}
 		VertIndexAndZ.Sort(FCompareIndexAndZ());
 
diff --git a/Engine/Source/Developer/ShaderCompilerCommon/Private/HlslExpressionParser.inl b/Engine/Source/Developer/ShaderCompilerCommon/Private/HlslExpressionParser.inl
--- a/Engine/Source/Developer/ShaderCompilerCommon/Private/HlslExpressionParser.inl
+++ b/Engine/Source/Developer/ShaderCompilerCommon/Private/HlslExpressionParser.inl
@@ -104,8 +104,8 @@ namespace CrossCompiler
 			Current(InCurrent)
 		{
 			Original = *InCurrent;
-			auto* NewScope = new(Original->Children) FSymbolScope(InAllocator, Original);
-			*Current = NewScope;
+			auto& NewScope = Original->Children.Emplace_GetRef(InAllocator, Original);
+			*Current = &NewScope;
 		}
 
 		~FCreateSymbolScope()
diff --git a/Engine/Source/Developer/ShaderCompilerCommon/Private/HlslUtils.h b/Engine/Source/Developer/ShaderCompilerCommon/Private/HlslUtils.h
--- a/Engine/Source/Developer/ShaderCompilerCommon/Private/HlslUtils.h
+++ b/Engine/Source/Developer/ShaderCompilerCommon/Private/HlslUtils.h
@@ -251,7 +251,7 @@ namespace CrossCompiler
 
 		inline void AddMessage(bool bIsError, const FString& Message)
 		{
-			auto* NewMessage = new(MessageList) FMessage(bIsError, Message);
+			auto& NewMessage = MessageList.Emplace_GetRef(bIsError, Message);
 		}
 
 		inline void SourceError(const FSourceInfo& SourceInfo, const TCHAR* String)
diff --git a/Engine/Source/Runtime/AnimGraphRuntime/Public/AnimNodes/AnimNode_BlendListBase.h b/Engine/Source/Runtime/AnimGraphRuntime/Public/AnimNodes/AnimNode_BlendListBase.h
--- a/Engine/Source/Runtime/AnimGraphRuntime/Public/AnimNodes/AnimNode_BlendListBase.h
+++ b/Engine/Source/Runtime/AnimGraphRuntime/Public/AnimNodes/AnimNode_BlendListBase.h
@@ -85,7 +85,7 @@ public:
 	virtual void AddPose()
 	{
 		BlendTime.Add(0.1f);
-		new (BlendPose) FPoseLink();
+		BlendPose.AddDefaulted();
 	}
 
 	virtual void RemovePose(int32 PoseIndex)
diff --git a/Engine/Source/Runtime/AnimGraphRuntime/Public/AnimNodes/AnimNode_LayeredBoneBlend.h b/Engine/Source/Runtime/AnimGraphRuntime/Public/AnimNodes/AnimNode_LayeredBoneBlend.h
--- a/Engine/Source/Runtime/AnimGraphRuntime/Public/AnimNodes/AnimNode_LayeredBoneBlend.h
+++ b/Engine/Source/Runtime/AnimGraphRuntime/Public/AnimNodes/AnimNode_LayeredBoneBlend.h
@@ -131,7 +131,7 @@ public:
 	void AddPose()
 	{
 		BlendWeights.Add(1.f);
-		new (BlendPoses) FPoseLink();
+		BlendPoses.AddDefaulted();
 
 		if (BlendMode == ELayeredBoneBlendMode::BlendMask) 
 		{ 
@@ -139,7 +139,7 @@ public:
 		}
 		else /*if (BlendMode::BranchFilter)*/ 
 		{ 
-			new (LayerSetup) FInputBlendPose(); 
+			LayerSetup.AddDefaulted(); 
 		}
 	}
 
diff --git a/Engine/Source/Runtime/Core/Public/Containers/SortedMap.h b/Engine/Source/Runtime/Core/Public/Containers/SortedMap.h
--- a/Engine/Source/Runtime/Core/Public/Containers/SortedMap.h
+++ b/Engine/Source/Runtime/Core/Public/Containers/SortedMap.h
@@ -346,7 +346,7 @@ public:
 	{
 		for (typename ElementArrayType::TConstIterator PairIt(Pairs); PairIt; ++PairIt)
 		{
-			new(OutKeys) KeyType(PairIt->Key);
+			OutKeys.Add(PairIt->Key);
 		}
 
 		return OutKeys.Num();
@@ -360,7 +360,7 @@ public:
 		OutArray.Empty(Pairs.Num());
 		for(typename ElementArrayType::TConstIterator PairIt(Pairs);PairIt;++PairIt)
 		{
-			new(OutArray) KeyType(PairIt->Key);
+			OutArray.Add(PairIt->Key);
 		}
 	}
 
@@ -372,7 +372,7 @@ public:
 		OutArray.Empty(Pairs.Num());
 		for(typename ElementArrayType::TConstIterator PairIt(Pairs);PairIt;++PairIt)
 		{
-			new(OutArray) ValueType(PairIt->Value);
+			OutArray.Add(PairIt->Value);
 		}
 	}
 
diff --git a/Engine/Source/Runtime/Core/Public/Stats/StatsFile.h b/Engine/Source/Runtime/Core/Public/Stats/StatsFile.h
--- a/Engine/Source/Runtime/Core/Public/Stats/StatsFile.h
+++ b/Engine/Source/Runtime/Core/Public/Stats/StatsFile.h
@@ -822,7 +822,7 @@ public:
 		out_MetadataMessages.Reserve((int32)Header.NumMetadataMessages );
 		for( int32 Index = 0; Index < Header.NumMetadataMessages; Index++ )
 		{
-			new(out_MetadataMessages)FStatMessage( ReadMessage( Ar, false ) );
+			out_MetadataMessages.Add( ReadMessage( Ar, false ) );
 		}
 	}
 
@@ -841,7 +841,7 @@ public:
 		out_MetadataMessages.Reserve(Header.NumMetadataMessages);
 		for (uint64 Index = 0; Index < Header.NumMetadataMessages; Index++)
 		{
-			new(out_MetadataMessages)FStatMessage(ReadMessage(Ar, false));
+			out_MetadataMessages.Add(ReadMessage(Ar, false));
 		}
 	}
 };
diff --git a/Engine/Source/Runtime/CoreUObject/Public/UObject/UObjectGlobals.h b/Engine/Source/Runtime/CoreUObject/Public/UObject/UObjectGlobals.h
--- a/Engine/Source/Runtime/CoreUObject/Public/UObject/UObjectGlobals.h
+++ b/Engine/Source/Runtime/CoreUObject/Public/UObject/UObjectGlobals.h
@@ -1416,7 +1416,7 @@ private:
 			{
 				check(SubobjectInits[Index].Subobject != Subobject);
 			}
-			new (SubobjectInits) FSubobjectInit(Subobject, Template);
+			SubobjectInits.Emplace(Subobject, Template);
 		}
 		/**  Element of the SubobjectInits array **/
 		struct FSubobjectInit
diff --git a/Engine/Source/Runtime/Engine/Classes/Commandlets/Commandlet.h b/Engine/Source/Runtime/Engine/Classes/Commandlets/Commandlet.h
--- a/Engine/Source/Runtime/Engine/Classes/Commandlets/Commandlet.h
+++ b/Engine/Source/Runtime/Engine/Classes/Commandlets/Commandlet.h
@@ -115,11 +115,11 @@ class UCommandlet : public UObject
 		{
 			if ( **NextToken == TCHAR('-') )
 			{
-				new(Switches) FString(NextToken.Mid(1));
+				Switches.Add(NextToken.Mid(1));
 			}
 			else
 			{
-				new(Tokens) FString(NextToken);
+				Tokens.Add(MoveTemp(NextToken));
 			}
 		}
 	}
diff --git a/Engine/Source/Runtime/Engine/Classes/Engine/BlueprintGeneratedClass.h b/Engine/Source/Runtime/Engine/Classes/Engine/BlueprintGeneratedClass.h
--- a/Engine/Source/Runtime/Engine/Classes/Engine/BlueprintGeneratedClass.h
+++ b/Engine/Source/Runtime/Engine/Classes/Engine/BlueprintGeneratedClass.h
@@ -406,7 +406,7 @@ public:
 		//@TODO: Nasty expansion behavior during compile time
 		if (bBreakpointSite)
 		{
-			new (DebugNodeLineNumbers) FNodeToCodeAssociation(SourceNode, InFunction, CodeOffset);
+			DebugNodeLineNumbers.Emplace(SourceNode, InFunction, CodeOffset);
 			DebugNodeIndexLookup.Add(SourceNode, DebugNodeLineNumbers.Num() - 1);
 		}
 
diff --git a/Engine/Source/Runtime/Engine/Classes/PhysicsEngine/PhysicsConstraintTemplate.h b/Engine/Source/Runtime/Engine/Classes/PhysicsEngine/PhysicsConstraintTemplate.h
--- a/Engine/Source/Runtime/Engine/Classes/PhysicsEngine/PhysicsConstraintTemplate.h
+++ b/Engine/Source/Runtime/Engine/Classes/PhysicsEngine/PhysicsConstraintTemplate.h
@@ -54,9 +54,9 @@ class UPhysicsConstraintTemplate : public UObject
 #if WITH_EDITOR
 	void AddConstraintProfile(FName ProfileName)
 	{
-		FPhysicsConstraintProfileHandle* NewHandle = new(ProfileHandles) FPhysicsConstraintProfileHandle;
-		NewHandle->ProfileName = ProfileName;
-		NewHandle->ProfileProperties = DefaultInstance.ProfileInstance;	//Copy current settings into new profile
+		FPhysicsConstraintProfileHandle& NewHandle = ProfileHandles.AddDefaulted_GetRef();
+		NewHandle.ProfileName = ProfileName;
+		NewHandle.ProfileProperties = DefaultInstance.ProfileInstance;	//Copy current settings into new profile
 	}
 
 	void RemoveConstraintProfile(FName ProfileName)
diff --git a/Engine/Source/Runtime/Engine/Public/MeshBatch.h b/Engine/Source/Runtime/Engine/Public/MeshBatch.h
--- a/Engine/Source/Runtime/Engine/Public/MeshBatch.h
+++ b/Engine/Source/Runtime/Engine/Public/MeshBatch.h
@@ -490,7 +490,7 @@ struct FMeshBatch
 #endif
 	{
 		// By default always add the first element.
-		new(Elements) FMeshBatchElement;
+		Elements.AddDefaulted();
 	}
 };
 
diff --git a/Engine/Source/Runtime/Engine/Public/TextureLayout.h b/Engine/Source/Runtime/Engine/Public/TextureLayout.h
--- a/Engine/Source/Runtime/Engine/Public/TextureLayout.h
+++ b/Engine/Source/Runtime/Engine/Public/TextureLayout.h
@@ -43,7 +43,7 @@ public:
 		bPowerOfTwoSize(bInPowerOfTwoSize),
 		bAlignByFour(bInAlignByFour)
 	{
-		new(Nodes) FTextureLayoutNode(0, 0, IntCastChecked<uint16>(MaxSizeX), IntCastChecked<uint16>(MaxSizeY));
+		Nodes.Emplace(0, 0, IntCastChecked<uint16>(MaxSizeX), IntCastChecked<uint16>(MaxSizeY));
 	}
 
 	/**
@@ -324,7 +324,7 @@ private:
 
 				// Create a child with the same width as the element being placed.
 				// The height may not be the same as the element height yet, in that case another subdivision will occur when traversing this child node.
-				new(Nodes) FTextureLayoutNode(
+				Nodes.Emplace(
 					CurrentNode.MinX,
 					CurrentNode.MinY,
 					IntCastChecked<uint16>(ElementSizeX),
@@ -333,7 +333,7 @@ private:
 
 				// Create a second child to contain the leftover area in the X direction
 				Nodes[NodeIndex].ChildB = Nodes.Num();
-				new(Nodes) FTextureLayoutNode(
+				Nodes.Emplace(
 					IntCastChecked<uint16>(CurrentNode.MinX + ElementSizeX),
 					CurrentNode.MinY,
 					IntCastChecked<uint16>(CurrentNode.SizeX - ElementSizeX),
@@ -343,7 +343,7 @@ private:
 			else
 			{
 				Nodes[NodeIndex].ChildA = Nodes.Num();
-				new(Nodes) FTextureLayoutNode(
+				Nodes.Emplace(
 					CurrentNode.MinX,
 					CurrentNode.MinY,
 					CurrentNode.SizeX,
@@ -351,7 +351,7 @@ private:
 					);
 
 				Nodes[NodeIndex].ChildB = Nodes.Num();
-				new(Nodes) FTextureLayoutNode(
+				Nodes.Emplace(
 					CurrentNode.MinX,
 					IntCastChecked<uint16>(CurrentNode.MinY + ElementSizeY),
 					CurrentNode.SizeX,
diff --git a/Engine/Source/Runtime/Engine/Public/TextureLayout3d.h b/Engine/Source/Runtime/Engine/Public/TextureLayout3d.h
--- a/Engine/Source/Runtime/Engine/Public/TextureLayout3d.h
+++ b/Engine/Source/Runtime/Engine/Public/TextureLayout3d.h
@@ -36,8 +36,8 @@ public:
 		bAllowShrink(bInAllowShrink)
 	{
 		check(MaxSizeX < USHRT_MAX && MaxSizeY < USHRT_MAX && MaxSizeZ < USHRT_MAX);
-		new(Nodes) FTextureLayoutNode3d(0, 0, 0, IntCastChecked<uint16>(MaxSizeX), IntCastChecked<uint16>(MaxSizeY), IntCastChecked<uint16>(MaxSizeZ), INDEX_NONE);
-		new (UnusedLeaves) FUnusedLeaf(0, 0, 0, IntCastChecked<uint16>(MaxSizeX), IntCastChecked<uint16>(MaxSizeY), IntCastChecked<uint16>(MaxSizeZ), 0);
+		Nodes.Emplace(0, 0, 0, IntCastChecked<uint16>(MaxSizeX), IntCastChecked<uint16>(MaxSizeY), IntCastChecked<uint16>(MaxSizeZ), INDEX_NONE);
+		UnusedLeaves.Emplace(0, 0, 0, IntCastChecked<uint16>(MaxSizeX), IntCastChecked<uint16>(MaxSizeY), IntCastChecked<uint16>(MaxSizeZ), 0);
 	}
 
 	/**
@@ -170,7 +170,7 @@ public:
 				FoundNodeIndex = LastParentNodeIndex;
 			}
 
-			new (UnusedLeaves) FUnusedLeaf(
+			UnusedLeaves.Emplace(
 				Nodes[FoundNodeIndex].MinX,
 				Nodes[FoundNodeIndex].MinY,
 				Nodes[FoundNodeIndex].MinZ,
@@ -458,7 +458,7 @@ private:
 
 			const FTextureLayoutNode3d& ChildBNode = Nodes[ChildBNodeIndex];
 
-			new (UnusedLeaves) FUnusedLeaf(ChildBNode.MinX, ChildBNode.MinY, ChildBNode.MinZ, ChildBNode.SizeX, ChildBNode.SizeY, ChildBNode.SizeZ, ChildBNodeIndex);
+			UnusedLeaves.Emplace(ChildBNode.MinX, ChildBNode.MinY, ChildBNode.MinZ, ChildBNode.SizeX, ChildBNode.SizeY, ChildBNode.SizeZ, ChildBNodeIndex);
 
 			// Only traversing ChildA, since ChildA is always the newly created node that matches the element size
 			return AddSurfaceInner(ChildANodeIndex, ElementSizeX, ElementSizeY, ElementSizeZ, bAllowTextureEnlargement);
diff --git a/Engine/Source/Runtime/Json/Public/Serialization/JsonSerializerMacros.h b/Engine/Source/Runtime/Json/Public/Serialization/JsonSerializerMacros.h
--- a/Engine/Source/Runtime/Json/Public/Serialization/JsonSerializerMacros.h
+++ b/Engine/Source/Runtime/Json/Public/Serialization/JsonSerializerMacros.h
@@ -118,8 +118,8 @@
 			{ \
 				for (auto It = Serializer.GetObject()->GetArrayField(JsonName).CreateConstIterator(); It; ++It) \
 				{ \
-					ElementType* Obj = new(JsonArray) ElementType(); \
-					Obj->FromJson((*It)->AsObject()); \
+					ElementType& Obj = JsonArray.AddDefaulted_GetRef(); \
+					Obj.FromJson((*It)->AsObject()); \
 				} \
 			} \
 		} \
@@ -140,8 +140,8 @@
 			{ \
 				for (auto It = Serializer.GetObject()->GetArrayField(JsonName).CreateConstIterator(); It; ++It) \
 				{ \
-					ElementType* Obj = new(JsonArray) ElementType(); \
-					Obj->FromJson((*It)->AsObject()); \
+					ElementType& Obj = JsonArray.AddDefaulted_GetRef(); \
+					Obj.FromJson((*It)->AsObject()); \
 				} \
 			} \
 			else \
@@ -167,8 +167,8 @@
 				TArray<ElementType>& JsonArray = OptionalJsonArray.Emplace(); \
 				for (auto It = Serializer.GetObject()->GetArrayField(JsonName).CreateConstIterator(); It; ++It) \
 				{ \
-					ElementType* Obj = new(JsonArray) ElementType(); \
-					Obj->FromJson((*It)->AsObject()); \
+					ElementType& Obj = JsonArray.AddDefaulted_GetRef(); \
+					Obj.FromJson((*It)->AsObject()); \
 				} \
 			} \
 		} \
diff --git a/Engine/Source/Runtime/Landscape/Public/LandscapeEdit.h b/Engine/Source/Runtime/Landscape/Public/LandscapeEdit.h
--- a/Engine/Source/Runtime/Landscape/Public/LandscapeEdit.h
+++ b/Engine/Source/Runtime/Landscape/Public/LandscapeEdit.h
@@ -78,11 +78,11 @@ struct FLandscapeTextureDataInfo
 			MipInfo[MipNum].bFull = true;
 			MipInfo[MipNum].MipUpdateRegions.Reset();
 			// Push a full region for UpdateTextureData() to process later
-			new(MipInfo[MipNum].MipUpdateRegions) FUpdateTextureRegion2D(0, 0, 0, 0, Width, Height);
+			MipInfo[MipNum].MipUpdateRegions.Emplace(0, 0, 0, 0, Width, Height);
 			return;
 		}
 
-		new(MipInfo[MipNum].MipUpdateRegions) FUpdateTextureRegion2D(InX1, InY1, InX1, InY1, Width, Height);
+		MipInfo[MipNum].MipUpdateRegions.Emplace(InX1, InY1, InX1, InY1, Width, Height);
 	}
 		
 	void* GetMipData(int32 MipNum)
diff --git a/Engine/Source/Runtime/Net/Core/Classes/Net/Serialization/FastArraySerializer.h b/Engine/Source/Runtime/Net/Core/Classes/Net/Serialization/FastArraySerializer.h
--- a/Engine/Source/Runtime/Net/Core/Classes/Net/Serialization/FastArraySerializer.h
+++ b/Engine/Source/Runtime/Net/Core/Classes/Net/Serialization/FastArraySerializer.h
@@ -1486,7 +1486,7 @@ bool FFastArraySerializer::FastArrayDeltaSerialize(TArray<Type> &Items, FNetDelt
 			{
 				UE_LOG(LogNetFastTArray, Log, TEXT("   New. ID: %d. New Element!"), ElementID);
 
-				ThisElement = new (Items)Type();
+				ThisElement = &Items.AddDefaulted_GetRef();
 				ThisElement->ReplicationID = ElementID;
 
 				ElementIndex = Items.Num()-1;
diff --git a/Engine/Source/Runtime/RenderCore/Public/ShaderParameterMetadataBuilder.h b/Engine/Source/Runtime/RenderCore/Public/ShaderParameterMetadataBuilder.h
--- a/Engine/Source/Runtime/RenderCore/Public/ShaderParameterMetadataBuilder.h
+++ b/Engine/Source/Runtime/RenderCore/Public/ShaderParameterMetadataBuilder.h
@@ -35,7 +35,7 @@ public:
 
 		NextMemberOffset = Align(NextMemberOffset, TParamTypeInfo::Alignment);
 
-		new(Members) FShaderParametersMetadata::FMember(
+		Members.Emplace(
 			Name,
 			TEXT(""),
 			__LINE__,
@@ -62,7 +62,7 @@ public:
 
 		NextMemberOffset = Align(NextMemberOffset, SHADER_PARAMETER_ARRAY_ELEMENT_ALIGNMENT);
 
-		new(Members) FShaderParametersMetadata::FMember(
+		Members.Emplace(
 			Name,
 			TEXT(""),
 			__LINE__,
@@ -117,7 +117,7 @@ public:
 		NextMemberOffset = Align(NextMemberOffset, TParamTypeInfo::Alignment);
 		const uint32 ThisMemberOffset = NextMemberOffset;
 
-		new(Members) FShaderParametersMetadata::FMember(
+		Members.Emplace(
 			Name,
 			TEXT(""),
 			__LINE__,
diff --git a/Engine/Source/Runtime/RenderCore/Public/ShaderPreprocessTypes.h b/Engine/Source/Runtime/RenderCore/Public/ShaderPreprocessTypes.h
--- a/Engine/Source/Runtime/RenderCore/Public/ShaderPreprocessTypes.h
+++ b/Engine/Source/Runtime/RenderCore/Public/ShaderPreprocessTypes.h
@@ -105,16 +105,16 @@ public:
 
 	inline void LogError(FString&& Message)
 	{
-		FShaderCompilerError* CompilerError = new(Errors) FShaderCompilerError;
-		CompilerError->StrippedErrorMessage = Message;
+		FShaderCompilerError& CompilerError = Errors.AddDefaulted_GetRef();;
+		CompilerError.StrippedErrorMessage = Message;
 	}
 
 	inline void LogError(FString&& FilePath, FString&& Message, FString&& LineNumberStr)
 	{
-		FShaderCompilerError* CompilerError = new(Errors) FShaderCompilerError;
-		CompilerError->ErrorVirtualFilePath = FilePath;
-		CompilerError->ErrorLineString = LineNumberStr;
-		CompilerError->StrippedErrorMessage = Message;
+		FShaderCompilerError& CompilerError = Errors.AddDefaulted_GetRef();
+		CompilerError.ErrorVirtualFilePath = FilePath;
+		CompilerError.ErrorLineString = LineNumberStr;
+		CompilerError.StrippedErrorMessage = Message;
 	}
 
 	inline void LogError(FString&& FilePath, FString&& Message, int32 LineNumber)
diff --git a/Engine/Source/Runtime/Slate/Public/Widgets/Navigation/SBreadcrumbTrail.h b/Engine/Source/Runtime/Slate/Public/Widgets/Navigation/SBreadcrumbTrail.h
--- a/Engine/Source/Runtime/Slate/Public/Widgets/Navigation/SBreadcrumbTrail.h
+++ b/Engine/Source/Runtime/Slate/Public/Widgets/Navigation/SBreadcrumbTrail.h
@@ -209,7 +209,7 @@ public:
 		];
 
 		// Push the crumb data
-		new (CrumbList) FCrumbItem(NextValidCrumbID, NewButton.ToSharedRef(), NewDelimiter.ToSharedRef(), NewButtonBox.ToSharedRef(), NewDelimiterBox.ToSharedRef(), NewCrumbData);
+		CrumbList.Emplace(NextValidCrumbID, NewButton.ToSharedRef(), NewDelimiter.ToSharedRef(), NewButtonBox.ToSharedRef(), NewDelimiterBox.ToSharedRef(), NewCrumbData);
 
 		// Increment the crumb ID for the next crumb
 		NextValidCrumbID = (NextValidCrumbID + 1) % (INT_MAX - 1);
