Made the paper2D on object reimported safe to use in a cooked editor and made the existing code more robust by not crashing if we are missing the SourceTexture meta data.

diff --git a/Engine/Plugins/2D/Paper2D/Source/Paper2DEditor/Private/Paper2DEditorModule.cpp b/Engine/Plugins/2D/Paper2D/Source/Paper2DEditor/Private/Paper2DEditorModule.cpp
--- a/Engine/Plugins/2D/Paper2D/Source/Paper2DEditor/Private/Paper2DEditorModule.cpp
+++ b/Engine/Plugins/2D/Paper2D/Source/Paper2DEditor/Private/Paper2DEditorModule.cpp
@@ -4,8 +4,10 @@
 #include "Paper2DEditorLog.h"
 #include "EditorModeRegistry.h"
 
+#include "AssetRegistry/AssetDataTagMap.h"
 #include "AssetRegistry/IAssetRegistry.h"
 #include "Subsystems/ImportSubsystem.h"
+#include "UObject/ObjectMacros.h"
 #include "UObject/UObjectIterator.h"
 
 #include "PropertyEditorModule.h"
@@ -311,23 +313,31 @@ private:
 				const FName SourceTextureMemberName = UPaperSprite::GetSourceTextureMemberName();
 				for (const FAssetData& SpriteAssetData : AssetsData)
 				{
-					FAssetData SavedTextureAssetData = AssetRegistry.GetAssetByObjectPath(SpriteAssetData.TagsAndValues.FindTag(SourceTextureMemberName).GetValue());
-				
-					FSoftObjectPath TextureSourcePath;
-					if (SavedTextureAssetData.IsRedirector())
-					{
-						TextureSourcePath = AssetRegistry.GetRedirectedObjectPath(SavedTextureAssetData.GetSoftObjectPath());
-					}
-					else
+					// We can't update a cooked asset
+					if (!SpriteAssetData.HasAnyPackageFlags(PKG_Cooked))
 					{
-						TextureSourcePath = SavedTextureAssetData.GetSoftObjectPath();
-					}
-
-					if (TextureSourcePath == TexturePath)
-					{
-						if (UPaperSprite* PaperSprite = Cast<UPaperSprite>(SpriteAssetData.GetAsset()))
+						FAssetTagValueRef AssetTagValue = SpriteAssetData.TagsAndValues.FindTag(SourceTextureMemberName);
+						if (AssetTagValue.IsSet())
 						{
-							PaperSprite->OnObjectReimported(Texture);
+							FAssetData SavedTextureAssetData = AssetRegistry.GetAssetByObjectPath(AssetTagValue.GetValue());
+				
+							FSoftObjectPath TextureSourcePath;
+							if (SavedTextureAssetData.IsRedirector())
+							{
+								TextureSourcePath = AssetRegistry.GetRedirectedObjectPath(SavedTextureAssetData.GetSoftObjectPath());
+							}
+							else
+							{
+								TextureSourcePath = SavedTextureAssetData.GetSoftObjectPath();
+							}
+
+							if (TextureSourcePath == TexturePath)
+							{
+								if (UPaperSprite* PaperSprite = Cast<UPaperSprite>(SpriteAssetData.GetAsset()))
+								{
+									PaperSprite->OnObjectReimported(Texture);
+								}
+							}
 						}
 					}
 				}
