Fixed crash in paper2D due to DestroyRenderThreadResources not being called when a proxy is both added and removed from the scene prior to a scene update.

CreateRenderThreadResources is now deferred until the scene update when render command pipes are enabled in order to avoid race conditions when binding vertex factories. For symmetry, DestroyRenderThreadResources is only called if a corresponding CreateRenderThreadResources call will occur later during scene update. However, paper 2D initializes the vertex factories outside of the Create / Destroy pair and then expects Destroy to always be called. As it's not necessary to release render resources from this method anymore, they can safely be moved to the destructor and avoid the problem altogether.

diff --git a/Engine/Plugins/2D/Paper2D/Source/Paper2D/Private/PaperRenderSceneProxy.cpp b/Engine/Plugins/2D/Paper2D/Source/Paper2D/Private/PaperRenderSceneProxy.cpp
--- a/Engine/Plugins/2D/Paper2D/Source/Paper2D/Private/PaperRenderSceneProxy.cpp
+++ b/Engine/Plugins/2D/Paper2D/Source/Paper2D/Private/PaperRenderSceneProxy.cpp
@@ -177,6 +177,21 @@ FPaperRenderSceneProxy::FPaperRenderSceneProxy(const UPrimitiveComponent* InComp
 	bSpritesUseVertexBufferPath = CVarDrawSpritesUsingPrebuiltVertexBuffers.GetValueOnAnyThread() != 0;
 }
 
+FPaperRenderSceneProxy::~FPaperRenderSceneProxy()
+{
+	for (FSpriteTextureOverrideRenderProxy* Proxy : MaterialTextureOverrideProxies)
+	{
+		if (Proxy != nullptr)
+		{
+			Proxy->ReleasePrimitiveResource();
+		}
+	}
+	MaterialTextureOverrideProxies.Empty();
+
+	VertexBuffer.ReleaseResource();
+	VertexFactory.ReleaseResource();
+}
+
 SIZE_T FPaperRenderSceneProxy::GetTypeHash() const
 {
 	static size_t UniquePointer;
@@ -264,21 +279,6 @@ void FPaperRenderSceneProxy::CreateRenderThreadResources(FRHICommandListBase& RH
 	}
 }
 
-void FPaperRenderSceneProxy::DestroyRenderThreadResources()
-{
-	for (FSpriteTextureOverrideRenderProxy* Proxy : MaterialTextureOverrideProxies)
-	{
-		if (Proxy != nullptr)
-		{
-			Proxy->ReleasePrimitiveResource();
-		}
-	}
-	MaterialTextureOverrideProxies.Empty();
-
-	VertexBuffer.ReleaseResource();
-	VertexFactory.ReleaseResource();
-}
-
 void FPaperRenderSceneProxy::DebugDrawCollision(const FSceneView* View, int32 ViewIndex, FMeshElementCollector& Collector, bool bDrawSolid) const
 {
 	if (MyBodySetup != nullptr)
diff --git a/Engine/Plugins/2D/Paper2D/Source/Paper2D/Private/PaperRenderSceneProxy.h b/Engine/Plugins/2D/Paper2D/Source/Paper2D/Private/PaperRenderSceneProxy.h
--- a/Engine/Plugins/2D/Paper2D/Source/Paper2D/Private/PaperRenderSceneProxy.h
+++ b/Engine/Plugins/2D/Paper2D/Source/Paper2D/Private/PaperRenderSceneProxy.h
@@ -125,6 +125,7 @@ public:
 	SIZE_T GetTypeHash() const override;
 
 	FPaperRenderSceneProxy(const UPrimitiveComponent* InComponent);
+	~FPaperRenderSceneProxy();
 
 	// FPrimitiveSceneProxy interface.
 	virtual void GetDynamicMeshElements(const TArray<const FSceneView*>& Views, const FSceneViewFamily& ViewFamily, uint32 VisibilityMap, FMeshElementCollector& Collector) const override;
@@ -133,7 +134,6 @@ public:
 	virtual bool CanBeOccluded() const override;
 	virtual bool IsUsingDistanceCullFade() const override;
 	virtual void CreateRenderThreadResources(FRHICommandListBase& RHICmdList) override;
-	virtual void DestroyRenderThreadResources() override;
 	// End of FPrimitiveSceneProxy interface.
 
 	void SetBodySetup_RenderThread(UBodySetup* NewSetup);
