Fix PaperSprite loading logic under editor and cook that could result in visual or cooking bugs when async loading was used to load the object

#rnx

diff --git a/Engine/Plugins/2D/Paper2D/Source/Paper2D/Private/PaperSprite.cpp b/Engine/Plugins/2D/Paper2D/Source/Paper2D/Private/PaperSprite.cpp
--- a/Engine/Plugins/2D/Paper2D/Source/Paper2D/Private/PaperSprite.cpp
+++ b/Engine/Plugins/2D/Paper2D/Source/Paper2D/Private/PaperSprite.cpp
@@ -1437,26 +1437,19 @@ void UPaperSprite::SetPivotMode(ESpritePivotMode::Type InPivotMode, FVector2D In
 
 #if WITH_EDITOR
 
-static TAutoConsoleVariable<bool> CVarDisableSyncLoadInEditor(
-	TEXT("Paper2d.DisableSyncLoadInEditor"),
-	true,
-	TEXT("Don't attempt to load the source texture synchronous if we are inside the AsyncLoading thread alreayd."),
-	ECVF_Cheat);
-
 UTexture2D* UPaperSprite::GetSourceTexture() const
 {
 	// Verify the cache is still valid.
-	if (SourceTextureCacheNeverSerialized && SourceTexture.Get() != nullptr && SourceTexture.Get() == SourceTextureCacheNeverSerialized)
+	UTexture2D* SourceTexturePtr = SourceTexture.Get();
+	if (SourceTextureCacheNeverSerialized && SourceTexturePtr == SourceTextureCacheNeverSerialized)
 	{
 		return SourceTextureCacheNeverSerialized;
 	}
-	
-	// We don't want to load synchronously unless we are OUTSIDE of an async loading thread.
-	// We could be in the Async loading thread if we are in the editor with cooked content already
-	UTexture2D* SourceTexturePtr =
-			CVarDisableSyncLoadInEditor.GetValueOnGameThread() && IsInAsyncLoadingThread() ?
-			nullptr :
-			SourceTexture.LoadSynchronous();
+
+	if (SourceTexturePtr == nullptr)
+	{
+		SourceTexturePtr = SourceTexture.LoadSynchronous();
+	}
 
 	// We need to completely load the texture if we're in post load and this texture is needed immediately,
 	// not safe to do in a game with EDL, but safe at editor time.
