Fixed race condition in Paper2D shadow GDME.
 - The proxy is gathered from multiple shadows simultaneously, but modifies a local mutable cache.

diff --git a/Engine/Plugins/2D/Paper2D/Source/Paper2D/Private/PaperRenderSceneProxy.cpp b/Engine/Plugins/2D/Paper2D/Source/Paper2D/Private/PaperRenderSceneProxy.cpp
--- a/Engine/Plugins/2D/Paper2D/Source/Paper2D/Private/PaperRenderSceneProxy.cpp
+++ b/Engine/Plugins/2D/Paper2D/Source/Paper2D/Private/PaperRenderSceneProxy.cpp
@@ -602,6 +602,8 @@ void FPaperRenderSceneProxy::SetTransientTextureOverride_RenderThread(const UTex
 
 FSpriteTextureOverrideRenderProxy* FPaperRenderSceneProxy::GetCachedMaterialProxyForSection(FMeshElementCollector& Collector, int32 SectionIndex, FMaterialRenderProxy* ParentMaterialProxy) const
 {
+	UE::TScopeLock Lock(MaterialTextureOverrideProxiesMutex);
+
 	if (MaterialTextureOverrideProxies.Num() < BatchedSections.Num())
 	{
 		MaterialTextureOverrideProxies.AddDefaulted(BatchedSections.Num() - MaterialTextureOverrideProxies.Num());
diff --git a/Engine/Plugins/2D/Paper2D/Source/Paper2D/Private/PaperRenderSceneProxy.h b/Engine/Plugins/2D/Paper2D/Source/Paper2D/Private/PaperRenderSceneProxy.h
--- a/Engine/Plugins/2D/Paper2D/Source/Paper2D/Private/PaperRenderSceneProxy.h
+++ b/Engine/Plugins/2D/Paper2D/Source/Paper2D/Private/PaperRenderSceneProxy.h
@@ -163,6 +163,8 @@ protected:
 protected:
 	TArray<FSpriteRenderSection> BatchedSections;
 	TArray<FDynamicMeshVertex> Vertices;
+
+	mutable UE::FMutex MaterialTextureOverrideProxiesMutex;
 	mutable TArray<FSpriteTextureOverrideRenderProxy*> MaterialTextureOverrideProxies;
 
 	FPaperSpriteVertexBuffer VertexBuffer;
