Set DEFAULT_ORTHOZOOM for Paper2D editor viewports.

diff --git a/Engine/Plugins/2D/Paper2D/Source/Paper2DEditor/Private/PaperEditorViewportClient.cpp b/Engine/Plugins/2D/Paper2D/Source/Paper2DEditor/Private/PaperEditorViewportClient.cpp
--- a/Engine/Plugins/2D/Paper2D/Source/Paper2DEditor/Private/PaperEditorViewportClient.cpp
+++ b/Engine/Plugins/2D/Paper2D/Source/Paper2DEditor/Private/PaperEditorViewportClient.cpp
@@ -51,7 +51,7 @@ FPaperEditorViewportClient::FPaperEditorViewportClient(const TWeakPtr<SEditorVie
 	bDeferZoomToSpriteIsInstant = true;
 
 	// Get the correct general direction of the perspective mode; the distance doesn't matter much as we've queued up a deferred zoom that will calculate a much better distance
-	SetInitialViewTransform(LVT_Perspective, -100.0f * PaperAxisZ, PaperAxisZ.Rotation(), 0.0f);
+	SetInitialViewTransform(LVT_Perspective, -100.0f * PaperAxisZ, PaperAxisZ.Rotation(), DEFAULT_ORTHOZOOM);
 }
 
 FPaperEditorViewportClient::~FPaperEditorViewportClient()
