PR #12064: Add GetTypeHash() to FPaperTileInfo

diff --git a/Engine/Plugins/2D/Paper2D/Source/Paper2D/Classes/PaperTileLayer.h b/Engine/Plugins/2D/Paper2D/Source/Paper2D/Classes/PaperTileLayer.h
--- a/Engine/Plugins/2D/Paper2D/Source/Paper2D/Classes/PaperTileLayer.h
+++ b/Engine/Plugins/2D/Paper2D/Source/Paper2D/Classes/PaperTileLayer.h
@@ -98,6 +98,11 @@ struct FPaperTileInfo
 			}
 		}
 	}
+
+	friend uint32 GetTypeHash(const FPaperTileInfo& PaperTileInfo)
+	{
+		return HashCombine(GetTypeHash(PaperTileInfo.TileSet), PaperTileInfo.PackedTileIndex);
+	}
 };
 
 // This class represents a single layer in a tile map.  All layers in the map must have the size dimensions.
