Rename "Enable Planar Snapping" to just "Planar"

This menu entry will from now on only be visible in the viewport toolbar's "Snapping" menu where this new shorter name fits in and "snapping" is understood from context.

diff --git a/Engine/Plugins/2D/Paper2D/Source/SmartSnapping/Private/SmartSnappingModule.cpp b/Engine/Plugins/2D/Paper2D/Source/SmartSnapping/Private/SmartSnappingModule.cpp
--- a/Engine/Plugins/2D/Paper2D/Source/SmartSnapping/Private/SmartSnappingModule.cpp
+++ b/Engine/Plugins/2D/Paper2D/Source/SmartSnapping/Private/SmartSnappingModule.cpp
@@ -29,9 +29,10 @@ public:
 
 			// Register the extension with the level editor
 			{
+				// TODO: This should extend the viewport toolbar menu "LevelEditor.ViewportToolbar.Show" using the
+				// ToolMenus system instead of using legacy menu extenders.
 				ViewMenuExtender = FLevelEditorModule::FLevelEditorMenuExtender::CreateRaw(this, &FSmartSnappingModule::OnExtendLevelEditorViewMenu);
 
-
 				FLevelEditorModule& LevelEditor = FModuleManager::LoadModuleChecked<FLevelEditorModule>(TEXT("LevelEditor"));
 				auto& MenuExtenders = LevelEditor.GetAllLevelEditorToolbarViewMenuExtenders();
 				MenuExtenders.Add(ViewMenuExtender);
@@ -66,19 +67,22 @@ public:
 			FCanExecuteAction(),
 			FIsActionChecked::CreateRaw(PlanarPolicy.Get(), &FPlanarConstraintSnapPolicy::IsEnabled));
 		Builder.AddMenuEntry(
-			LOCTEXT("View_Extension_PlanarSnap_Enable", "Enable Planar Snapping"), 
-			LOCTEXT("View_Extension_PlanarSnap_Tooltip", "If Enabled, actors will snap to the nearest location on the constraint plane (NOTE: Only works correctly in perspective views right now!)"),
+			LOCTEXT("View_Extension_PlanarSnap_Enable", "Planar"),
+			LOCTEXT(
+				"View_Extension_PlanarSnap_Tooltip", "If Enabled, actors will snap to the nearest location on the constraint plane (NOTE: Only works correctly in perspective views right now!)"
+			),
 			FSlateIcon(),
 			Action_TogglePlanarSnap,
 			NAME_None,
-			EUserInterfaceActionType::ToggleButton);
-
-// 		Builder.AddWidget(SNew(STextBlock).Text(TEXT("Depth value goes here")),
-// 			LOCTEXT("View_Extension_PlanarSnap_Depth", "Distance from plane"));
-			// 		FUIAction Action_TogglePlanarSnap(
-// 			FExecuteAction::CreateSP(PlanarPolicy, &FPlanarConstraintSnapPolicy::ToggleEnabled),
-// 			FCanExecuteAction(),
-// 			FIsActionChecked::CreateSP(PlanarPolicy, &FPlanarConstraintSnapPolicy::IsEnabled));
+			EUserInterfaceActionType::ToggleButton
+		);
+
+		// 		Builder.AddWidget(SNew(STextBlock).Text(TEXT("Depth value goes here")),
+		// 			LOCTEXT("View_Extension_PlanarSnap_Depth", "Distance from plane"));
+		// 		FUIAction Action_TogglePlanarSnap(
+		// 			FExecuteAction::CreateSP(PlanarPolicy, &FPlanarConstraintSnapPolicy::ToggleEnabled),
+		// 			FCanExecuteAction(),
+		// 			FIsActionChecked::CreateSP(PlanarPolicy, &FPlanarConstraintSnapPolicy::IsEnabled));
 	}
 
 	TSharedRef<FExtender> OnExtendLevelEditorViewMenu(const TSharedRef<FUICommandList> CommandList)
