PR #10716: Adding a getter to the property bLayerCollides of UPaperTileLayer

diff --git a/Engine/Plugins/2D/Paper2D/Source/Paper2D/Classes/PaperTileLayer.h b/Engine/Plugins/2D/Paper2D/Source/Paper2D/Classes/PaperTileLayer.h
--- a/Engine/Plugins/2D/Paper2D/Source/Paper2D/Classes/PaperTileLayer.h
+++ b/Engine/Plugins/2D/Paper2D/Source/Paper2D/Classes/PaperTileLayer.h
@@ -251,6 +251,10 @@ public:
 	}
 
 	void SetLayerCollides(bool bShouldCollide) { bLayerCollides = bShouldCollide; }
+	
+	// Should this layer generate collision?
+	// Note: This only has an effect if the owning tile map has collision enabled
+	bool GetLayerCollides() const { return bLayerCollides; }
 
 	void SetLayerCollisionThickness(bool bShouldOverride, float OverrideValue)
 	{
