Paper 2D Time Map Asset: Fix to localization of Detail Panel category

diff --git a/Engine/Plugins/2D/Paper2D/Source/Paper2DEditor/Private/TileMapEditing/PaperTileMapDetailsCustomization.cpp b/Engine/Plugins/2D/Paper2D/Source/Paper2DEditor/Private/TileMapEditing/PaperTileMapDetailsCustomization.cpp
--- a/Engine/Plugins/2D/Paper2D/Source/Paper2DEditor/Private/TileMapEditing/PaperTileMapDetailsCustomization.cpp
+++ b/Engine/Plugins/2D/Paper2D/Source/Paper2DEditor/Private/TileMapEditing/PaperTileMapDetailsCustomization.cpp
@@ -79,7 +79,7 @@ void FPaperTileMapDetailsCustomization::CustomizeDetails(IDetailLayoutBuilder& D
 	TileMapComponentPtr = TileComponent;
 
 	// Make sure the Tile Map category is right below the Transform
-	IDetailCategoryBuilder& TileMapCategory = DetailLayout.EditCategory("Tile Map", FText::GetEmpty(), ECategoryPriority::Important);
+	IDetailCategoryBuilder& TileMapCategory = DetailLayout.EditCategory("Tile Map", LOCTEXT("TileMapDetailCategory", "Time Map"), ECategoryPriority::Important);
 
 	// Add the 'instanced' versus 'asset' indicator to the tile map header
 	TileMapCategory.HeaderContent
